Rename/move HeadlessChaos.

diff --git a/Engine/Source/Programs/HeadlessChaos/HeadlessChaos.Build.cs b/Engine/Source/Programs/HeadlessChaos/HeadlessChaos.Build.cs
new file mode 100644
--- /dev/null
+++ b/Engine/Source/Programs/HeadlessChaos/HeadlessChaos.Build.cs
@@ -0,0 +1,46 @@
+// Copyright Epic Games, Inc. All Rights Reserved.
+
+using UnrealBuildTool;
+
+public class HeadlessChaos : ModuleRules
+{
+	public HeadlessChaos(ReadOnlyTargetRules Target) : base(Target)
+	{
+		PublicIncludePaths.Add("Runtime/Launch/Public");
+
+		// For LaunchEngineLoop.cpp include
+		PrivateIncludePaths.Add("Runtime/Launch/Private");
+
+		PrivateDependencyModuleNames.AddRange(
+			new string[] {
+                "ApplicationCore",
+				"Core",
+                "Projects",
+                "GoogleTest",
+				"Chaos",
+				"ChaosSolvers",
+            }
+        );
+
+		if (Target.Platform == UnrealTargetPlatform.Win32 || Target.Platform == UnrealTargetPlatform.Win64 || Target.Platform == UnrealTargetPlatform.XboxOne)
+		{
+			PublicDefinitions.Add("GTEST_OS_WINDOWS=1");
+		}
+		else if (Target.Platform == UnrealTargetPlatform.Mac)
+		{
+			PublicDefinitions.Add("GTEST_OS_MAC=1");
+		}
+		else if (Target.Platform == UnrealTargetPlatform.IOS || Target.Platform == UnrealTargetPlatform.TVOS)
+		{
+			PublicDefinitions.Add("GTEST_OS_IOS=1");
+		}
+		else if (Target.Platform == UnrealTargetPlatform.Android || Target.Platform == UnrealTargetPlatform.Quail || Target.Platform == UnrealTargetPlatform.Lumin)
+		{
+			PublicDefinitions.Add("GTEST_OS_LINUX_ANDROID=1");
+		}
+		else if (Target.Platform == UnrealTargetPlatform.Linux || Target.Platform == UnrealTargetPlatform.PS4)
+		{
+			PublicDefinitions.Add("GTEST_OS_LINUX=1");
+		}
+	}
+}
diff --git a/Engine/Source/Programs/HeadlessChaos/HeadlessChaos.Target.cs b/Engine/Source/Programs/HeadlessChaos/HeadlessChaos.Target.cs
new file mode 100644
--- /dev/null
+++ b/Engine/Source/Programs/HeadlessChaos/HeadlessChaos.Target.cs
@@ -0,0 +1,28 @@
+// Copyright Epic Games, Inc. All Rights Reserved.
+
+using UnrealBuildTool;
+using System.Collections.Generic;
+
+[SupportedPlatforms(UnrealPlatformClass.All)]
+public class HeadlessChaosTarget : TargetRules
+{
+	public HeadlessChaosTarget(TargetInfo Target) : base(Target)
+	{
+		Type = TargetType.Program;
+		LinkType = TargetLinkType.Monolithic;
+
+        ExeBinariesSubFolder = "NotForLicensees";
+		LaunchModuleName = "HeadlessChaos";
+
+		bBuildDeveloperTools = false;
+
+		// HeadlessChaos doesn't ever compile with the engine linked in
+		bCompileAgainstEngine = false;
+		bCompileAgainstCoreUObject = false;
+
+        bHasExports = false;
+
+        // UnrealHeaderTool is a console application, not a Windows app (sets entry point to main(), instead of WinMain())
+        bIsBuildingConsoleApplication = true;
+    }
+}
diff --git a/Engine/Source/Programs/HeadlessChaos/Private/HeadlessChaosApp.cpp b/Engine/Source/Programs/HeadlessChaos/Private/HeadlessChaosApp.cpp
new file mode 100644
--- /dev/null
+++ b/Engine/Source/Programs/HeadlessChaos/Private/HeadlessChaosApp.cpp
@@ -0,0 +1,283 @@
+// Copyright Epic Games, Inc. All Rights Reserved.
+
+#include "HeadlessChaosApp.h"
+
+#include "HeadlessChaos.h"
+#include "HeadlessChaosTestBP.h"
+#include "HeadlessChaosTestBroadphase.h"
+#include "HeadlessChaosTestCloth.h"
+#include "HeadlessChaosTestClustering.h"
+#include "HeadlessChaosTestCollisions.h"
+#include "HeadlessChaosTestForces.h"
+#include "HeadlessChaosTestGJK.h"
+#include "HeadlessChaosTestImplicits.h"
+#include "HeadlessChaosTestRaycast.h"
+#include "HeadlessChaosTestSerialization.h"
+#include "HeadlessChaosTestSpatialHashing.h"
+#include "Modules/ModuleManager.h"
+#include "RequiredProgramMainCPPInclude.h"
+#include "Chaos/PBDRigidsEvolution.h"
+#include "Chaos/PBDJointConstraints.h"
+#include "Chaos/Box.h"
+#include "Chaos/Sphere.h"
+#include "Chaos/Utilities.h"
+#include "HeadlessChaosTestParticleHandle.h"
+#include "HeadlessChaosTestClustering.h"
+#include "HeadlessChaosTestSerialization.h"
+#include "HeadlessChaosTestBP.h"
+#include "HeadlessChaosTestRaycast.h"
+#include "HeadlessChaosTestGJK.h"
+#include "HeadlessChaosTestEPA.h"
+#include "HeadlessChaosTestBroadphase.h"
+#include "HeadlessChaosTestMostOpposing.h"
+#include "HeadlessChaosTestSolverCommandList.h"
+#include "HeadlessChaosTestSolverProxies.h"
+#include "HeadlessChaosTestHandles.h"
+
+
+IMPLEMENT_APPLICATION(HeadlessChaos, "HeadlessChaos");
+
+#define LOCTEXT_NAMESPACE "HeadlessChaos"
+
+DEFINE_LOG_CATEGORY(LogHeadlessChaos);
+
+TEST(ImplicitTests, Implicit) {
+	ChaosTest::ImplicitPlane<float>();
+	ChaosTest::ImplicitCube<float>();
+	ChaosTest::ImplicitSphere<float>();
+	ChaosTest::ImplicitCylinder<float>();
+	ChaosTest::ImplicitTaperedCylinder<float>();
+	ChaosTest::ImplicitCapsule<float>();
+	ChaosTest::ImplicitScaled<float>();
+	ChaosTest::ImplicitTransformed<float>();
+	ChaosTest::ImplicitIntersection<float>();
+	ChaosTest::ImplicitUnion<float>();
+	// @todo: Make this work at some point
+	//ChaosTest::ImplicitLevelset<float>();
+
+	SUCCEED();
+}
+
+TEST(ImplicitTests, Rasterization) {
+	ChaosTest::RasterizationImplicit<float>();
+	ChaosTest::RasterizationImplicitWithHole<float>();
+}
+
+TEST(ImplicitTests, ConvexHull) {
+	ChaosTest::ConvexHull<float>();
+	ChaosTest::ConvexHull2<float>();
+	ChaosTest::Simplify<float>();
+}
+
+TEST(CollisionTests, Collisions) {
+	GEnsureOnNANDiagnostic = 1;
+
+	ChaosTest::LevelsetConstraint<float>();
+	// ChaosTest::LevelsetConstraintGJK<float>();
+	ChaosTest::CollisionBoxPlane<float>();
+	ChaosTest::CollisionBoxPlaneZeroResitution<float>();
+	ChaosTest::CollisionBoxPlaneRestitution<float>();
+	ChaosTest::CollisionBoxToStaticBox<float>();
+	ChaosTest::CollisionConvexConvex<float>();
+
+	// @ todo: Make this work at some point
+	//ChaosTest::SpatialHashing<float>();
+
+	SUCCEED();
+}
+
+TEST(CollisionTests, PGS) {
+	ChaosTest::CollisionPGS<float>();
+	ChaosTest::CollisionPGS2<float>();
+	SUCCEED();
+}
+
+TEST(ForceTests, Forces) {
+	ChaosTest::Gravity<Chaos::TPBDRigidsEvolutionGBF<float, 3>, float>();
+#if CHAOS_PARTICLEHANDLE_TODO
+	ChaosTest::Gravity<Chaos::TPBDRigidsEvolutionPGS<float, 3>, float>();
+#endif
+	SUCCEED();
+}
+
+TEST(Clustering, Clustering) {
+	ChaosTest::ImplicitCluster<float>();
+	ChaosTest::FractureCluster<float>();
+	ChaosTest::PartialFractureCluster<float>();
+	SUCCEED();
+}
+
+TEST(SerializationTests, Serialization) {
+	ChaosTest::SimpleObjectsSerialization<float>();
+	ChaosTest::SharedObjectsSerialization<float>();
+	ChaosTest::GraphSerialization<float>();
+	ChaosTest::ObjectUnionSerialization<float>();
+	ChaosTest::ParticleSerialization<float>();
+	ChaosTest::BVHSerialization<float>();
+	ChaosTest::RigidParticlesSerialization<float>();
+	ChaosTest::BVHParticlesSerialization<float>();
+	SUCCEED();
+}
+
+TEST(BroadphaseTests, Broadphase) {
+	ChaosTest::BPPerfTest<float>();
+	ChaosTest::TestPendingSpatialDataHandlePointerConflict();
+	SUCCEED();
+}
+
+TEST(ClothTests, DeformableGravity) {
+	ChaosTest::DeformableGravity<float>();
+
+	SUCCEED();
+}
+
+TEST(RaycastTests, Raycast) {
+	ChaosTest::SphereRaycast<float>();
+	ChaosTest::PlaneRaycast<float>();
+	ChaosTest::CylinderRaycast<float>();
+	ChaosTest::TaperedCylinderRaycast<float>();
+	ChaosTest::CapsuleRaycast<float>();
+	ChaosTest::TriangleRaycast<float>();
+	ChaosTest::BoxRaycast<float>();
+	ChaosTest::ScaledRaycast<float>();
+	ChaosTest::TransformedRaycast<float>();
+	ChaosTest::UnionRaycast<float>();
+	ChaosTest::IntersectionRaycast<float>();
+	
+	SUCCEED();
+}
+
+TEST(MostOpposingTests, MostOpposing) {
+	ChaosTest::TrimeshMostOpposing<float>();
+	ChaosTest::ConvexMostOpposing<float>();
+	ChaosTest::ScaledMostOpposing<float>();
+
+	SUCCEED();
+}
+
+TEST(GJK, Simplexes) {
+	ChaosTest::SimplexLine<float>();
+	ChaosTest::SimplexTriangle<float>();
+	ChaosTest::SimplexTetrahedron<float>();
+	
+	SUCCEED();
+}
+
+TEST(GJK, GJKIntersectTests) {
+	ChaosTest::GJKSphereSphereTest<float>();
+	ChaosTest::GJKSphereBoxTest<float>();
+	ChaosTest::GJKSphereCapsuleTest<float>();
+	ChaosTest::GJKSphereConvexTest<float>();
+	ChaosTest::GJKSphereScaledSphereTest<float>();
+	
+	SUCCEED();
+}
+
+TEST(GJK, GJKRaycastTests) {
+	ChaosTest::GJKSphereSphereSweep<float>();
+	ChaosTest::GJKSphereBoxSweep<float>();
+	ChaosTest::GJKSphereCapsuleSweep<float>();
+	ChaosTest::GJKSphereConvexSweep<float>();
+	ChaosTest::GJKSphereScaledSphereSweep<float>();
+	ChaosTest::GJKBoxCapsuleSweep<float>();
+	ChaosTest::GJKBoxBoxSweep<float>();
+	ChaosTest::GJKCapsuleConvexInitialOverlapSweep<float>();
+	SUCCEED();
+}
+
+TEST(EPA, EPATests) {
+	ChaosTest::EPAInitTest<float>();
+	ChaosTest::EPASimpleTest<float>();
+	SUCCEED();
+}
+
+TEST(BP, BroadphaseTests) {
+	ChaosTest::GridBPTest<float>();
+	ChaosTest::GridBPTest2<float>();
+	ChaosTest::AABBTreeTest<float>();
+	ChaosTest::AABBTreeTimesliceTest<float>();
+	ChaosTest::BroadphaseCollectionTest<float>();
+	SUCCEED();
+}
+
+TEST(ParticleHandle, ParticleHandleTests)
+{
+	ChaosTest::ParticleIteratorTest<float>();
+	ChaosTest::ParticleHandleTest<float>();
+	ChaosTest::AccelerationStructureHandleComparison();
+	ChaosTest::HandleObjectStateChangeTest();
+	SUCCEED();
+}
+
+TEST(Perf, PerfTests)
+{
+	ChaosTest::EvolutionPerfHarness();
+	SUCCEED();
+}
+
+
+TEST(ChaosSolver, ChaosSolverTests)
+{
+	ChaosTest::SingleParticleProxySingleThreadTest<float>();
+	ChaosTest::CommandListTest<float>();
+}
+TEST(ChaosSolver, DISABLED_ChaosSolverTests)
+{
+	ChaosTest::SingleParticleProxyTaskGraphTest<float>();
+}
+
+TEST(Handles, FrameworkTests)
+{
+	ChaosTest::Handles::HandleArrayTest<float>();
+	ChaosTest::Handles::HandleHeapTest<float>();
+	ChaosTest::Handles::HandleSerializeTest<float>();
+}
+
+
+class UEGTestPrinter : public ::testing::EmptyTestEventListener
+{
+    virtual void OnTestStart(const ::testing::TestInfo& TestInfo)
+    {
+        UE_LOG(LogHeadlessChaos, Verbose, TEXT("Test %s.%s Starting"), *FString(TestInfo.test_case_name()), *FString(TestInfo.name()));
+    }
+
+    virtual void OnTestPartResult(const ::testing::TestPartResult& TestPartResult)
+    {
+        if (TestPartResult.failed())
+        {
+            UE_LOG(LogHeadlessChaos, Error, TEXT("FAILED in %s:%d\n%s"), *FString(TestPartResult.file_name()), TestPartResult.line_number(), *FString(TestPartResult.summary()))
+        }
+        else
+        {
+            UE_LOG(LogHeadlessChaos, Verbose, TEXT("Succeeded in %s:%d\n%s"), *FString(TestPartResult.file_name()), TestPartResult.line_number(), *FString(TestPartResult.summary()))
+        }
+    }
+
+    virtual void OnTestEnd(const ::testing::TestInfo& TestInfo)
+    {
+        UE_LOG(LogHeadlessChaos, Verbose, TEXT("Test %s.%s Ending"), *FString(TestInfo.test_case_name()), *FString(TestInfo.name()));
+    }
+};
+
+INT32_MAIN_INT32_ARGC_TCHAR_ARGV()
+{
+    // start up the main loop
+	GEngineLoop.PreInit(ArgC, ArgV);
+
+	::testing::InitGoogleTest(&ArgC, ArgV);
+
+	// Add a UE-formatting printer
+    ::testing::TestEventListeners& Listeners = ::testing::UnitTest::GetInstance()->listeners();
+    Listeners.Append(new UEGTestPrinter);
+
+	ensure(RUN_ALL_TESTS() == 0);
+
+	FCoreDelegates::OnExit.Broadcast();
+	FModuleManager::Get().UnloadModulesAtShutdown();
+
+	FPlatformMisc::RequestExit(false);
+
+	return 0;
+}
+
+#undef LOCTEXT_NAMESPACE
diff --git a/Engine/Source/Programs/HeadlessChaos/Private/HeadlessChaosTestBroadphase.cpp b/Engine/Source/Programs/HeadlessChaos/Private/HeadlessChaosTestBroadphase.cpp
new file mode 100644
--- /dev/null
+++ b/Engine/Source/Programs/HeadlessChaos/Private/HeadlessChaosTestBroadphase.cpp
@@ -0,0 +1,692 @@
+// Copyright Epic Games, Inc. All Rights Reserved.
+
+#include "HeadlessChaosTestBroadphase.h"
+
+#include "HeadlessChaos.h"
+#include "Chaos/Box.h"
+#include "Chaos/BoundingVolume.h"
+#include "Chaos/ParticleHandle.h"
+#include "Chaos/PBDRigidsSOAs.h"
+#include "Chaos/PBDRigidsEvolutionGBF.h"
+#include "Chaos/AABBTree.h"
+#include "PBDRigidsSolver.h"
+#include "Chaos/SpatialAccelerationCollection.h"
+
+namespace ChaosTest
+{
+	using namespace Chaos;
+
+	/*In general we want to test the following for each broadphase type:
+	- simple intersection test as used by sim (IntersectAll)
+	- ray, sweep, overlap
+	- miss entire structure
+	- stop mid structure
+	- multi overlap
+	- multi block (adjust length)
+	- any
+	*/
+
+	template <typename T>
+	struct TVisitor : ISpatialVisitor<int32, T>
+	{
+		const TGeometryParticles<T, 3>& Boxes;
+		const TVector<T, 3> Start;
+		const TVector<T, 3> Dir;
+		TVector<T, 3> HalfExtents;
+		const T Thickness;
+		int32 BlockAfterN;
+		bool bAny;
+
+		TVisitor(const TVector<T,3>& InStart, const TVector<T,3>& InDir, const T InThickness, const TGeometryParticles<T,3>& InBoxes)
+		: Boxes(InBoxes)
+		, Start(InStart)
+		, Dir(InDir)
+		, HalfExtents(0)
+		, Thickness(InThickness)
+		, BlockAfterN(TNumericLimits<int32>::Max())
+		, bAny(false)
+		{}
+
+		enum class SQType
+		{
+			Raycast,
+			Sweep,
+			Overlap
+		};
+
+		template <SQType>
+		bool Visit(int32 Idx, FQueryFastData& CurData)
+		{
+			const TRigidTransform<T,3> BoxTM(Boxes.X(Idx), Boxes.R(Idx));
+			TAABB<T, 3> Box = static_cast<const TBox<T, 3>*>(Boxes.Geometry(Idx).Get())->BoundingBox().TransformedAABB(BoxTM);
+			TAABB<T, 3> ThicknedBox(Box.Min() - HalfExtents, Box.Max() + HalfExtents);
+
+			T NewLength;
+			TVector<T, 3> Position;
+			TVector<T, 3> Normal;
+			int32 FaceIndex;
+			const T OldLength = CurData.CurrentLength;
+			if (ThicknedBox.Raycast(Start, Dir, CurData.CurrentLength, 0, NewLength, Position, Normal, FaceIndex))
+			{
+				Instances.Add(Idx);
+				if (bAny)
+				{
+					return false;
+				}
+				if (Instances.Num() >= BlockAfterN)
+				{
+					//blocking so adjust Length
+					CurData.SetLength(NewLength);
+				}
+			}
+			
+			return true;
+		}
+
+		bool VisitRaycast(TSpatialVisitorData<int32> Idx, FQueryFastData& CurData)
+		{
+			return Visit<SQType::Raycast>(Idx.Payload, CurData);
+		}
+
+		bool VisitSweep(TSpatialVisitorData<int32> Idx, FQueryFastData& CurData)
+		{
+			return Visit<SQType::Sweep>(Idx.Payload, CurData);
+		}
+
+		bool VisitOverlap(TSpatialVisitorData<int32> Idx)
+		{
+			check(false);
+			return false;
+		}
+
+		virtual bool Overlap(const TSpatialVisitorData<int32>& Instance) override
+		{
+			return VisitOverlap(Instance);
+		}
+
+		virtual bool Raycast(const TSpatialVisitorData<int32>& Instance, FQueryFastData& CurData) override
+		{
+			return VisitRaycast(Instance, CurData);
+		}
+
+		virtual bool Sweep(const TSpatialVisitorData<int32>& Instance, FQueryFastData& CurData) override
+		{
+			return VisitSweep(Instance, CurData);
+		}
+
+		TArray<int32> Instances;
+	};
+
+	template <typename T>
+	struct TOverlapVisitor : public ISpatialVisitor<int32, T>
+	{
+		const TGeometryParticles<T, 3>& Boxes;
+		const TAABB<T, 3> Bounds;
+		bool bAny;
+
+		TOverlapVisitor(const TAABB<T,3>& InBounds, const TGeometryParticles<T, 3>& InBoxes)
+			: Boxes(InBoxes)
+			, Bounds(InBounds)
+			, bAny(false)
+		{}
+
+		bool VisitOverlap(TSpatialVisitorData<int32> Instance)
+		{
+			const int32 Idx = Instance.Payload;
+			const TRigidTransform<T, 3> BoxTM(Boxes.X(Idx), Boxes.R(Idx));
+			TAABB<T, 3> Box = static_cast<const TBox<T, 3>*>(Boxes.Geometry(Idx).Get())->BoundingBox().TransformedAABB(BoxTM);
+			
+			if (Box.Intersects(Bounds))
+			{
+				Instances.Add(Idx);
+				if (bAny)
+				{
+					return false;
+				}
+			}
+
+			return true;
+		}
+
+		bool VisitRaycast(TSpatialVisitorData<int32> Idx, FQueryFastData&)
+		{
+			check(false);
+			return false;
+		}
+
+		bool VisitSweep(TSpatialVisitorData<int32> Idx, FQueryFastData&)
+		{
+			check(false);
+			return false;
+		}
+
+		virtual bool Overlap(const TSpatialVisitorData<int32>& Instance) override
+		{
+			return VisitOverlap(Instance);
+		}
+
+		virtual bool Raycast(const TSpatialVisitorData<int32>& Instance, FQueryFastData& CurData) override
+		{
+			return VisitRaycast(Instance, CurData);
+		}
+
+		virtual bool Sweep(const TSpatialVisitorData<int32>& Instance, FQueryFastData& CurData) override
+		{
+			return VisitSweep(Instance, CurData);
+		}
+
+		TArray<int32> Instances;
+	};
+
+	template <typename T>
+	auto BuildBoxes(TUniquePtr<TBox<T,3>>& Box)
+	{
+		Box = MakeUnique<TBox<T, 3>>(TVector<T, 3>(0, 0, 0), TVector<T, 3>(100, 100, 100));
+		auto Boxes = MakeUnique<TGeometryParticles<T, 3>>();
+		const int32 NumRows = 10;
+		const int32 NumCols = 10;
+		const int32 NumHeight = 10;
+
+		Boxes->AddParticles(NumRows * NumCols * NumHeight);
+
+		int32 Idx = 0;
+		for (int32 Height = 0; Height < NumHeight; ++Height)
+		{
+			for (int32 Row = 0; Row < NumRows; ++Row)
+			{
+				for (int32 Col = 0; Col < NumCols; ++Col)
+				{
+					Boxes->SetGeometry(Idx, MakeSerializable(Box));
+					Boxes->X(Idx) = TVector<T, 3>(Col * 100, Row * 100, Height * 100);
+					Boxes->R(Idx) = TRotation<T, 3>::Identity;
+					Boxes->LocalBounds(Idx) = Box->BoundingBox();
+					Boxes->HasBounds(Idx) = true;
+					Boxes->SetWorldSpaceInflatedBounds(Idx, Box->BoundingBox().TransformedAABB(TRigidTransform<T, 3>(Boxes->X(Idx), Boxes->R(Idx))));
+					++Idx;
+				}
+			}
+		}
+
+		return Boxes;
+	}
+	
+	template <typename TSpatial, typename T>
+	void SpatialTestHelper(TSpatial& Spatial, TGeometryParticles<T,3>* Boxes, TUniquePtr<TBox<T,3>>& Box, FSpatialAccelerationIdx SpatialIdx = FSpatialAccelerationIdx())
+	{
+		//raycast
+		//miss
+		{
+			TVisitor<T> Visitor(TVector<T, 3>(-100, 0, 0), TVector<T, 3>(0, 1, 0), 0, *Boxes);
+			Spatial.Raycast(Visitor.Start, Visitor.Dir, 1000, Visitor);
+			EXPECT_EQ(Visitor.Instances.Num(), 0);
+		}
+
+		//gather along ray
+		{
+			TVisitor<T> Visitor(TVector<T, 3>(10, 0, 0), TVector<T, 3>(0, 1, 0), 0, *Boxes);
+			Spatial.Raycast(Visitor.Start, Visitor.Dir, 1000, Visitor);
+			EXPECT_EQ(Visitor.Instances.Num(), 10);
+		}
+
+		//gather along ray and then make modifications
+		{
+			auto Spatial2 = Spatial.Copy();
+			TVisitor<T> Visitor(TVector<T, 3>(10, 0, 0), TVector<T, 3>(0, 1, 0), 0, *Boxes);
+			Spatial2->Raycast(Visitor.Start, Visitor.Dir, 1000, Visitor);
+			EXPECT_EQ(Visitor.Instances.Num(), 10);
+
+			//remove from structure
+			Spatial2->RemoveElementFrom(Visitor.Instances[0], SpatialIdx);
+
+			TVisitor<T> Visitor2(TVector<T, 3>(10, 0, 0), TVector<T, 3>(0, 1, 0), 0, *Boxes);
+			Spatial2->Raycast(Visitor2.Start, Visitor2.Dir, 1000, Visitor2);
+			EXPECT_EQ(Visitor2.Instances.Num(), 9);
+
+			//move instance away
+			{
+				const int32 MoveIdx = Visitor2.Instances[0];
+				Boxes->X(MoveIdx) += TVector<T, 3>(1000, 0, 0);
+				TAABB<T, 3> NewBounds = Boxes->Geometry(MoveIdx)->template GetObject<TBox<T, 3>>()->BoundingBox().TransformedAABB(TRigidTransform<T, 3>(Boxes->X(MoveIdx), Boxes->R(MoveIdx)));
+				Spatial2->UpdateElementIn(MoveIdx, NewBounds, true, SpatialIdx);
+
+				TVisitor<T> Visitor3(TVector<T, 3>(10, 0, 0), TVector<T, 3>(0, 1, 0), 0, *Boxes);
+				Spatial2->Raycast(Visitor3.Start, Visitor3.Dir, 1000, Visitor3);
+				EXPECT_EQ(Visitor3.Instances.Num(), 8);
+
+				//move instance back
+				Boxes->X(MoveIdx) -= TVector<T, 3>(1000, 0, 0);
+				NewBounds = Boxes->Geometry(MoveIdx)->template GetObject<TBox<T, 3>>()->BoundingBox().TransformedAABB(TRigidTransform<T, 3>(Boxes->X(MoveIdx), Boxes->R(MoveIdx)));
+				Spatial2->UpdateElementIn(MoveIdx, NewBounds, true, SpatialIdx);
+			}
+
+			//move other instance into view
+			{
+				const int32 MoveIdx = 5 * 5 * 5;
+				const TVector<T, 3> OldPos = Boxes->X(MoveIdx);
+				Boxes->X(MoveIdx) = TVector<T, 3>(0, 0, 0);
+				TAABB<T, 3> NewBounds = Boxes->Geometry(MoveIdx)->template GetObject<TBox<T, 3>>()->BoundingBox().TransformedAABB(TRigidTransform<T, 3>(Boxes->X(MoveIdx), Boxes->R(MoveIdx)));
+				Spatial2->UpdateElementIn(MoveIdx, NewBounds, true, SpatialIdx);
+
+				TVisitor<T> Visitor3(TVector<T, 3>(10, 0, 0), TVector<T, 3>(0, 1, 0), 0, *Boxes);
+				Spatial2->Raycast(Visitor3.Start, Visitor3.Dir, 1000, Visitor3);
+				EXPECT_EQ(Visitor3.Instances.Num(), 10);
+
+				//move instance back
+				Boxes->X(MoveIdx) = OldPos;
+				NewBounds = Boxes->Geometry(MoveIdx)->template GetObject<TBox<T,3>>()->BoundingBox().TransformedAABB(TRigidTransform<T, 3>(Boxes->X(MoveIdx), Boxes->R(MoveIdx)));
+				Spatial2->UpdateElementIn(MoveIdx, NewBounds, true, SpatialIdx);
+			}
+
+			//move instance outside of grid bounds
+			{
+				const int32 MoveIdx = 5 * 5 * 5;
+				const TVector<T, 3> OldPos = Boxes->X(MoveIdx);
+				Boxes->X(MoveIdx) = TVector<T, 3>(-50, 0, 0);
+				TAABB<T, 3> NewBounds = Boxes->Geometry(MoveIdx)->template GetObject<TBox<T,3>>()->BoundingBox().TransformedAABB(TRigidTransform<T, 3>(Boxes->X(MoveIdx), Boxes->R(MoveIdx)));
+				Spatial2->UpdateElementIn(MoveIdx, NewBounds, true, SpatialIdx);
+
+				TVisitor<T> Visitor3(TVector<T, 3>(10, 0, 0), TVector<T, 3>(0, 1, 0), 0, *Boxes);
+				Spatial2->Raycast(Visitor3.Start, Visitor3.Dir, 1000, Visitor3);
+				EXPECT_EQ(Visitor3.Instances.Num(), 10);
+
+				//try ray outside of bounds which should hit
+				TVisitor<T> Visitor4(TVector<T, 3>(-20, 0, 0), TVector<T, 3>(0, 1, 0), 0, *Boxes);
+				Spatial2->Raycast(Visitor4.Start, Visitor4.Dir, 1000, Visitor4);
+				EXPECT_EQ(Visitor4.Instances.Num(), 1);
+
+				//delete dirty instance
+				Spatial2->RemoveElementFrom(MoveIdx, SpatialIdx);
+				TVisitor<T> Visitor5(TVector<T, 3>(-20, 0, 0), TVector<T, 3>(0, 1, 0), 0, *Boxes);
+				Spatial2->Raycast(Visitor5.Start, Visitor5.Dir, 1000, Visitor4);
+				EXPECT_EQ(Visitor5.Instances.Num(), 0);
+
+				//move instance back
+				Boxes->X(MoveIdx) = OldPos;
+
+				//create a new box
+				const int32 NewIdx = Boxes->Size();
+				Boxes->AddParticles(1);
+				Boxes->SetGeometry(NewIdx, MakeSerializable(Box));
+				Boxes->X(NewIdx) = TVector<T, 3>(-20, 0, 0);
+				Boxes->R(NewIdx) = TRotation<T, 3>::Identity;
+				NewBounds = Boxes->Geometry(NewIdx)->template GetObject<TBox<T,3>>()->BoundingBox().TransformedAABB(TRigidTransform<T, 3>(Boxes->X(NewIdx), Boxes->R(NewIdx)));
+				Spatial2->UpdateElementIn(NewIdx, NewBounds, true, SpatialIdx);
+				TVisitor<T> Visitor6(TVector<T, 3>(-20, 0, 0), TVector<T, 3>(0, 1, 0), 0, *Boxes);
+				Spatial2->Raycast(Visitor6.Start, Visitor6.Dir, 1000, Visitor6);
+				EXPECT_EQ(Visitor6.Instances.Num(), 1);
+			}
+		}
+
+		//stop half way through
+		{
+			TVisitor<T> Visitor(TVector<T, 3>(10, 0, 0), TVector<T, 3>(0, 1, 0), 0, *Boxes);
+			Spatial.Raycast(Visitor.Start, Visitor.Dir, 499, Visitor);
+			EXPECT_EQ(Visitor.Instances.Num(), 5);
+		}
+
+		//any
+		{
+			TVisitor<T> Visitor(TVector<T, 3>(10, 0, 0), TVector<T, 3>(0, 1, 0), 0, *Boxes);
+			Visitor.bAny = true;
+			Spatial.Raycast(Visitor.Start, Visitor.Dir, 1000, Visitor);
+			EXPECT_EQ(Visitor.Instances.Num(), 1);
+		}
+
+		//sweep
+		//miss
+		{
+			TVisitor<T> Visitor(TVector<T, 3>(-100, 0, 0), TVector<T, 3>(0, 1, 0), 0, *Boxes);
+			Visitor.HalfExtents = TVector<T, 3>(10, 0, 0);
+			Spatial.Sweep(Visitor.Start, Visitor.Dir, 1000, Visitor.HalfExtents, Visitor);
+			EXPECT_EQ(Visitor.Instances.Num(), 0);
+		}
+
+		//gather along ray
+		{
+			TVisitor<T> Visitor(TVector<T, 3>(-100, 0, 0), TVector<T, 3>(0, 1, 0), 0, *Boxes);
+			Visitor.HalfExtents = TVector<T, 3>(110, 0, 0);
+			Spatial.Sweep(Visitor.Start, Visitor.Dir, 1000, Visitor.HalfExtents, Visitor);
+			EXPECT_EQ(Visitor.Instances.Num(), 10);
+		}
+
+		//stop half way through
+		{
+			TVisitor<T> Visitor(TVector<T, 3>(-100, 0, 0), TVector<T, 3>(0, 1, 0), 0, *Boxes);
+			Visitor.HalfExtents = TVector<T, 3>(110, 0, 0);
+			Spatial.Sweep(Visitor.Start, Visitor.Dir, 499, Visitor.HalfExtents, Visitor);
+			EXPECT_EQ(Visitor.Instances.Num(), 5);
+		}
+
+		//right on edge and corner
+		{
+			TVisitor<T> Visitor(TVector<T, 3>(100, 0, 0), TVector<T, 3>(0, 1, 0), 0, *Boxes);
+			Visitor.HalfExtents = TVector<T, 3>(10, 0, 0);
+			Spatial.Sweep(Visitor.Start, Visitor.Dir, 499, Visitor.HalfExtents, Visitor);
+			EXPECT_EQ(Visitor.Instances.Num(), 10);
+		}
+
+		//overlap
+		//miss
+		{
+			TOverlapVisitor<T> Visitor(TAABB<T, 3>(TVector<T, 3>(-100, 0, 0), TVector<T, 3>(-10, 0, 0)), *Boxes);
+			Spatial.Overlap(Visitor.Bounds, Visitor);
+			EXPECT_EQ(Visitor.Instances.Num(), 0);
+		}
+
+		//overlap some
+		{
+			TOverlapVisitor<T> Visitor(TAABB<T, 3>(TVector<T, 3>(-100, 0, -10), TVector<T, 3>(110, 110, 10)), *Boxes);
+			Spatial.Overlap(Visitor.Bounds, Visitor);
+			EXPECT_EQ(Visitor.Instances.Num(), 4);
+		}
+
+		//overlap any
+		{
+			TOverlapVisitor<T> Visitor(TAABB<T, 3>(TVector<T, 3>(-100, 0, -10), TVector<T, 3>(110, 110, 10)), *Boxes);
+			Visitor.bAny = true;
+			Spatial.Overlap(Visitor.Bounds, Visitor);
+			EXPECT_EQ(Visitor.Instances.Num(), 1);
+		}
+	}
+
+	template<typename T>
+	void GridBPTest()
+	{
+		TUniquePtr<TBox<T, 3>> Box;
+		auto Boxes = BuildBoxes<T>(Box);
+		TBoundingVolume<int32, T, 3> Spatial(MakeParticleView(Boxes.Get()));
+		SpatialTestHelper(Spatial, Boxes.Get(), Box);
+	}
+
+	template void GridBPTest<float>();
+
+	template<typename T>
+	void GridBPTest2()
+	{
+		TUniquePtr<TBox<T, 3>> Box = MakeUnique<TBox<T, 3>>(TVector<T, 3>(0, 0, 0), TVector<T, 3>(100, 100, 100));
+		TPBDRigidsSOAs<T, 3> SOAs;
+		const int32 NumRows = 10;
+		const int32 NumCols = 10;
+		const int32 NumHeight = 10;
+
+		SOAs.CreateStaticParticles(NumRows * NumCols * NumHeight);
+		auto& Boxes = SOAs.GetNonDisabledStaticParticles();
+		int32 Idx = 0;
+		for (int32 Height = 0; Height < NumHeight; ++Height)
+		{
+			for (int32 Row = 0; Row < NumRows; ++Row)
+			{
+				for (int32 Col = 0; Col < NumCols; ++Col)
+				{
+					Boxes.SetGeometry(Idx, MakeSerializable(Box));
+					Boxes.X(Idx) = TVector<T, 3>(Col * 100, Row * 100, Height * 100);
+					Boxes.R(Idx) = TRotation<T, 3>::Identity;
+					Boxes.LocalBounds(Idx) = Box->BoundingBox();
+					Boxes.HasBounds(Idx) = true;
+					Boxes.SetWorldSpaceInflatedBounds(Idx, Box->BoundingBox().TransformedAABB(TRigidTransform<T, 3>(Boxes.X(Idx), Boxes.R(Idx))));
+					++Idx;
+				}
+			}
+		}
+
+		TArray<TSOAView<TGeometryParticles<T, 3>>> TmpArray = { &Boxes };
+		TBoundingVolume<TGeometryParticleHandle<T,3>*, T, 3> BV(MakeParticleView(MoveTemp(TmpArray)));
+		TArray<TGeometryParticleHandle<T, 3>*> Handles = BV.FindAllIntersections(TAABB<T, 3>(TVector<T, 3>(0), TVector<T, 3>(10)));
+		EXPECT_EQ(Handles.Num(), 1);
+		EXPECT_EQ(Handles[0], Boxes.Handle(0));
+
+		Handles = BV.FindAllIntersections(TAABB<T, 3>(TVector<T, 3>(0), TVector<T, 3>(0, 0, 110)));
+		EXPECT_EQ(Handles.Num(), 2);
+
+		//create BV with an array of handles instead (useful for partial structures)
+		{
+			TBoundingVolume<TGeometryParticleHandle<T,3>*, T, 3> BV2(MakeHandleView(Handles));
+			TArray<TGeometryParticleHandle<T, 3>*> Handles2 = BV2.FindAllIntersections(TAABB<T, 3>(TVector<T, 3>(0), TVector<T, 3>(10)));
+			EXPECT_EQ(Handles2.Num(), 1);
+			EXPECT_EQ(Handles2[0], Boxes.Handle(0));
+
+			Handles2 = BV2.FindAllIntersections(TAABB<T, 3>(TVector<T, 3>(0), TVector<T, 3>(0, 0, 110)));
+			EXPECT_EQ(Handles2.Num(), 2);
+		}
+	}
+
+	template void GridBPTest2<float>();
+
+	template <typename T>
+	void AABBTreeTest()
+	{
+		{
+			TUniquePtr<TBox<T, 3>> Box;
+			auto Boxes = BuildBoxes<T>(Box);
+			TAABBTree<int32, TAABBTreeLeafArray<int32, T>, T> Spatial(MakeParticleView(Boxes.Get()));
+
+			while (!Spatial.IsAsyncTimeSlicingComplete())
+			{
+				Spatial.ProgressAsyncTimeSlicing(false);
+			}
+
+			SpatialTestHelper(Spatial, Boxes.Get(), Box);
+		}
+
+		{
+			TUniquePtr<TBox<T, 3>> Box;
+			auto Boxes = BuildBoxes<T>(Box);
+			TAABBTree<int32, TBoundingVolume<int32, T, 3>, T> Spatial(MakeParticleView(Boxes.Get()));
+
+			while (!Spatial.IsAsyncTimeSlicingComplete())
+			{
+				Spatial.ProgressAsyncTimeSlicing(false);
+			}
+
+			SpatialTestHelper(Spatial, Boxes.Get(), Box);
+		}
+	}
+	template void AABBTreeTest<float>();
+
+
+	template <typename T>
+	void AABBTreeTimesliceTest()
+	{
+		TUniquePtr<TBox<T, 3>> Box;
+		auto Boxes = BuildBoxes<T>(Box);
+
+		// build AABB in one go
+		TAABBTree<int32, TAABBTreeLeafArray<int32, T>, T> SpatialBuildImmediate(
+			MakeParticleView(Boxes.Get()) 
+			, TAABBTree<int32, TAABBTreeLeafArray<int32, T>, T>::DefaultMaxChildrenInLeaf 
+			, TAABBTree<int32, TAABBTreeLeafArray<int32, T>, T>::DefaultMaxTreeDepth
+			, TAABBTree<int32, TAABBTreeLeafArray<int32, T>, T>::DefaultMaxPayloadBounds
+			, 0); // build entire tree in one go, no timeslicing
+
+		EXPECT_TRUE(SpatialBuildImmediate.IsAsyncTimeSlicingComplete());
+		
+		// build AABB in time-sliced sections
+		TAABBTree<int32, TAABBTreeLeafArray<int32, T>, T> SpatialTimesliced(
+			MakeParticleView(Boxes.Get())
+			, TAABBTree<int32, TAABBTreeLeafArray<int32, T>, T>::DefaultMaxChildrenInLeaf
+			, TAABBTree<int32, TAABBTreeLeafArray<int32, T>, T>::DefaultMaxTreeDepth
+			, TAABBTree<int32, TAABBTreeLeafArray<int32, T>, T>::DefaultMaxPayloadBounds
+			, 20); // build in small iteration steps, 20 iterations per call to ProgressAsyncTimeSlicing
+		EXPECT_FALSE(SpatialTimesliced.IsAsyncTimeSlicingComplete());
+
+		while (!SpatialTimesliced.IsAsyncTimeSlicingComplete())
+		{
+			SpatialTimesliced.ProgressAsyncTimeSlicing(false);
+		}	
+
+		// now check both AABBs have the same hierarchy
+		// (indices will be different but walking tree should give same results)
+
+		TAABB<T, 3> Tmp = TAABB<T, 3>::ZeroAABB();
+
+		TArray<TAABB<T, 3>> AllBoundsBuildImmediate;
+		SpatialBuildImmediate.GetAsBoundsArray(AllBoundsBuildImmediate, 0, -1, Tmp);
+
+		TArray<TAABB<T, 3>> AllBoundsTimesliced;
+		SpatialTimesliced.GetAsBoundsArray(AllBoundsTimesliced, 0, -1, Tmp);
+
+		EXPECT_EQ(AllBoundsBuildImmediate.Num(), AllBoundsTimesliced.Num());
+
+		for (int i=0; i<AllBoundsBuildImmediate.Num(); i++)
+		{
+			EXPECT_EQ(AllBoundsBuildImmediate[i].Center(), AllBoundsTimesliced[i].Center());
+			EXPECT_EQ(AllBoundsBuildImmediate[i].Extents(), AllBoundsTimesliced[i].Extents());
+		}
+	}
+	template void AABBTreeTimesliceTest<float>();
+
+	template <typename T>
+	void BroadphaseCollectionTest()
+	{
+		using TreeType = TAABBTree<int32, TAABBTreeLeafArray<int32, T>, T>;
+		{
+			TUniquePtr<TBox<T, 3>> Box;
+			auto Boxes = BuildBoxes<T>(Box);
+			auto Spatial = MakeUnique<TreeType>(MakeParticleView(Boxes.Get()));
+
+			while (!Spatial->IsAsyncTimeSlicingComplete())
+			{
+				Spatial->ProgressAsyncTimeSlicing(false);
+			}
+
+			TSpatialAccelerationCollection<TreeType> AccelerationCollection;
+			AccelerationCollection.AddSubstructure(MoveTemp(Spatial), 0);
+			FSpatialAccelerationIdx SpatialIdx = { 0,0 };
+			SpatialTestHelper(AccelerationCollection, Boxes.Get(), Box, SpatialIdx);
+		}
+
+		{
+			using BVType = TBoundingVolume<int32, T, 3>;
+			TUniquePtr<TBox<T, 3>> Box;
+			auto Boxes0 = BuildBoxes<T>(Box);
+			auto Spatial0 = MakeUnique<TreeType>(MakeParticleView(Boxes0.Get()));
+			while (!Spatial0->IsAsyncTimeSlicingComplete())
+			{
+				Spatial0->ProgressAsyncTimeSlicing(false);
+			}
+
+			TGeometryParticles<T, 3> EmptyBoxes;
+			auto Spatial1 = MakeUnique<BVType>(MakeParticleView(&EmptyBoxes));
+			while (!Spatial1->IsAsyncTimeSlicingComplete())
+			{
+				Spatial1->ProgressAsyncTimeSlicing(false);
+			}
+
+			TSpatialAccelerationCollection<TreeType, BVType> AccelerationCollection;
+			AccelerationCollection.AddSubstructure(MoveTemp(Spatial0), 0);
+			AccelerationCollection.AddSubstructure(MoveTemp(Spatial1), 1);
+
+			FSpatialAccelerationIdx SpatialIdx = { 0,0 };
+			SpatialTestHelper(AccelerationCollection, Boxes0.Get(), Box, SpatialIdx);
+		}
+
+		{
+			using BVType = TBoundingVolume<int32, T, 3>;
+			TUniquePtr<TBox<T, 3>> Box;
+			auto Boxes1 = BuildBoxes<T>(Box);
+			TGeometryParticles<T, 3> EmptyBoxes;
+
+			auto Spatial0 = MakeUnique<TreeType>(MakeParticleView(&EmptyBoxes));
+			auto Spatial1 = MakeUnique<BVType>(MakeParticleView(Boxes1.Get()));
+
+			TSpatialAccelerationCollection<TreeType, BVType> AccelerationCollection;
+			AccelerationCollection.AddSubstructure(MoveTemp(Spatial0), 0);
+			AccelerationCollection.AddSubstructure(MoveTemp(Spatial1), 1);
+
+			FSpatialAccelerationIdx SpatialIdx = { 1,0 };
+			SpatialTestHelper(AccelerationCollection, Boxes1.Get(), Box, SpatialIdx);
+		}
+	}
+	template void BroadphaseCollectionTest<float>();
+
+	void TestPendingSpatialDataHandlePointerConflict()
+	{
+		// We are testing an edge case in the way we queue updates to our acceleration structures.
+		// Particle A is deleted and is pending removal from accel structure
+		// Particle B is allocated, coincidentally having the same handle pointer as Particle A
+		// Particle B is pending an insertion into structure
+		// Particle B is deleted 
+		// Queue is flushed, we are testing that Particle A was removed successfully.
+
+		// If this test fails, it means Particle B is probably stomping on data related to Particle A's removal.
+
+
+		TPBDRigidsSOAs<FReal, 3> Particles;
+		TArray<TPBDRigidParticleHandle<FReal, 3>*> ParticleHandles = Particles.CreateDynamicParticles(1);
+
+		auto Particle = ParticleHandles[0];
+		Particle->GTGeometryParticle() = TGeometryParticle<FReal, 3>::CreateParticle().Release();
+		Particle->X() = TVector<FReal, 3>(0, 0, 0);
+
+		TPBDRigidsEvolutionGBF<FReal, 3> Evolution(Particles);
+
+		// Flush spatial acceleration structures to put particle into structure.
+		Evolution.FlushSpatialAcceleration();
+
+		// Confirm particle is returned via query.
+		{
+			auto SpatialAccelerationCollection = Evolution.GetSpatialAcceleration();
+			TArray<FSpatialAccelerationIdx> SpatialIndices = SpatialAccelerationCollection->GetAllSpatialIndices();
+			TAABB<FReal, 3> QueryBounds(FVec3(-1, -1, -1), FVec3(1, 1, 1));
+			uint32 HitNum = 0;
+			for (FSpatialAccelerationIdx SpatialIdx : SpatialIndices)
+			{
+				auto SpatialAcceleration = SpatialAccelerationCollection->GetSubstructure(SpatialIdx);
+				TArray<TAccelerationStructureHandle<FReal, 3>> Hits = SpatialAcceleration->FindAllIntersections(QueryBounds);
+				HitNum += Hits.Num();
+			}
+
+			EXPECT_EQ(HitNum, 1);
+		}
+
+
+		// Removal from accel structure is now pending (not completely applied until flush)
+		Evolution.RemoveParticleFromAccelerationStructure(*Particle);
+
+
+		// Assign new game thread particle to handle (this simulates Particle B allocation, with reused Handle pointer)
+		Particle->GTGeometryParticle() = TGeometryParticle<FReal, 3>::CreateParticle().Release();
+
+		// Remove new Particle from accel structure
+		Evolution.RemoveParticleFromAccelerationStructure(*Particle);
+
+		// Flush changes, this should remove first particle successfully.
+		Evolution.FlushSpatialAcceleration();
+
+		// Confirm particle is not in internal structure by performing same query.
+		{
+			auto SpatialAccelerationCollection = Evolution.GetSpatialAcceleration();
+			TArray<FSpatialAccelerationIdx> SpatialIndices = SpatialAccelerationCollection->GetAllSpatialIndices();
+			TAABB<FReal, 3> QueryBounds(FVec3(-1, -1, -1), FVec3(1, 1, 1));
+			uint32 HitNum = 0;
+			for (FSpatialAccelerationIdx SpatialIdx : SpatialIndices)
+			{
+				auto SpatialAcceleration = SpatialAccelerationCollection->GetSubstructure(SpatialIdx);
+				TArray<TAccelerationStructureHandle<FReal, 3>> Hits = SpatialAcceleration->FindAllIntersections(QueryBounds);
+				HitNum += Hits.Num();
+			}
+
+			EXPECT_EQ(HitNum, 0);
+		}
+
+		// Confirm particle is not in external structure by swapping and testing.
+		{
+			using AABBTreeType = TAABBTree<TAccelerationStructureHandle<FReal, 3>, TAABBTreeLeafArray<TAccelerationStructureHandle<FReal, 3>, FReal>, FReal>;
+			TUniquePtr<ISpatialAccelerationCollection<TAccelerationStructureHandle<FReal, 3>, FReal, 3>> ExternalStructure = TUniquePtr<ISpatialAccelerationCollection<TAccelerationStructureHandle<FReal, 3>, FReal, 3>>(new TSpatialAccelerationCollection<AABBTreeType>());
+			Evolution.UpdateExternalAccelerationStructure(ExternalStructure);
+
+			ISpatialAccelerationCollection<TAccelerationStructureHandle<FReal, 3>, FReal, 3>& SpatialAccelerationCollection = *ExternalStructure.Get();
+			TArray<FSpatialAccelerationIdx> SpatialIndices = SpatialAccelerationCollection.GetAllSpatialIndices();
+			TAABB<FReal, 3> QueryBounds(FVec3(-1, -1, -1), FVec3(1, 1, 1));
+			uint32 HitNum = 0;
+			for (FSpatialAccelerationIdx SpatialIdx : SpatialIndices)
+			{
+				auto SpatialAcceleration = SpatialAccelerationCollection.GetSubstructure(SpatialIdx);
+				TArray<TAccelerationStructureHandle<FReal, 3>> Hits = SpatialAcceleration->FindAllIntersections(QueryBounds);
+				HitNum += Hits.Num();
+			}
+
+			EXPECT_EQ(HitNum, 0);
+		}
+	}
+}
\ No newline at end of file
diff --git a/Engine/Source/Programs/HeadlessChaos/Private/HeadlessChaosTestCloth.cpp b/Engine/Source/Programs/HeadlessChaos/Private/HeadlessChaosTestCloth.cpp
new file mode 100644
--- /dev/null
+++ b/Engine/Source/Programs/HeadlessChaos/Private/HeadlessChaosTestCloth.cpp
@@ -0,0 +1,290 @@
+// Copyright Epic Games, Inc. All Rights Reserved.
+
+#include "HeadlessChaosTestCloth.h"
+
+#include "HeadlessChaos.h"
+#include "HeadlessChaosTestUtility.h"
+
+#include "Chaos/PBDEvolution.h"
+#include "Chaos/GeometryParticles.h"
+#include "Chaos/PBDSpringConstraints.h"
+#include "Chaos/TriangleMesh.h"
+#include "Chaos/Utilities.h"
+
+#include "Modules/ModuleManager.h"
+
+namespace ChaosTest {
+
+	using namespace Chaos;
+	DEFINE_LOG_CATEGORY_STATIC(LogChaosTestCloth, Verbose, All);
+
+	template <class T>
+	TUniquePtr<TPBDEvolution<T, 3>> InitPBDEvolution(
+		const int32 NumIterations=1,
+		const T CollisionThickness=KINDA_SMALL_NUMBER,
+		const T SelfCollisionThickness=KINDA_SMALL_NUMBER,
+		const T Friction=0.0,
+		const T Damping=0.04)
+	{
+		Chaos::TPBDParticles<T, 3> Particles;
+		Chaos::TKinematicGeometryClothParticles<T, 3> RigidParticles;
+		TUniquePtr<Chaos::TPBDEvolution<T, 3>> Evolution(
+			new Chaos::TPBDEvolution<T, 3>(
+				MoveTemp(Particles),
+				MoveTemp(RigidParticles),
+				{},
+				NumIterations,
+				CollisionThickness,
+				SelfCollisionThickness,
+				Friction,
+				Damping));
+		return Evolution;
+	}
+
+	template <class T, class TEvolutionPtr>
+	void InitSingleParticle(
+		TEvolutionPtr& Evolution,
+		const TVector<T, 3>& Position = TVector<T, 3>(0),
+		const TVector<T, 3>& Velocity = TVector<T, 3>(0),
+		const T Mass = 1.0)
+	{
+		auto& Particles = Evolution->Particles();
+		const uint32 Idx = Particles.Size();
+		Particles.AddParticles(1);
+		Particles.X(Idx) = Position;
+		Particles.V(Idx) = Velocity;
+		Particles.M(Idx) = Mass;
+		Particles.InvM(Idx) = 1.0 / Mass;
+	}
+
+	template <class T, class TEvolutionPtr>
+	void InitTriMesh_EquilateralTri(
+		Chaos::TTriangleMesh<T>& TriMesh, 
+		TEvolutionPtr& Evolution, 
+		const TVector<T,3>& XOffset=TVector<T,3>(0))
+	{
+		auto& Particles = Evolution->Particles();
+		const uint32 InitialNumParticles = Particles.Size();
+		Chaos::TTriangleMesh<T>::InitEquilateralTriangleYZ(TriMesh, Particles);
+
+		// Initialize particles.  Use 1/3 area of connected triangles for particle mass.
+		for (uint32 i = InitialNumParticles; i < Particles.Size(); i++)
+		{
+			Particles.X(i) += XOffset;
+			Particles.V(i) = Chaos::TVector<T, 3>(0);
+			Particles.M(i) = 0;
+		}
+		for (const Chaos::TVector<int32, 3>& Tri : TriMesh.GetElements())
+		{
+			const T TriArea = 0.5 * Chaos::TVector<T, 3>::CrossProduct(
+				Particles.X(Tri[1]) - Particles.X(Tri[0]),
+				Particles.X(Tri[2]) - Particles.X(Tri[0])).Size();
+			for (int32 i = 0; i < 3; i++)
+			{
+				Particles.M(Tri[i]) += TriArea / 3;
+			}
+		}
+		for (uint32 i = InitialNumParticles; i < Particles.Size(); i++)
+		{
+			check(Particles.M(i) > SMALL_NUMBER);
+			Particles.InvM(i) = 1.0 / Particles.M(i);
+		}
+	}
+
+	template <class T, int d, class TEvolutionPtr>
+	void AddEdgeLengthConstraint(
+		TEvolutionPtr& Evolution,
+		const TArray<Chaos::TVector<int32, d>>& Topology,
+		const T Stiffness)
+	{
+		check(Stiffness >= 0. && Stiffness <= 1.);
+		Evolution->AddPBDConstraintFunction(
+			[SpringConstraints = Chaos::TPBDSpringConstraints<float, 3>(
+				Evolution->Particles(), Topology, Stiffness)](
+					TPBDParticles<float, 3>& InParticles, const float Dt) 
+			{
+				SpringConstraints.Apply(InParticles, Dt);
+			});
+	}
+
+	template <class T, class TEvolutionPtr>
+	void AdvanceTime(
+		TEvolutionPtr& Evolution,
+		const uint32 NumFrames=1,
+		const uint32 NumTimeStepsPerFrame=1,
+		const uint32 FPS=24)
+	{
+		check(NumTimeStepsPerFrame > 0);
+		Evolution->SetIterations(NumTimeStepsPerFrame);
+
+		check(FPS > 0);
+		const T Dt = 1.0 / FPS;
+		for (uint32 i = 0; i < NumFrames; i++)
+		{
+			Evolution->AdvanceOneTimeStep(Dt);
+		}
+	}
+
+	template <class T, class TParticleContainer>
+	TArray<Chaos::TVector<T, 3>> CopyPoints(const TParticleContainer& Particles)
+	{
+		TArray<Chaos::TVector<T, 3>> Points;
+		Points.SetNum(Particles.Size());
+
+		for (uint32 i = 0; i < Particles.Size(); i++)
+			Points[i] = Particles.X(i);
+
+		return Points;
+	}
+
+	template <class T, class TParticleContainer>
+	void Reset(TParticleContainer& Particles, const TArray<Chaos::TVector<T, 3>>& Points)
+	{
+		for (uint32 i = 0; i < Particles.Size(); i++)
+		{
+			Particles.X(i) = Points[i];
+			Particles.V(i) = Chaos::TVector<T, 3>(0);
+		}
+	}
+
+	template <class T>
+	TArray<T> GetDifference(const TArray<T>& A, const TArray<T>& B)
+	{
+		TArray<T> C;
+		C.SetNum(A.Num());
+		for (int32 i = 0; i < A.Num(); i++)
+			C[i] = A[i] - B[i];
+		return C;
+	}
+
+	template <class T>
+	TArray<T> GetMagnitude(const TArray<Chaos::TVector<T, 3>>& V)
+	{
+		TArray<T> M;
+		M.SetNum(V.Num());
+		for (int32 i = 0; i < V.Num(); i++)
+			M[i] = V[i].Size();
+		return M;
+	}
+
+	template <class T>
+	bool AllSame(const TArray<T>& V, int32& Idx, const T Tolerance=KINDA_SMALL_NUMBER)
+	{
+		if (!V.Num())
+			return true;
+		const T& V0 = V[0];
+		for (int32 i = 1; i < V.Num(); i++)
+		{
+			if (!FMath::IsNearlyEqual(V0, V[i], Tolerance))
+			{
+				Idx = i;
+				return false;
+			}
+		}
+		return true;
+	}
+
+	template <class T, class TV>
+	bool RunDropTest(
+		TUniquePtr<TPBDEvolution<T, 3>>& Evolution,
+		const T GravMag,
+		const TV& GravDir,
+		const TArray<TV>& InitialPoints,
+		const int32 SubFrameSteps,
+		const T DistTolerance,
+		const char* TestID)
+	{
+		const T PreTime = Evolution->GetTime();
+		AdvanceTime<T>(Evolution, 24, SubFrameSteps); // 1 second
+		const T PostTime = Evolution->GetTime();
+		EXPECT_NEAR(PostTime-PreTime, 1.0, KINDA_SMALL_NUMBER)
+			<< TestID
+			<< "Evolution advanced time by " << (PostTime - PreTime)
+			<< " seconds, expected 1.0 seconds.";
+
+		const TArray<Chaos::TVector<T, 3>> PostPoints = CopyPoints<T>(Evolution->Particles());
+		const TArray<Chaos::TVector<T, 3>> Diff = GetDifference(PostPoints, InitialPoints);
+		const TArray<T> ScalarDiff = GetMagnitude(Diff);
+
+		// All points did the same thing
+		int32 Idx = 0;
+		EXPECT_TRUE(AllSame(ScalarDiff, Idx, (T)0.1))
+			<< TestID
+			<< "Points fell different distances - Index 0: "
+			<< ScalarDiff[0] << " != Index " 
+			<< Idx << ": " 
+			<< ScalarDiff[Idx] << " +/- 0.1.";
+
+		// Fell the right amount
+		EXPECT_NEAR(ScalarDiff[0], (T)0.5 * GravMag, DistTolerance)
+			<< TestID
+			<< "Points fell by " << ScalarDiff[0] 
+			<< ", expected " << ((T)0.5 * GravMag) 
+			<< " +/- " << DistTolerance << "."; 
+
+		// Fell the right direction
+		const T DirDot = Chaos::TVector<T, 3>::DotProduct(GravDir, Diff[0].GetSafeNormal());
+		EXPECT_NEAR(DirDot, (T)1.0, (T)KINDA_SMALL_NUMBER)
+			<< TestID
+			<< "Points fell in different directions.";
+
+		return true;
+	}
+
+	template <class T>
+	void DeformableGravity()
+	{
+		const T DistTol = 0.0002;
+
+		//
+		// Initialize solver and gravity
+		//
+
+		TUniquePtr<TPBDEvolution<T, 3>> Evolution = InitPBDEvolution<T>();
+
+		const Chaos::TVector<float, 3> GravDir(0.f, 0.f, -1.f);
+		const T GravMag = 980.665;
+
+		//
+		// Drop a single particle
+		//
+
+		InitSingleParticle<T>(Evolution);
+		TArray<Chaos::TVector<T, 3>> InitialPoints = CopyPoints<T>(Evolution->Particles());
+
+		RunDropTest<T>(Evolution, GravMag, GravDir, InitialPoints, 1, DistTol, "Single point falling under gravity, iters: 1 - ");
+		Reset(Evolution->Particles(), InitialPoints);
+		RunDropTest<T>(Evolution, GravMag, GravDir, InitialPoints, 100, DistTol, "Single point falling under gravity, iters: 100 - ");
+		Reset(Evolution->Particles(), InitialPoints);
+
+		//
+		// Add a triangle mesh
+		//
+
+		Chaos::TTriangleMesh<T> TriMesh;
+		InitTriMesh_EquilateralTri(TriMesh, Evolution);
+		InitialPoints = CopyPoints<T>(Evolution->Particles());
+
+		// 
+		// Points falling under gravity
+		//
+
+		RunDropTest<T>(Evolution, GravMag, GravDir, InitialPoints, 1, DistTol, "Points falling under gravity, iters: 1 - ");
+		Reset(Evolution->Particles(), InitialPoints);
+		RunDropTest<T>(Evolution, GravMag, GravDir, InitialPoints, 100, DistTol, "Points falling under gravity, iters: 100 - ");
+		Reset(Evolution->Particles(), InitialPoints);
+
+		// 
+		// Points falling under gravity with edge length constraint
+		//
+
+		AddEdgeLengthConstraint(Evolution, TriMesh.GetSurfaceElements(), 1.0);
+
+		RunDropTest<T>(Evolution, GravMag, GravDir, InitialPoints, 1, DistTol, "Points falling under gravity & edge cnstr, iters: 1 - ");
+		Reset(Evolution->Particles(), InitialPoints);
+		RunDropTest<T>(Evolution, GravMag, GravDir, InitialPoints, 100, DistTol, "Points falling under gravity & edge cnstr, iters: 100 - ");
+		Reset(Evolution->Particles(), InitialPoints);
+	}
+	template void DeformableGravity<float>();
+
+} // namespace ChaosTest
\ No newline at end of file
diff --git a/Engine/Source/Programs/HeadlessChaos/Private/HeadlessChaosTestClustering.cpp b/Engine/Source/Programs/HeadlessChaos/Private/HeadlessChaosTestClustering.cpp
new file mode 100644
--- /dev/null
+++ b/Engine/Source/Programs/HeadlessChaos/Private/HeadlessChaosTestClustering.cpp
@@ -0,0 +1,276 @@
+// Copyright Epic Games, Inc. All Rights Reserved.
+
+#include "HeadlessChaosTestClustering.h"
+
+#include "HeadlessChaos.h"
+#include "Modules/ModuleManager.h"
+#include "Chaos/PBDRigidsEvolution.h"
+#include "Chaos/PBDRigidParticles.h"
+#include "Chaos/Box.h"
+#include "Chaos/Plane.h"
+#include "Chaos/Sphere.h"
+#include "Chaos/ImplicitObject.h"
+#include "Chaos/ImplicitObjectTransformed.h"
+#include "Chaos/ImplicitObjectUnion.h"
+#include "Chaos/Levelset.h"
+#include "Chaos/UniformGrid.h"
+#include "Chaos/Utilities.h"
+#include "Chaos/PBDRigidsEvolutionGBF.h"
+
+namespace ChaosTest {
+
+	using namespace Chaos;
+
+	template<class T>
+	void ImplicitCluster()
+	{
+		TPBDRigidsSOAs<T, 3> Particles;
+		TPBDRigidsEvolutionGBF<T, 3> Evolution(Particles);
+		auto& ClusteredParticles = Particles.GetClusteredParticles();
+		
+		uint32 FirstId = ClusteredParticles.Size();
+		auto& Box1 = *AppendClusteredParticleBox<T>(Particles, TVector<T, 3>((T)100, (T)100, (T)100));
+		uint32 BoxId = FirstId++;
+		auto& Box2 = *AppendClusteredParticleBox<T>(Particles, TVector<T, 3>((T)100, (T)100, (T)100));
+		uint32 Box2Id = FirstId++;
+		
+		Box2.X() = TVector<T, 3>((T)100, (T)0, (T)0);
+		Box2.P() = Box2.X();
+
+		Evolution.AdvanceOneTimeStep(0);	//hack to initialize islands
+		//Evolution.InitializeAccelerationStructures();	//make sure islands are created
+
+		Evolution.GetRigidClustering().CreateClusterParticle(0,TArray<uint32>{ BoxId, Box2Id }, TSerializablePtr<FImplicitObject>(), nullptr);
+		EXPECT_EQ(ClusteredParticles.Size(), 3);
+
+		TVector<T, 3> ClusterX = ClusteredParticles.X(2);
+		TRotation<T,3> ClusterRot = ClusteredParticles.R(2);
+
+		EXPECT_TRUE(ClusterX.Equals(TVector<T, 3> {(T)50, 0, 0}));
+		EXPECT_TRUE(ClusterRot.Equals(TRotation<T, 3>::Identity));
+		EXPECT_TRUE(ClusterX.Equals(ClusteredParticles.P(2)));
+		EXPECT_TRUE(ClusterRot.Equals(ClusteredParticles.Q(2)));
+
+		TRigidTransform<float, 3> ClusterTM(ClusterX, ClusterRot);
+		TVector<T, 3> LocalPos = ClusterTM.InverseTransformPositionNoScale(TVector<T, 3> {(T)200, (T)0, (T)0});
+		TVector<T, 3> Normal;
+		T Phi = ClusteredParticles.Geometry(2)->PhiWithNormal(LocalPos, Normal);
+		EXPECT_TRUE(FMath::IsNearlyEqual(Phi, (T)50));
+		EXPECT_TRUE(Normal.Equals(TVector<T, 3>{(T)1, (T)0, (T)0}));
+
+		//EXPECT_TRUE(Evolution.GetParticles().Geometry(2)->IsConvex());  we don't actually guarantee this
+	}
+	template void ImplicitCluster<float>();
+
+	template<class T>
+	void FractureCluster()
+	{
+		TPBDRigidsSOAs<T, 3> Particles;
+		TPBDRigidsEvolutionGBF<T, 3> Evolution(Particles);
+		auto& ClusteredParticles = Particles.GetClusteredParticles();
+
+		//create a long row of boxes - the depth 0 cluster is the entire row, the depth 1 clusters 4 boxes each, the depth 2 clusters are 1 box each
+
+		constexpr int32 NumBoxes = 32;
+		TArray<TGeometryParticleHandle<T, 3>*> Boxes;
+		TArray<uint32> BoxIDs;
+		for (int i = 0; i < NumBoxes; ++i)
+		{
+			BoxIDs.Add(ClusteredParticles.Size());
+			auto Box = AppendClusteredParticleBox<T>(Particles, TVector<T, 3>((T)100, (T)100, (T)100));
+			Box->X() = TVector<T, 3>((T)i * (T)100, (T)0, (T)0);
+			Box->P() = Box->X();
+			Boxes.Add(Box);
+		}
+
+		Evolution.AdvanceOneTimeStep(0);	//hack to generate islands
+
+		TArray<uint32> ClusterIDs;
+
+		for (int i = 0; i < NumBoxes / 4; ++i)
+		{
+			ClusterIDs.Add(Evolution.GetRigidClustering().CreateClusterParticle(0, TArray<uint32>{ BoxIDs[i*4], BoxIDs[i*4+1], BoxIDs[i*4+2], BoxIDs[i*4+3] }, TSerializablePtr<FImplicitObject>(), nullptr));
+		}
+
+		uint32 RootClusterID = Evolution.GetRigidClustering().CreateClusterParticle(0, TArray<uint32>(ClusterIDs), TSerializablePtr<FImplicitObject>(), nullptr);
+		TVector<T, 3> InitialVelocity((T)50, (T)20, (T)100);
+		ClusteredParticles.V(RootClusterID) = InitialVelocity;
+		
+		constexpr int NumParticles = NumBoxes + NumBoxes / 4 + 1;
+		EXPECT_EQ(ClusteredParticles.Size(), NumParticles);
+
+		for (int i = 0; i < NumParticles-1; ++i)
+		{
+			EXPECT_TRUE(ClusteredParticles.Disabled(i));
+		}
+
+		EXPECT_TRUE(ClusteredParticles.Disabled(RootClusterID) == false);
+		EXPECT_EQ(Particles.GetNonDisabledView().Num(), 1);
+		EXPECT_EQ(Particles.GetNonDisabledView().Begin()->Handle(), ClusteredParticles.Handle(RootClusterID));
+
+		const T Dt = 0;	//don't want to integrate gravity, just want to test fracture
+		Evolution.AdvanceOneTimeStep(Dt);
+		EXPECT_TRUE(ClusteredParticles.Disabled(RootClusterID));	//not a cluster anymore, so disabled
+		for (auto& Particle : Particles.GetNonDisabledView())
+		{
+			EXPECT_NE(Particle.Handle(), ClusteredParticles.Handle(RootClusterID));	//view no longer contains root
+		}
+
+		EXPECT_EQ(Particles.GetNonDisabledView().Num(), NumBoxes / 4);
+
+		//children are still in a cluster, so disabled
+		for (uint32 BoxID : BoxIDs)
+		{
+			EXPECT_TRUE(ClusteredParticles.Disabled(BoxID));
+			for (auto& Particle : Particles.GetNonDisabledView())
+			{
+				EXPECT_NE(Particle.Handle(), ClusteredParticles.Handle(BoxID));	//make sure boxes are not in non disabled array
+			}
+			
+			for (int32 Island = 0; Island < Evolution.NumIslands(); ++Island)
+			{
+				EXPECT_TRUE(Evolution.GetIslandParticles(Island).Contains(ClusteredParticles.Handle(BoxID)) == false);
+			}
+		}
+
+		for (uint32 ClusterID : ClusterIDs)
+		{
+			EXPECT_TRUE(ClusteredParticles.Disabled(ClusterID) == false);	//not a cluster anymore, so disabled
+			bool bFoundInNonDisabled = false;
+			for (auto& Particle : Particles.GetNonDisabledView())
+			{
+				bFoundInNonDisabled |= Particle.Handle() == ClusteredParticles.Handle(ClusterID);
+			}
+
+			EXPECT_TRUE(bFoundInNonDisabled);	//clusters are enabled and in non disabled array
+			EXPECT_TRUE(ClusteredParticles.V(ClusterID).Equals(InitialVelocity));
+		}
+
+		Evolution.AdvanceOneTimeStep(Dt);
+		//second fracture, all clusters are now disabled
+		for (uint32 ClusterID : ClusterIDs)
+		{
+			EXPECT_TRUE(ClusteredParticles.Disabled(ClusterID) == true);	//not a cluster anymore, so disabled
+			for (auto& Particle : Particles.GetNonDisabledView())
+			{
+				EXPECT_NE(Particle.Handle(), ClusteredParticles.Handle(ClusterID));	//make sure boxes are not in non disabled array
+			}
+
+			for (int32 Island = 0; Island < Evolution.NumIslands(); ++Island)
+			{
+				EXPECT_TRUE(Evolution.GetIslandParticles(Island).Contains(ClusteredParticles.Handle(ClusterID)) == false);
+			}
+		}
+
+		EXPECT_EQ(Particles.GetNonDisabledView().Num(), NumBoxes);
+		for (uint32 BoxID : BoxIDs)
+		{
+			EXPECT_TRUE(ClusteredParticles.Disabled(BoxID) == false);
+			bool bFoundInNonDisabled = false;
+			for (auto& Particle : Particles.GetNonDisabledView())
+			{
+				bFoundInNonDisabled |= Particle.Handle() == ClusteredParticles.Handle(BoxID);
+			}
+			EXPECT_TRUE(bFoundInNonDisabled);
+			EXPECT_TRUE(ClusteredParticles.V(BoxID).Equals(InitialVelocity));
+		}
+	}
+	template void FractureCluster<float>();
+
+	template<class T>
+	void PartialFractureCluster()
+	{
+		TPBDRigidsSOAs<T, 3> Particles;
+		TPBDRigidsEvolutionGBF<T, 3> Evolution(Particles);
+		auto& ClusteredParticles = Particles.GetClusteredParticles();
+
+		//create a long row of boxes - the depth 0 cluster is the entire row, the depth 1 clusters 4 boxes each, the depth 2 clusters are 1 box each
+
+		constexpr int32 NumBoxes = 32;
+		TArray<TGeometryParticleHandle<T, 3>*> Boxes;
+		TArray<uint32> BoxIDs;
+		for (int i = 0; i < NumBoxes; ++i)
+		{
+			BoxIDs.Add(ClusteredParticles.Size());
+			auto Box = AppendClusteredParticleBox<T>(Particles, TVector<T, 3>((T)100, (T)100, (T)100));
+			Box->X() = TVector<T, 3>((T)i * (T)100, (T)0, (T)0);
+			Box->P() = Box->X();
+			Boxes.Add(Box);
+		}
+
+		Evolution.AdvanceOneTimeStep(0);	//hack to generate islands
+
+		TArray<uint32> ClusterIDs;
+
+		for (int i = 0; i < NumBoxes / 4; ++i)
+		{
+			ClusterIDs.Add(Evolution.GetRigidClustering().CreateClusterParticle(0, TArray<uint32>{ BoxIDs[i * 4], BoxIDs[i * 4 + 1], BoxIDs[i * 4 + 2], BoxIDs[i * 4 + 3] }, TSerializablePtr<FImplicitObject>(), nullptr));
+		}
+
+		TArray<uint32> ClusterIDsDup = ClusterIDs;
+
+		uint32 RootClusterID = Evolution.GetRigidClustering().CreateClusterParticle(0, TArray<uint32>(ClusterIDs), TSerializablePtr<FImplicitObject>(), nullptr);
+		TVector<T, 3> InitialVelocity((T)50, (T)20, (T)100);
+		ClusteredParticles.V(RootClusterID) = InitialVelocity;
+
+		TUniquePtr<FChaosPhysicsMaterial> PhysicalMaterial = MakeUnique<FChaosPhysicsMaterial>();
+		PhysicalMaterial->Friction = 0;
+		PhysicalMaterial->Restitution = 0;
+		PhysicalMaterial->SleepingLinearThreshold = 0;
+		PhysicalMaterial->SleepingAngularThreshold = 0;
+		PhysicalMaterial->DisabledLinearThreshold = 0;
+		PhysicalMaterial->DisabledAngularThreshold = 0;
+
+		Chaos::TArrayCollectionArray<float>& SolverStrainArray = Evolution.GetRigidClustering().GetStrainArray();
+
+		for (int i = 0; i < NumBoxes + NumBoxes / 4 + 1; ++i)
+		{
+			SolverStrainArray[i] = (T)1;
+			Evolution.SetPhysicsMaterial(ClusteredParticles.Handle(i), MakeSerializable(PhysicalMaterial));
+		}
+
+		Evolution.AdvanceOneTimeStep((T)1 / (T)60);
+		EXPECT_TRUE(ClusteredParticles.Disabled(RootClusterID) == false);	//strain > 0 so no fracture yet
+
+		SolverStrainArray[ClusterIDsDup[2]] = (T)0;	//fracture the third cluster, this should leave us with three pieces (0, 1), (2), (3,4,5,6,7)
+		Evolution.AdvanceOneTimeStep((T)1 / (T)60);
+		//EXPECT_TRUE(Evolution.GetParticles().Disabled(RootClusterID) == false);	//one of the connected pieces should re-use this
+		EXPECT_TRUE(ClusteredParticles.Disabled(ClusterIDsDup[2]) == false);	//this cluster is on its own and should be enabled 
+		
+		EXPECT_EQ(Evolution.GetActiveClusteredArray().Num(), 3);	//there should only be 3 pieces
+		for (uint32 BoxID : BoxIDs)
+		{
+			EXPECT_TRUE(ClusteredParticles.Disabled(BoxID));	//no boxes should be active yet
+			EXPECT_TRUE(Evolution.GetActiveClusteredArray().Contains(ClusteredParticles.Handle(BoxID)) == false);
+			for (int32 Island = 0; Island < Evolution.NumIslands(); ++Island)
+			{
+				EXPECT_TRUE(Evolution.GetIslandParticles(Island).Contains(ClusteredParticles.Handle(BoxID)) == false);
+			}
+		}
+
+		SolverStrainArray[NumBoxes + NumBoxes / 4 + 1] = (T)1;
+		Evolution.SetPhysicsMaterial(ClusteredParticles.Handle(NumBoxes + NumBoxes / 4 + 1), MakeSerializable(PhysicalMaterial));
+		SolverStrainArray[NumBoxes + NumBoxes / 4 + 2] = (T)1;
+		Evolution.SetPhysicsMaterial(ClusteredParticles.Handle(NumBoxes + NumBoxes / 4 + 2), MakeSerializable(PhysicalMaterial));
+
+		Evolution.AdvanceOneTimeStep((T)1 / (T)60);	//next frame nothing should fracture
+		//EXPECT_TRUE(Evolution.GetParticles().Disabled(RootClusterID) == false);	//one of the connected pieces should re-use this
+		EXPECT_TRUE(ClusteredParticles.Disabled(ClusterIDsDup[2]) == false);	//this cluster is on its own and should be enabled 
+
+		EXPECT_EQ(Evolution.GetActiveClusteredArray().Num(), 3);	//there should only be 3 pieces
+		for (uint32 BoxID : BoxIDs)
+		{
+			EXPECT_TRUE(ClusteredParticles.Disabled(BoxID));	//no boxes should be active yet
+			EXPECT_TRUE(Evolution.GetActiveClusteredArray().Contains(ClusteredParticles.Handle(BoxID)) == false);
+			for (int32 Island = 0; Island < Evolution.NumIslands(); ++Island)
+			{
+				EXPECT_TRUE(Evolution.GetIslandParticles(Island).Contains(ClusteredParticles.Handle(BoxID)) == false);
+			}
+		}
+	}
+	template void PartialFractureCluster<float>();
+
+
+
+}
+
diff --git a/Engine/Source/Programs/HeadlessChaos/Private/HeadlessChaosTestCollisions.cpp b/Engine/Source/Programs/HeadlessChaos/Private/HeadlessChaosTestCollisions.cpp
new file mode 100644
--- /dev/null
+++ b/Engine/Source/Programs/HeadlessChaos/Private/HeadlessChaosTestCollisions.cpp
@@ -0,0 +1,595 @@
+// Copyright Epic Games, Inc. All Rights Reserved.
+
+#include "HeadlessChaosTestCollisions.h"
+
+#include "HeadlessChaos.h"
+#include "HeadlessChaosCollisionConstraints.h"
+#include "Chaos/GJK.h"
+#include "Chaos/Pair.h"
+#include "Chaos/PBDRigidsEvolution.h"
+#include "Chaos/PBDRigidParticles.h"
+#include "Chaos/PBDCollisionConstraintsPGS.h"
+#include "Chaos/CollisionResolutionTypes.h"
+#include "Chaos/Sphere.h"
+#include "Chaos/Utilities.h"
+#include "Modules/ModuleManager.h"
+
+#define SMALL_THRESHOLD 1e-4
+
+#define RESET_PQ(Particle) Particle->P() = Particle->X(); Particle->Q() = Particle->R()
+#define INVARIANT_XR_START(Particle) TVector<T,3> InvariantPreX_##Particle = Particle->X(); TRotation<T,3> InvariantPreR_##Particle = Particle->R()
+#define INVARIANT_XR_END(Particle) EXPECT_TRUE(InvariantPreX_##Particle.Equals(Particle->X())); EXPECT_TRUE(InvariantPreR_##Particle.Equals(Particle->R()))
+#define INVARIANT_VW_START(Particle) TVector<T,3> InvariantPreV_##Particle = Particle->V(); TVector<T,3> InvariantPreW_##Particle = Particle->W()
+#define INVARIANT_VW_END(Particle) EXPECT_TRUE(InvariantPreV_##Particle.Equals(Particle->V())); EXPECT_TRUE(InvariantPreW_##Particle.Equals(Particle->W()))
+
+namespace ChaosTest {
+
+	using namespace Chaos;
+
+	DEFINE_LOG_CATEGORY_STATIC(LogHChaosCollisions, Verbose, All);
+
+	template<class T>
+	void LevelsetConstraint()
+	{
+		TArrayCollectionArray<bool> Collided;
+		TUniquePtr<FChaosPhysicsMaterial> PhysicsMaterial = MakeUnique<FChaosPhysicsMaterial>();
+		PhysicsMaterial->Friction = (T)0;
+		PhysicsMaterial->Restitution = (T)0;
+		TArrayCollectionArray<TSerializablePtr<FChaosPhysicsMaterial>> PhysicsMaterials;
+		
+		TPBDRigidsSOAs<T, 3> Particles;
+		Particles.GetParticleHandles().AddArray(&Collided);
+		Particles.GetParticleHandles().AddArray(&PhysicsMaterials);
+
+		auto Box1 = AppendDynamicParticleBox<T>(Particles);
+		Box1->X() = TVector<T, 3>(1.f);
+		Box1->R() = TRotation<T, 3>(FQuat::Identity);
+		Box1->P() = Box1->X();
+		Box1->Q() = Box1->R();
+		Box1->AuxilaryValue(PhysicsMaterials) = MakeSerializable(PhysicsMaterial);
+
+		auto Box2 = AppendDynamicParticleBox<T>(Particles);
+		Box2->X() = TVector<T, 3>(1.5f, 1.5f, 1.9f);
+		Box2->R() = TRotation<T, 3>(FQuat::Identity);
+		Box2->P() = Box2->X();
+		Box2->Q() = Box2->R();
+		Box2->AuxilaryValue(PhysicsMaterials) = MakeSerializable(PhysicsMaterial);
+
+		TPBDCollisionConstraintAccessor<T, 3> Collisions(Particles, Collided, PhysicsMaterials, 1, 1, T(0));
+		Collisions.ComputeConstraints(0.f);
+		EXPECT_EQ(Collisions.NumConstraints(), 1);
+
+		TPBDCollisionConstraints<T, 3>::FConstraintBase& Constraint = Collisions.GetConstraint(0);
+		if (auto PBDRigid = Constraint.Particle[0]->CastToRigidParticle())
+		{
+			//Question: non dynamics don't have collision particles, seems wrong if the levelset is dynamic and the static is something like a box
+			PBDRigid->CollisionParticles()->UpdateAccelerationStructures();
+		}
+		Collisions.UpdateLevelsetConstraint(*Constraint.template As<TPBDCollisionConstraints<T,3>::FPointContactConstraint>());
+
+		EXPECT_EQ(Constraint.Particle[0], Box2);
+		EXPECT_EQ(Constraint.Particle[1], Box1);
+		EXPECT_TRUE(Constraint.GetNormal().operator==(TVector<T, 3>(0, 0, 1)));
+		EXPECT_TRUE(FMath::Abs(ChaosTest::SignedDistance(*Constraint.Particle[0], Constraint.GetLocation())) < SMALL_THRESHOLD);
+	}
+	template void LevelsetConstraint<float>();
+
+
+	template<class T>
+	void LevelsetConstraintGJK()
+	{
+		TArrayCollectionArray<bool> Collided;
+		TUniquePtr<FChaosPhysicsMaterial> PhysicsMaterial = MakeUnique<FChaosPhysicsMaterial>();
+		PhysicsMaterial->Friction = (T)0;
+		PhysicsMaterial->Restitution = (T)0;
+		TArrayCollectionArray<TSerializablePtr<FChaosPhysicsMaterial>> PhysicsMaterials;
+
+		TPBDRigidsSOAs<T, 3> Particles;
+		Particles.GetParticleHandles().AddArray(&Collided);
+		Particles.GetParticleHandles().AddArray(&PhysicsMaterials);
+
+		auto Box1 = AppendDynamicParticleConvexBox<T>(Particles, TVector<T,3>(1.f) );
+		Box1->X() = TVector<T, 3>(0.f);
+		Box1->R() = TRotation<T, 3>(FQuat::Identity);
+		Box1->P() = Box1->X();
+		Box1->Q() = Box1->R();
+		Box1->AuxilaryValue(PhysicsMaterials) = MakeSerializable(PhysicsMaterial);
+
+		auto Box2 = AppendDynamicParticleBox<T>(Particles, TVector<T, 3>(1.f) );
+		Box2->X() = TVector<T, 3>(1.25f, 0.f, 0.f);
+		Box2->R() = TRotation<T, 3>(FQuat::Identity);
+		Box2->P() = Box2->X();
+		Box2->Q() = Box2->R();
+		Box2->AuxilaryValue(PhysicsMaterials) = MakeSerializable(PhysicsMaterial);
+
+		TPBDCollisionConstraintAccessor<T, 3> Collisions(Particles, Collided, PhysicsMaterials, 1, 1, T(0));
+		Collisions.ComputeConstraints(0.f);
+		EXPECT_EQ(Collisions.NumConstraints(), 1);
+
+		TPBDCollisionConstraints<T, 3>::FConstraintBase & Constraint = Collisions.GetConstraint(0);
+		Collisions.UpdateLevelsetConstraint(*Constraint.template As<TPBDCollisionConstraints<T,3>::FPointContactConstraint>());
+		
+		EXPECT_EQ(Constraint.Particle[0], Box2);
+		EXPECT_EQ(Constraint.Particle[1], Box1);
+		EXPECT_TRUE(Constraint.GetNormal().operator==(TVector<T, 3>(0, 0, 1)));
+		EXPECT_TRUE(FMath::Abs(ChaosTest::SignedDistance(*Constraint.Particle[0], Constraint.GetLocation())) < SMALL_THRESHOLD);
+	}
+	template void LevelsetConstraintGJK<float>();
+
+	template<class T>
+	void CollisionBoxPlane()
+	{
+		// test a box and plane in a colliding state
+		TArrayCollectionArray<bool> Collided;
+		TUniquePtr<FChaosPhysicsMaterial> PhysicsMaterial = MakeUnique<FChaosPhysicsMaterial>();
+		PhysicsMaterial->Friction = (T)0;
+		PhysicsMaterial->Restitution = (T)1;
+		TArrayCollectionArray<TSerializablePtr<FChaosPhysicsMaterial>> PhysicsMaterials;
+
+		TPBDRigidsSOAs<T, 3> Particles;
+		Particles.GetParticleHandles().AddArray(&Collided);
+		Particles.GetParticleHandles().AddArray(&PhysicsMaterials);
+
+		auto Floor = AppendStaticAnalyticFloor<T>(Particles);
+		auto Box = AppendDynamicParticleBox<T>(Particles);
+		Box->X() = TVector<T, 3>(0, 1, 0);
+		Box->R() = TRotation<T, 3>(FQuat::Identity);
+		Box->V() = TVector<T, 3>(0, 0, -1);
+		Box->P() = Box->X();
+		Box->Q() = Box->R();
+		Box->AuxilaryValue(PhysicsMaterials) = MakeSerializable(PhysicsMaterial);
+
+		const float Dt = 1 / 24.;
+
+		TPBDCollisionConstraintAccessor<T, 3> Collisions(Particles, Collided, PhysicsMaterials, 2, 5, T(0));
+
+		Collisions.ComputeConstraints(Dt);
+		EXPECT_EQ(Collisions.NumConstraints(), 1);
+
+		TPBDCollisionConstraints<T, 3>::FConstraintBase & Constraint = Collisions.GetConstraint(0);
+		if (auto PBDRigid = Constraint.Particle[0]->CastToRigidParticle())
+		{
+			PBDRigid->CollisionParticles()->UpdateAccelerationStructures();
+		}
+		Collisions.UpdateLevelsetConstraint(*Constraint.template As<TPBDCollisionConstraints<T, 3>::FPointContactConstraint>());
+
+		EXPECT_EQ(Constraint.Particle[0], Box);
+		EXPECT_EQ(Constraint.Particle[1], Floor);
+		EXPECT_TRUE(Constraint.GetNormal().operator==(TVector<T, 3>(0, 0, 1)));
+		EXPECT_TRUE(FMath::Abs(ChaosTest::SignedDistance(*Constraint.Particle[0], Constraint.GetLocation())) < SMALL_THRESHOLD);
+		EXPECT_TRUE(FMath::Abs(Constraint.GetPhi() - T(-0.5) ) < SMALL_THRESHOLD );
+
+		{
+			INVARIANT_XR_START(Box);
+			Collisions.Apply(Dt, { Collisions.GetConstraintHandle(0) }, 0, 1);
+			INVARIANT_XR_END(Box);
+		}
+
+		// 0 restitution so expecting 0 velocity
+		EXPECT_TRUE(Box->V().Equals(TVector<T, 3>(0)));
+		EXPECT_TRUE(Box->W().Equals(TVector<T, 3>(0)));
+
+		{
+			RESET_PQ(Box);
+			{
+				INVARIANT_XR_START(Box);
+				INVARIANT_VW_START(Box);
+				Collisions.ApplyPushOut(Dt, { Collisions.GetConstraintHandle(0) }, TSet<const TGeometryParticleHandle<T,3>*>(), 0, 1);
+				INVARIANT_XR_END(Box);
+				INVARIANT_VW_END(Box);
+			}
+		}
+		EXPECT_TRUE(Box->P().Equals(FVector(0.f, 1.f, 0.5f)));
+		EXPECT_TRUE(Box->Q().Equals(FQuat::Identity));
+	}
+	template void CollisionBoxPlane<float>();
+
+
+	template<class T>
+	void CollisionConvexConvex()
+	{
+
+		// test a box and plane in a colliding state
+		TArrayCollectionArray<bool> Collided;
+		TUniquePtr<FChaosPhysicsMaterial> PhysicsMaterial = MakeUnique<FChaosPhysicsMaterial>();
+		PhysicsMaterial->Friction = (T)0;
+		PhysicsMaterial->Restitution = (T)0;
+		TArrayCollectionArray<TSerializablePtr<FChaosPhysicsMaterial>> PhysicsMaterials;
+
+		TPBDRigidsSOAs<T, 3> Particles;
+		Particles.GetParticleHandles().AddArray(&Collided);
+		Particles.GetParticleHandles().AddArray(&PhysicsMaterials);
+
+		auto Floor = AppendStaticConvexFloor<T>(Particles);
+		auto Box = AppendDynamicParticleConvexBox<T>( Particles, TVector<T, 3>(50) );
+		Box->X() = TVector<T, 3>(0, 0, 49);
+		Box->R() = TRotation<T, 3>(FQuat::Identity);
+		Box->V() = TVector<T, 3>(0, 0, -1);
+		Box->P() = Box->X();
+		Box->Q() = Box->R();
+		Box->AuxilaryValue(PhysicsMaterials) = MakeSerializable(PhysicsMaterial);
+
+		const float Dt = 1 / 24.;
+
+		TPBDCollisionConstraintAccessor<T, 3> Collisions(Particles, Collided, PhysicsMaterials, 2, 5, T(0));
+
+		Collisions.ComputeConstraints(Dt);
+		EXPECT_EQ(Collisions.NumConstraints(), 1);
+
+		TPBDCollisionConstraints<T, 3>::FMultiPointContactConstraint * Constraint = Collisions.GetConstraint(0).template As<TPBDCollisionConstraints<T, 3>::FMultiPointContactConstraint>();
+		EXPECT_TRUE(Constraint);
+		EXPECT_EQ(Constraint->NumSamples(), 4);
+
+		Collisions.UpdateManifold(*Constraint);
+		EXPECT_EQ(Constraint->NumSamples(), 4);
+
+		EXPECT_EQ(Constraint->Particle[0], Box);
+		EXPECT_EQ(Constraint->Particle[1], Floor);
+		EXPECT_TRUE(Constraint->GetNormal().operator==(TVector<T, 3>(0, 0, 1)));
+		EXPECT_TRUE(FMath::Abs( Constraint->GetLocation().Z - TVector<T,3>(0,0,-1).Z ) < SMALL_THRESHOLD);
+		EXPECT_TRUE(FMath::Abs(Constraint->GetPhi() - T(-1.0)) < SMALL_THRESHOLD);
+
+		{
+			INVARIANT_XR_START(Box);
+			Collisions.Apply(Dt, { Collisions.GetConstraintHandle(0) }, 0, 1);
+			INVARIANT_XR_END(Box);
+		}
+
+		// 0 restitution so expecting 0 velocity
+		//EXPECT_TRUE(Box->V().Equals(TVector<T, 3>(0)));
+		//EXPECT_TRUE(Box->W().Equals(TVector<T, 3>(0)));
+
+		{
+			RESET_PQ(Box);
+			{
+				INVARIANT_XR_START(Box);
+				INVARIANT_VW_START(Box);
+				Collisions.ApplyPushOut(Dt, { Collisions.GetConstraintHandle(0) }, TSet<const TGeometryParticleHandle<T, 3>*>(), 0, 1);
+				INVARIANT_XR_END(Box);
+				INVARIANT_VW_END(Box);
+			}
+		}
+
+		//EXPECT_TRUE(Box->P().Equals(FVector(0.f, 0.f, 50.f)));
+		//EXPECT_TRUE(Box->Q().Equals(FQuat::Identity));
+
+	}
+	template void CollisionConvexConvex<float>();
+
+
+	template<class T>
+	void CollisionBoxPlaneZeroResitution()
+	{
+		// test a box and plane in a colliding state
+		TArrayCollectionArray<bool> Collided;
+		TUniquePtr<FChaosPhysicsMaterial> PhysicsMaterial = MakeUnique<FChaosPhysicsMaterial>();
+		PhysicsMaterial->Friction = (T)0;
+		PhysicsMaterial->Restitution = (T)0;
+		TArrayCollectionArray<TSerializablePtr<FChaosPhysicsMaterial>> PhysicsMaterials;
+		
+		TPBDRigidsSOAs<T, 3> Particles;
+		Particles.GetParticleHandles().AddArray(&Collided);
+		Particles.GetParticleHandles().AddArray(&PhysicsMaterials);
+
+		auto Floor = AppendStaticAnalyticFloor<T>(Particles);
+		auto Box = AppendDynamicParticleBox<T>(Particles);
+		Box->X() = TVector<T, 3>(0, 1, 0);
+		Box->R() = TRotation<T, 3>(FQuat::Identity);
+		Box->V() = TVector<T, 3>(0, 0, -1);
+		Box->P() = Box->X();
+		Box->Q() = Box->R();
+		Box->AuxilaryValue(PhysicsMaterials) = MakeSerializable(PhysicsMaterial);
+
+		const float Dt = 1 / 24.;
+
+		TPBDCollisionConstraintAccessor<T, 3> Collisions(Particles, Collided, PhysicsMaterials, 2, 5, T(0));
+
+		Collisions.ComputeConstraints(Dt);
+		EXPECT_EQ(Collisions.NumConstraints(), 1);
+
+		TPBDCollisionConstraints<T, 3>::FConstraintBase & Constraint = Collisions.GetConstraint(0);
+		if (auto PBDRigid = Constraint.Particle[0]->CastToRigidParticle())
+		{
+			PBDRigid->CollisionParticles()->UpdateAccelerationStructures();
+		}
+		Collisions.UpdateLevelsetConstraint(*Constraint.template As<TPBDCollisionConstraints<T,3>::FPointContactConstraint>());
+
+		EXPECT_EQ(Constraint.Particle[0], Box);
+		EXPECT_EQ(Constraint.Particle[1], Floor);
+		EXPECT_TRUE(Constraint.GetNormal().operator==(TVector<T, 3>(0, 0, 1)));
+		EXPECT_TRUE(FMath::Abs(ChaosTest::SignedDistance(*Constraint.Particle[0], Constraint.GetLocation())) < SMALL_THRESHOLD);
+		EXPECT_TRUE(FMath::Abs(Constraint.GetPhi() - T(-0.5)) < SMALL_THRESHOLD);
+
+		{
+			INVARIANT_XR_START(Box);
+			Collisions.Apply(Dt, { Collisions.GetConstraintHandle(0) }, 0, 1);
+			INVARIANT_XR_END(Box);
+		}
+
+		// 0 restitution so expecting 0 velocity
+		EXPECT_TRUE(Box->V().Equals(TVector<T, 3>(0)));
+		EXPECT_TRUE(Box->W().Equals(TVector<T, 3>(0)));
+
+		{
+			RESET_PQ(Box);
+			{
+				INVARIANT_XR_START(Box);
+				INVARIANT_VW_START(Box);
+				Collisions.ApplyPushOut(Dt, { Collisions.GetConstraintHandle(0) }, TSet<const TGeometryParticleHandle<T,3>*>(), 0, 1);
+				INVARIANT_XR_END(Box);
+				INVARIANT_VW_END(Box);
+			}
+		}
+
+		EXPECT_TRUE(Box->P().Equals(FVector(0.f, 1.f, 0.5f)));
+		EXPECT_TRUE(Box->Q().Equals(FQuat::Identity));
+	}
+	template void CollisionBoxPlaneZeroResitution<float>();
+
+
+	template<class T>
+	void CollisionBoxPlaneRestitution()
+	{
+		TArrayCollectionArray<bool> Collided;
+		TUniquePtr<FChaosPhysicsMaterial> PhysicsMaterial = MakeUnique<FChaosPhysicsMaterial>();
+		PhysicsMaterial->Friction = (T)0;
+		PhysicsMaterial->Restitution = (T)1;
+		TArrayCollectionArray<TSerializablePtr<FChaosPhysicsMaterial>> PhysicsMaterials;
+		TPBDRigidsSOAs<T, 3> Particles;
+		Particles.GetParticleHandles().AddArray(&Collided);
+		Particles.GetParticleHandles().AddArray(&PhysicsMaterials);
+
+		auto Floor = AppendStaticAnalyticFloor<T>(Particles);
+		auto Box = AppendDynamicParticleBox<T>(Particles);
+		Box->X() = TVector<T, 3>(0, 0, 0);
+		Box->R() = TRotation<T, 3>(FQuat::Identity);
+		Box->V() = TVector<T, 3>(0, 0, -100);
+		Box->P() = Box->X();
+		Box->Q() = Box->R();
+		Box->AuxilaryValue(PhysicsMaterials) = MakeSerializable(PhysicsMaterial);
+
+		const float Dt = 1 / 24.;
+
+		TPBDCollisionConstraintAccessor<T, 3> Collisions(Particles, Collided, PhysicsMaterials, 2, 5, T(0));
+
+		Collisions.ComputeConstraints(Dt);
+		EXPECT_EQ(Collisions.NumConstraints(), 1);
+
+		TPBDCollisionConstraints<T, 3>::FConstraintBase & Constraint = Collisions.GetConstraint(0);
+		if (auto PBDRigid = Constraint.Particle[0]->CastToRigidParticle())
+		{
+			PBDRigid->CollisionParticles()->UpdateAccelerationStructures();
+		}
+		Collisions.UpdateLevelsetConstraint(*Constraint.template As<TPBDCollisionConstraints<T,3>::FPointContactConstraint>());
+		EXPECT_EQ(Constraint.Particle[0], Box);
+		EXPECT_EQ(Constraint.Particle[1], Floor);
+		EXPECT_TRUE(Constraint.GetNormal().operator==(TVector<T, 3>(0, 0, 1)));
+		EXPECT_TRUE(FMath::Abs(ChaosTest::SignedDistance(*Constraint.Particle[0], Constraint.GetLocation())) < SMALL_THRESHOLD);
+		EXPECT_TRUE(FMath::Abs(Constraint.GetPhi() - T(-0.5)) < SMALL_THRESHOLD);
+
+		{
+			INVARIANT_XR_START(Box);
+			Collisions.Apply(Dt, { Collisions.GetConstraintHandle(0) }, 0, 1);
+			INVARIANT_XR_END(Box);
+		}
+
+		// full restitution, so expecting negative velocity
+		EXPECT_TRUE(Box->V().Equals(TVector<T, 3>(0.f, 0.f, 100.f)));
+		EXPECT_TRUE(Box->W().Equals(TVector<T, 3>(0)));
+		// collision occurs before full dt takes place, so need some bounce back for the remaining time we have
+		//EXPECT_TRUE(Particles.P(BoxId).Equals(Particles.X(BoxId)));
+		//EXPECT_TRUE(Particles.Q(BoxId).Equals(Particles.R(BoxId)));
+
+		{
+			RESET_PQ(Box);
+			{
+				INVARIANT_XR_START(Box);
+				INVARIANT_VW_START(Box);
+				Collisions.ApplyPushOut(Dt, { Collisions.GetConstraintHandle(0) }, TSet<const TGeometryParticleHandle<T,3>*>(), 0, 1);
+				INVARIANT_XR_END(Box);
+				INVARIANT_VW_END(Box);
+			}
+		}
+
+		//for push out velocity is unimportant, so expecting simple pop out
+		EXPECT_TRUE(Box->P().Equals(FVector(0.f, 0.f, 0.5f)));
+		EXPECT_TRUE(Box->Q().Equals(FQuat::Identity));
+	}
+	template void CollisionBoxPlaneRestitution<float>();
+
+
+	template<class T>
+	void CollisionBoxToStaticBox()
+	{
+		TArrayCollectionArray<bool> Collided;
+		TUniquePtr<FChaosPhysicsMaterial> PhysicsMaterial = MakeUnique<FChaosPhysicsMaterial>();
+		PhysicsMaterial->Friction = (T)0;
+		PhysicsMaterial->Restitution = (T)0;
+		TArrayCollectionArray<TSerializablePtr<FChaosPhysicsMaterial>> PhysicsMaterials;
+		TPBDRigidsSOAs<T, 3> Particles;
+		Particles.GetParticleHandles().AddArray(&Collided);
+		Particles.GetParticleHandles().AddArray(&PhysicsMaterials);
+
+		auto StaticBox = AppendStaticParticleBox<T>(Particles);
+		StaticBox->X() = TVector<T, 3>(-0.05f, -0.05f, -0.1f);
+		StaticBox->AuxilaryValue(PhysicsMaterials) = MakeSerializable(PhysicsMaterial);
+
+		auto Box2 = AppendDynamicParticleBox<T>(Particles);
+		TVector<T, 3> StartingPoint(0.5f);
+		Box2->X() = StartingPoint;
+		Box2->P() = Box2->X();
+		Box2->Q() = Box2->R();
+		Box2->V() = TVector<T, 3>(0, 0, -1);
+		Box2->AuxilaryValue(PhysicsMaterials) = MakeSerializable(PhysicsMaterial);
+
+		FBox Region(FVector(.2), FVector(.5));
+
+		float Dt = 1 / 24.;
+
+		TPBDCollisionConstraintAccessor<T, 3> Collisions(Particles, Collided, PhysicsMaterials, 1, 1, T(0));
+		Collisions.ComputeConstraints(Dt);
+		EXPECT_EQ(Collisions.NumConstraints(), 1);
+
+		TPBDCollisionConstraints<T, 3>::FConstraintBase & Constraint = Collisions.GetConstraint(0);
+
+		//Collisions.UpdateLevelsetConstraintGJK(Particles, Constraint);
+		//EXPECT_EQ(Constraint.ParticleIndex, 1);
+		//EXPECT_EQ(Constraint.LevelsetIndex, 0);
+		//EXPECT_TRUE(Constraint.GetNormal().Equals(FVector(0.0, 1.0, 0.0f))); // GJK returns a different result!
+		//EXPECT_TRUE(FMath::Abs(ChaosTest::SignedDistance(Particles, Constraint.ParticleIndex, Constraint.GetLocation())) < SMALL_THRESHOLD);
+		//EXPECT_TRUE(FMath::Abs(Constraint.GetPhi() - T(-0.233)) < SMALL_THRESHOLD);
+
+		if (auto PBDRigid = Constraint.Particle[0]->CastToRigidParticle())
+		{
+			PBDRigid->CollisionParticles()->UpdateAccelerationStructures();
+		}
+
+		EXPECT_EQ(Constraint.Particle[0], Box2);
+		EXPECT_EQ(Constraint.Particle[1], StaticBox);
+		EXPECT_TRUE(Constraint.GetNormal().Equals(FVector(0.0, 0.0, 1.0f)));
+		EXPECT_TRUE(FMath::Abs(ChaosTest::SignedDistance(*Constraint.Particle[0], Constraint.GetLocation())) < SMALL_THRESHOLD);
+		EXPECT_TRUE(FMath::Abs(Constraint.GetPhi() - T(-0.4)) < SMALL_THRESHOLD);
+
+
+		EXPECT_TRUE(FMath::Abs(Box2->V().Size() - 1.f)<SMALL_THRESHOLD ); // no velocity change yet  
+
+		{
+			INVARIANT_XR_START(Box2);
+			INVARIANT_XR_START(StaticBox);
+			Collisions.Apply(Dt, { Collisions.GetConstraintHandle(0) }, 0, 1);
+			INVARIANT_XR_END(Box2);
+			INVARIANT_XR_END(StaticBox);
+		}
+
+		EXPECT_TRUE(Box2->V().Size()<FVector(0,-1,0).Size()); // slowed down  
+		EXPECT_TRUE(Box2->W().Size()>0); // now has rotation 
+
+		RESET_PQ(Box2);
+		{
+			INVARIANT_XR_START(Box2);
+			INVARIANT_XR_START(StaticBox);
+			INVARIANT_VW_START(Box2);
+			Collisions.ApplyPushOut(Dt, { Collisions.GetConstraintHandle(0) }, TSet<const TGeometryParticleHandle<T, 3>*>(), 0, 1);
+			INVARIANT_XR_END(Box2);
+			INVARIANT_XR_END(StaticBox);
+			INVARIANT_VW_END(Box2);
+		}
+
+		EXPECT_FALSE(Box2->P().Equals(StartingPoint)); // moved
+		EXPECT_FALSE(Box2->Q().Equals(FQuat::Identity)); // and rotated
+	}
+	template void CollisionBoxToStaticBox<float>();
+
+	template<class T>
+	void CollisionPGS()
+	{
+#if CHAOS_PARTICLEHANDLE_TODO
+		TSet<int32> ActiveIndices;
+		TArray<TSet<int32>> IslandParticles;
+		TArray<int32> IslandSleepCounts;
+		TArrayCollectionArray<bool> Collided;
+		TUniquePtr<FChaosPhysicsMaterial> PhysicsMaterial = MakeUnique<FChaosPhysicsMaterial>();
+		PhysicsMaterial->Friction = (T)0;
+		PhysicsMaterial->Restitution = (T)0;
+		TArrayCollectionArray<TSerializablePtr<FChaosPhysicsMaterial>> PhysicsMaterials;
+		Chaos::TPBDRigidParticles<T, 3> Particles;
+		Particles.AddArray(&Collided);
+		Particles.AddArray(&PhysicsMaterials);
+
+		int32 BoxId1 = AppendParticleBox<T>(Particles);
+		Particles.X(BoxId1) = TVector<T, 3>(0.f, 0.f, 0.5f);
+		Particles.R(BoxId1) = TRotation<T, 3>(FQuat::Identity);
+		Particles.V(BoxId1) = TVector<T, 3>(0.f, 0.f, -10.f);
+		PhysicsMaterials[BoxId1]  = MakeSerializable(PhysicsMaterial);
+
+		int32 BoxId2 = AppendParticleBox<T>(Particles);
+		Particles.X(BoxId2) = TVector<T, 3>(0.f, 0.f, 0.5f);
+		Particles.R(BoxId2) = TRotation<T, 3>(FQuat::Identity);
+		Particles.V(BoxId2) = TVector<T, 3>(0.f, 0.f, -10.f);
+		PhysicsMaterials[BoxId2]  = MakeSerializable(PhysicsMaterial);
+
+		int32 FloorId = AppendStaticAnalyticFloor<T>(Particles);
+
+		ActiveIndices.Add(BoxId1);
+		ActiveIndices.Add(BoxId2);
+		ActiveIndices.Add(FloorId);
+
+		TArray<int32> Indices = ActiveIndices.Array();
+
+		TPBDCollisionConstraintPGS<T, 3> CollisionConstraints(Particles, Indices, Collided, PhysicsMaterials);
+		TRigidBodyContactConstraintPGS<T, 3> Constraint1;
+		Constraint1.ParticleIndex = BoxId1;
+		Constraint1.LevelsetIndex = BoxId2;
+		TRigidBodyContactConstraintPGS<T, 3> Constraint2;
+		Constraint2.ParticleIndex = BoxId1;
+		Constraint2.LevelsetIndex = FloorId;
+		TRigidBodyContactConstraintPGS<T, 3> Constraint3;
+		Constraint3.ParticleIndex = BoxId2;
+		Constraint3.LevelsetIndex = FloorId;
+		CollisionConstraints.Constraints.Add(Constraint1);
+		CollisionConstraints.Constraints.Add(Constraint2);
+		CollisionConstraints.Constraints.Add(Constraint3);
+		CollisionConstraints.Apply(Particles, 1.f, { 0, 1, 2 });
+		EXPECT_LT(FMath::Abs(Particles.V(BoxId1)[2]), 1e-3);
+		EXPECT_LT(FMath::Abs(Particles.V(BoxId2)[2]), 1e-3);
+#endif
+	}
+	template void CollisionPGS<float>();
+
+	template<class T>
+	void CollisionPGS2()
+	{
+#if CHAOS_PARTICLEHANDLE_TODO
+		TSet<int32> ActiveIndices;
+		TArray<TSet<int32>> IslandParticles;
+		TArray<int32> IslandSleepCounts;
+		TArrayCollectionArray<bool> Collided;
+		TUniquePtr<FChaosPhysicsMaterial> PhysicsMaterial = MakeUnique<FChaosPhysicsMaterial>();
+		PhysicsMaterial->Friction = (T)0;
+		PhysicsMaterial->Restitution = (T)0;
+		TArrayCollectionArray<TSerializablePtr<FChaosPhysicsMaterial>> PhysicsMaterials;
+		Chaos::TPBDRigidParticles<T, 3> Particles;
+		Particles.AddArray(&Collided);
+		Particles.AddArray(&PhysicsMaterials);
+
+		int32 BoxId1 = AppendParticleBox<T>(Particles);
+		Particles.X(BoxId1) = TVector<T, 3>(0.f, 0.f, 0.5f);
+		Particles.R(BoxId1) = TRotation<T, 3>(FQuat::Identity);
+		Particles.V(BoxId1) = TVector<T, 3>(0.f, 0.f, -10.f);
+		PhysicsMaterials[BoxId1] = MakeSerializable(PhysicsMaterial);
+
+		int32 BoxId2 = AppendParticleBox<T>(Particles);
+		Particles.X(BoxId2) = TVector<T, 3>(0.f, 0.f, 0.5f);
+		Particles.R(BoxId2) = TRotation<T, 3>(FQuat::Identity);
+		Particles.V(BoxId2) = TVector<T, 3>(0.f, 0.f, -10.f);
+		PhysicsMaterials[BoxId2] = MakeSerializable(PhysicsMaterial);
+
+		int32 FloorId = AppendStaticAnalyticFloor<T>(Particles);
+
+		ActiveIndices.Add(BoxId1);
+		ActiveIndices.Add(BoxId2);
+		ActiveIndices.Add(FloorId);
+
+		TArray<int32> Indices = ActiveIndices.Array();
+
+		TPBDCollisionConstraintPGS<T, 3> CollisionConstraints(Particles, Indices, Collided, PhysicsMaterials);
+		TRigidBodyContactConstraintPGS<T, 3> Constraint1;
+		Constraint1.ParticleIndex = BoxId1;
+		Constraint1.LevelsetIndex = BoxId2;
+		TRigidBodyContactConstraintPGS<T, 3> Constraint2;
+		Constraint2.ParticleIndex = BoxId1;
+		Constraint2.LevelsetIndex = FloorId;
+		TRigidBodyContactConstraintPGS<T, 3> Constraint3;
+		Constraint3.ParticleIndex = BoxId2;
+		Constraint3.LevelsetIndex = FloorId;
+		CollisionConstraints.Constraints.Add(Constraint2);
+		CollisionConstraints.Constraints.Add(Constraint3);
+		CollisionConstraints.Constraints.Add(Constraint1);
+		CollisionConstraints.Apply(Particles, 1.f, { 0, 1, 2 });
+		EXPECT_LT(FMath::Abs(Particles.V(BoxId1)[2] - 0.5), 1e-3);
+		EXPECT_LT(FMath::Abs(Particles.V(BoxId2)[2] + 0.5), 1e-3);
+#endif
+	}
+	template void CollisionPGS2<float>();
+}
+
diff --git a/Engine/Source/Programs/HeadlessChaos/Private/HeadlessChaosTestConstraintHandle.cpp b/Engine/Source/Programs/HeadlessChaos/Private/HeadlessChaosTestConstraintHandle.cpp
new file mode 100644
--- /dev/null
+++ b/Engine/Source/Programs/HeadlessChaos/Private/HeadlessChaosTestConstraintHandle.cpp
@@ -0,0 +1,221 @@
+// Copyright Epic Games, Inc. All Rights Reserved.
+
+#include "HeadlessChaos.h"
+#include "HeadlessChaosTestUtility.h"
+#include "Chaos/PBDConstraintRule.h"
+#include "Chaos/PBDJointConstraints.h"
+#include "Chaos/PBDPositionConstraints.h"
+#include "Chaos/PBDRigidParticles.h"
+#include "Chaos/PBDRigidsEvolution.h"
+#include "Chaos/PBDSpringConstraints.h"
+#include "Chaos/Sphere.h"
+#include "Chaos/Utilities.h"
+#include "Chaos/PBDRigidsEvolutionGBF.h"
+#include "Chaos/PBDPositionConstraints.h"
+#include "Chaos/PBDConstraintRule.h"
+#include "Chaos/PBDJointConstraints.h"
+#include "Chaos/PBDRigidSpringConstraints.h"
+#include "Chaos/PBDRigidDynamicSpringConstraints.h"
+
+namespace ChaosTest {
+
+	using namespace Chaos;
+
+	/**
+	 * Check that we can access and remove constraints using handles
+	 */
+	template<typename TConstraints>
+	void CheckConstraintHandles(TConstraints& Constraints, TArray<TPBDRigidParticleHandle<FReal, 3>*> ParticleHandles, TArray<typename TConstraints::FConstraintContainerHandle*> ConstraintsHandles)
+	{
+		// Constraints are created in valid state
+		EXPECT_EQ(Constraints.NumConstraints(), 4);
+		EXPECT_TRUE(ConstraintsHandles[0]->IsValid());
+		EXPECT_TRUE(ConstraintsHandles[1]->IsValid());
+		EXPECT_TRUE(ConstraintsHandles[2]->IsValid());
+		EXPECT_TRUE(ConstraintsHandles[3]->IsValid());
+
+		// Can access constraints' particles
+		EXPECT_EQ(ConstraintsHandles[0]->GetConstrainedParticles()[0], ParticleHandles[0]);
+		EXPECT_EQ(ConstraintsHandles[1]->GetConstrainedParticles()[0], ParticleHandles[1]);
+		EXPECT_EQ(ConstraintsHandles[2]->GetConstrainedParticles()[0], ParticleHandles[2]);
+		EXPECT_EQ(ConstraintsHandles[3]->GetConstrainedParticles()[0], ParticleHandles[3]);
+
+		// Some constraints are single-particle, so can't do this check for all TConstraint
+		//EXPECT_EQ(ConstraintsHandles[0]->GetConstrainedParticles()[1], ParticleHandles[1]);
+		//EXPECT_EQ(ConstraintsHandles[1]->GetConstrainedParticles()[1], ParticleHandles[2]);
+		//EXPECT_EQ(ConstraintsHandles[2]->GetConstrainedParticles()[1], ParticleHandles[3]);
+		//EXPECT_EQ(ConstraintsHandles[3]->GetConstrainedParticles()[1], ParticleHandles[4]);
+
+		// Array is packed as expected
+		EXPECT_EQ(ConstraintsHandles[0]->GetConstraintIndex(), 0);
+		EXPECT_EQ(ConstraintsHandles[1]->GetConstraintIndex(), 1);
+		EXPECT_EQ(ConstraintsHandles[2]->GetConstraintIndex(), 2);
+		EXPECT_EQ(ConstraintsHandles[3]->GetConstraintIndex(), 3);
+
+		// Can remove constraints (from the middle of the constraint array)
+		ConstraintsHandles[1]->RemoveConstraint();
+		ConstraintsHandles[1] = nullptr;
+		EXPECT_EQ(Constraints.NumConstraints(), 3);
+
+		// Array is still packed as expected
+		EXPECT_EQ(ConstraintsHandles[0]->GetConstraintIndex(), 0);
+		EXPECT_EQ(ConstraintsHandles[2]->GetConstraintIndex(), 2);
+		EXPECT_EQ(ConstraintsHandles[3]->GetConstraintIndex(), 1);
+
+		// Can remove constraints (from the end of the constraint array)
+		ConstraintsHandles[3]->RemoveConstraint();
+		ConstraintsHandles[3] = nullptr;
+		EXPECT_EQ(Constraints.NumConstraints(), 2);
+
+		// Array is still packed as expected
+		EXPECT_EQ(ConstraintsHandles[0]->GetConstraintIndex(), 0);
+		EXPECT_EQ(ConstraintsHandles[2]->GetConstraintIndex(), 1);
+
+		// Can remove constraints (from beginning of the constraint array)
+		ConstraintsHandles[0]->RemoveConstraint();
+		ConstraintsHandles[0] = nullptr;
+		EXPECT_EQ(Constraints.NumConstraints(), 1);
+		EXPECT_EQ(ConstraintsHandles[2]->GetConstraintIndex(), 0);
+
+		// Can remove last constraint
+		ConstraintsHandles[2]->RemoveConstraint();
+		ConstraintsHandles[2] = nullptr;
+		EXPECT_EQ(Constraints.NumConstraints(), 0);
+	}
+
+	template<typename TEvolution>
+	void CollisionConstraintHandles()
+	{
+#if CHAOS_CONSTRAINTHANDLE_TODO
+		// @todo(ccaulfield): Collision Constraints Container can't be used without collision detection loop.
+		TPBDRigidsSOAs<FReal, 3> Particles;
+		TPBDCollisionConstraints<T, 3> Constraints;
+		TEvolution Evolution(Particles);
+
+		TArray<TPBDRigidParticleHandle<T, 3>*> ParticleHandles = Evolution.CreateDynamicParticles(5);
+
+		TArray<TPBDCollisionConstraintHandle<T, 3>*> ConstraintsHandles =
+		{
+			Constraints.AddConstraint({ ParticleHandles[0], ParticleHandles[1] }, ...),
+			Constraints.AddConstraint({ ParticleHandles[1], ParticleHandles[2] }, ...),
+			Constraints.AddConstraint({ ParticleHandles[2], ParticleHandles[3] }, ...),
+			Constraints.AddConstraint({ ParticleHandles[3], ParticleHandles[4] }, ...),
+		};
+
+		//CheckConstraintHandles(Constraints, ParticleHandles, ConstraintsHandles);
+#endif
+	}
+
+	template<typename TEvolution>
+	void JointConstraintHandles()
+	{
+		TPBDRigidsSOAs<FReal, 3> Particles;
+		TEvolution Evolution(Particles);
+
+		TArray<TPBDRigidParticleHandle<FReal, 3>*> ParticleHandles = Evolution.CreateDynamicParticles(5);
+
+		FPBDJointConstraints Constraints;
+		TArray<FPBDJointConstraintHandle*> ConstraintsHandles =
+		{
+			Constraints.AddConstraint({ ParticleHandles[0], ParticleHandles[1] }, { TRigidTransform<FReal, 3>(), TRigidTransform<FReal, 3>() }),
+			Constraints.AddConstraint({ ParticleHandles[1], ParticleHandles[2] }, { TRigidTransform<FReal, 3>(), TRigidTransform<FReal, 3>() }),
+			Constraints.AddConstraint({ ParticleHandles[2], ParticleHandles[3] }, { TRigidTransform<FReal, 3>(), TRigidTransform<FReal, 3>() }),
+			Constraints.AddConstraint({ ParticleHandles[3], ParticleHandles[4] }, { TRigidTransform<FReal, 3>(), TRigidTransform<FReal, 3>() }),
+		};
+
+		CheckConstraintHandles(Constraints, ParticleHandles, ConstraintsHandles);
+	}
+
+	template<typename TEvolution>
+	void PositionConstraintHandles()
+	{
+		TPBDRigidsSOAs<FReal, 3> Particles;
+		TEvolution Evolution(Particles);
+		TArray<TPBDRigidParticleHandle<FReal, 3>*> ParticleHandles = Evolution.CreateDynamicParticles(5);
+
+		TPBDPositionConstraints<FReal, 3> Constraints;
+		TArray<TPBDPositionConstraintHandle<FReal, 3>*> ConstraintsHandles =
+		{
+			Constraints.AddConstraint(ParticleHandles[0], { 0, 0, 0 }),
+			Constraints.AddConstraint(ParticleHandles[1], { 0, 0, 0 }),
+			Constraints.AddConstraint(ParticleHandles[2], { 0, 0, 0 }),
+			Constraints.AddConstraint(ParticleHandles[3], { 0, 0, 0 }),
+		};
+
+		CheckConstraintHandles(Constraints, ParticleHandles, ConstraintsHandles);
+	}
+
+	template<typename TEvolution>
+	void RigidSpringConstraintHandles()
+	{
+		TPBDRigidsSOAs<FReal, 3> Particles;
+		TEvolution Evolution(Particles);
+		TArray<TPBDRigidParticleHandle<FReal, 3>*> ParticleHandles = Evolution.CreateDynamicParticles(5);
+
+		TPBDRigidSpringConstraints<FReal, 3> Constraints;
+		TArray<TPBDRigidSpringConstraintHandle<FReal, 3>*> ConstraintsHandles =
+		{
+			Constraints.AddConstraint({ ParticleHandles[0], ParticleHandles[1] }, { { 0, 0, 0 }, { 0, 0, 0 } }),
+			Constraints.AddConstraint({ ParticleHandles[1], ParticleHandles[2] }, { { 0, 0, 0 }, { 0, 0, 0 } }),
+			Constraints.AddConstraint({ ParticleHandles[2], ParticleHandles[3] }, { { 0, 0, 0 }, { 0, 0, 0 } }),
+			Constraints.AddConstraint({ ParticleHandles[3], ParticleHandles[4] }, { { 0, 0, 0 }, { 0, 0, 0 } }),
+		};
+
+		CheckConstraintHandles(Constraints, ParticleHandles, ConstraintsHandles);
+	}
+
+
+	template<typename TEvolution>
+	void RigidDynamicSpringConstraintHandles()
+	{
+		TPBDRigidsSOAs<FReal, 3> Particles;
+		TEvolution Evolution(Particles);
+		TArray<TPBDRigidParticleHandle<FReal, 3>*> ParticleHandles = Evolution.CreateDynamicParticles(5);
+
+		TPBDRigidDynamicSpringConstraints<FReal, 3> Constraints;
+		TArray<TPBDRigidDynamicSpringConstraintHandle<FReal, 3>*> ConstraintsHandles =
+		{
+			Constraints.AddConstraint({ ParticleHandles[0], ParticleHandles[1] }),
+			Constraints.AddConstraint({ ParticleHandles[1], ParticleHandles[2] }),
+			Constraints.AddConstraint({ ParticleHandles[2], ParticleHandles[3] }),
+			Constraints.AddConstraint({ ParticleHandles[3], ParticleHandles[4] }),
+		};
+
+		CheckConstraintHandles(Constraints, ParticleHandles, ConstraintsHandles);
+	}
+
+	TEST(ConstraintHandleTests, DISABLED_CollisionConstraintHandle)
+	{
+		CollisionConstraintHandles<Chaos::TPBDRigidsEvolutionGBF<FReal, 3>>();
+
+		SUCCEED();
+	}
+
+	TEST(ConstraintHandleTests, JointConstraintHandle)
+	{
+		JointConstraintHandles<Chaos::TPBDRigidsEvolutionGBF<FReal, 3>>();
+
+		SUCCEED();
+	}
+
+	TEST(ConstraintHandleTests, PositionConstraintHandles)
+	{
+		PositionConstraintHandles<Chaos::TPBDRigidsEvolutionGBF<FReal, 3>>();
+
+		SUCCEED();
+	}
+
+	TEST(ConstraintHandleTests, RigidSpringConstraintHandles)
+	{
+		RigidSpringConstraintHandles<Chaos::TPBDRigidsEvolutionGBF<FReal, 3>>();
+
+		SUCCEED();
+	}
+
+	TEST(ConstraintHandleTests, RigidDynamicSpringConstraintHandles)
+	{
+		RigidDynamicSpringConstraintHandles<Chaos::TPBDRigidsEvolutionGBF<FReal, 3>>();
+
+		SUCCEED();
+	}
+}
\ No newline at end of file
diff --git a/Engine/Source/Programs/HeadlessChaos/Private/HeadlessChaosTestConstraints.cpp b/Engine/Source/Programs/HeadlessChaos/Private/HeadlessChaosTestConstraints.cpp
new file mode 100644
--- /dev/null
+++ b/Engine/Source/Programs/HeadlessChaos/Private/HeadlessChaosTestConstraints.cpp
@@ -0,0 +1,129 @@
+// Copyright Epic Games, Inc. All Rights Reserved.
+
+#include "HeadlessChaosTestConstraints.h"
+
+#include "HeadlessChaos.h"
+#include "HeadlessChaosTestUtility.h"
+#include "Modules/ModuleManager.h"
+#include "Chaos/PBDConstraintRule.h"
+#include "Chaos/PBDJointConstraints.h"
+#include "Chaos/PBDPositionConstraints.h"
+#include "Chaos/PBDRigidParticles.h"
+#include "Chaos/PBDRigidsEvolution.h"
+#include "Chaos/PBDSpringConstraints.h"
+#include "Chaos/Sphere.h"
+#include "Chaos/Utilities.h"
+#include "Chaos/PBDRigidsEvolutionGBF.h"
+#include "Chaos/PBDPositionConstraints.h"
+#include "Chaos/PBDConstraintRule.h"
+#include "Chaos/PBDJointConstraints.h"
+
+namespace ChaosTest {
+
+	using namespace Chaos;
+
+	/**
+	 * Position constraint test
+	 */
+	template<typename TEvolution>
+	void Position()
+	{
+		{
+			TPBDRigidsSOAs<FReal, 3> Particles;
+			TEvolution Evolution(Particles);
+			TArray<TPBDRigidParticleHandle<FReal, 3>*> Dynamics = Evolution.CreateDynamicParticles(1);
+			TArray<FVec3> Positions = { FVec3(0) };
+			TPBDPositionConstraints<FReal, 3> PositionConstraints(MoveTemp(Positions), MoveTemp(Dynamics), 1.f);
+			auto ConstraintRule = TPBDConstraintIslandRule<TPBDPositionConstraints<FReal, 3>>(PositionConstraints);
+			
+			Evolution.AddConstraintRule(&ConstraintRule);
+			Evolution.AdvanceOneTimeStep(0.1);
+			Evolution.EndFrame(0.1);
+			EXPECT_LT(Evolution.GetParticleHandles().Handle(0)->X().SizeSquared(), SMALL_NUMBER);
+		}
+		{
+			TPBDRigidsSOAs<FReal, 3> Particles;
+			TPBDRigidsEvolutionGBF<FReal, 3> Evolution(Particles);
+			TArray<TPBDRigidParticleHandle<FReal, 3>*> Dynamics = Evolution.CreateDynamicParticles(1);
+			Evolution.GetGravityForces().SetEnabled(*Dynamics[0], false);
+
+			TArray<FVec3> Positions = { FVec3(1) };
+			TPBDPositionConstraints<FReal, 3> PositionConstraints(MoveTemp(Positions), MoveTemp(Dynamics), 0.5f);
+			auto ConstraintRule = TPBDConstraintIslandRule<TPBDPositionConstraints<FReal, 3>>(PositionConstraints);
+			Evolution.AddConstraintRule(&ConstraintRule);
+
+			Evolution.AdvanceOneTimeStep(0.1);
+			Evolution.EndFrame(0.1);
+			auto& Handle = Evolution.GetParticleHandles().Handle(0);
+			EXPECT_LT(FMath::Abs(Handle->X()[0] - 0.5), (FReal)SMALL_NUMBER);
+			EXPECT_LT(FMath::Abs(Handle->X()[1] - 0.5), (FReal)SMALL_NUMBER);
+			EXPECT_LT(FMath::Abs(Handle->X()[2] - 0.5), (FReal)SMALL_NUMBER);
+
+			Evolution.AdvanceOneTimeStep(0.1);
+			Evolution.EndFrame(0.1);
+
+			EXPECT_LT(FMath::Abs(Handle->X()[0] - 1), (FReal)SMALL_NUMBER);
+			EXPECT_LT(FMath::Abs(Handle->X()[1] - 1), (FReal)SMALL_NUMBER);
+			EXPECT_LT(FMath::Abs(Handle->X()[2] - 1), (FReal)SMALL_NUMBER);
+		}
+	}
+
+
+	/**
+	 * Joint constraints test with the fixed body held in place with a position constraint.
+	 * Joint body swings under the fixed body at fixed distance.
+	 */
+	template<typename TEvolution>
+	void PositionAndJoint()
+	{
+		const int32 Iterations = 10;
+		TPBDRigidsSOAs<FReal, 3> Particles;
+		TEvolution Evolution(Particles, Iterations);
+		TArray<TPBDRigidParticleHandle<FReal, 3>*> Dynamics = Evolution.CreateDynamicParticles(2);
+		TArray<FVec3> PositionConstraintPositions = { FVec3(0, 0, 0) };
+
+		Dynamics[1]->X() = FVec3(500, 0, 0);
+		FVec3 JointConstraintPosition = FVec3(0, 0, 0);
+
+		TArray<TPBDRigidParticleHandle<FReal, 3>*> PositionParticles = { Dynamics[0] };
+		TPBDPositionConstraints<FReal, 3> PositionConstraints(MoveTemp(PositionConstraintPositions), MoveTemp(PositionParticles), 1.f);
+		auto PositionConstraintRule = TPBDConstraintIslandRule<TPBDPositionConstraints<FReal, 3>>(PositionConstraints);
+		Evolution.AddConstraintRule(&PositionConstraintRule);
+
+		TVector<TGeometryParticleHandle<FReal, 3>*, 2> JointParticles = { Dynamics[0], Dynamics[1] };
+		FPBDJointConstraints JointConstraints;
+		JointConstraints.AddConstraint(JointParticles, FRigidTransform3(JointConstraintPosition, FRotation3::FromIdentity()));
+		auto JointConstraintRule = TPBDConstraintIslandRule<FPBDJointConstraints>(JointConstraints);
+		Evolution.AddConstraintRule(&JointConstraintRule);
+
+		FReal Dt = 0.1f;
+		for (int32 TimeIndex = 0; TimeIndex < 100; ++TimeIndex)
+		{
+			Evolution.AdvanceOneTimeStep(Dt);
+			Evolution.EndFrame(Dt);
+
+			const auto& Pos0 = Dynamics[0]->X();
+			const auto& Pos1 = Dynamics[1]->X();
+			const float Delta0 = (Pos0 - FVec3(0, 0, 0)).Size();
+			const float Separation = (Pos1 - Pos0).Size();
+
+			EXPECT_LT(Delta0, 5);
+			EXPECT_GT(Separation, 495);
+			EXPECT_LT(Separation, 505);
+		}
+	}
+
+
+	TEST(ConstraintTests, Constraints)
+	{
+		ChaosTest::Position<Chaos::TPBDRigidsEvolutionGBF<float, 3>>();
+		ChaosTest::PositionAndJoint<Chaos::TPBDRigidsEvolutionGBF<float, 3>>();
+
+#if CHAOS_PARTICLEHANDLE_TODO
+		ChaosTest::Position<Chaos::TPBDRigidsEvolutionPGS<float, 3>, float>();
+		ChaosTest::PositionAndJoint<Chaos::TPBDRigidsEvolutionPGS<float, 3>>();
+#endif
+		SUCCEED();
+	}
+
+}
\ No newline at end of file
