Rename/move HeadlessChaos.

diff --git a/Engine/Source/Programs/HeadlessChaos/Private/HeadlessChaosTestEPA.cpp b/Engine/Source/Programs/HeadlessChaos/Private/HeadlessChaosTestEPA.cpp
new file mode 100644
--- /dev/null
+++ b/Engine/Source/Programs/HeadlessChaos/Private/HeadlessChaosTestEPA.cpp
@@ -0,0 +1,621 @@
+// Copyright Epic Games, Inc. All Rights Reserved.
+
+#include "HeadlessChaosTestEPA.h"
+
+#include "HeadlessChaos.h"
+#include "HeadlessChaosTestUtility.h"
+#include "Chaos/EPA.h"
+#include "Chaos/Sphere.h"
+#include "Chaos/Capsule.h"
+#include "Chaos/GJK.h"
+#include "Chaos/Convex.h"
+#include "Chaos/ImplicitObjectScaled.h"
+
+namespace ChaosTest
+{
+	using namespace Chaos;
+
+	template <typename T>
+	void ValidFace(const TVec3<T>* Verts, const TArray<TEPAEntry<T>>& TetFaces, int32 Idx)
+	{
+		const TEPAEntry<T>& Entry = TetFaces[Idx];
+
+		//does not contain vertex associated with face
+		EXPECT_NE(Entry.IdxBuffer[0], Idx);
+		EXPECT_NE(Entry.IdxBuffer[1], Idx);
+		EXPECT_NE(Entry.IdxBuffer[2], Idx);
+
+		//doesn't have itself as adjacent face
+		EXPECT_NE(Entry.AdjFaces[0], Idx);
+		EXPECT_NE(Entry.AdjFaces[1], Idx);
+		EXPECT_NE(Entry.AdjFaces[2], Idx);
+
+		//adjacent edges and faces are valid for both sides of the edge
+		EXPECT_EQ(TetFaces[Entry.AdjFaces[0]].AdjFaces[Entry.AdjEdges[0]], Idx);
+		EXPECT_EQ(TetFaces[Entry.AdjFaces[1]].AdjFaces[Entry.AdjEdges[1]], Idx);
+		EXPECT_EQ(TetFaces[Entry.AdjFaces[2]].AdjFaces[Entry.AdjEdges[2]], Idx);
+
+		//make sure that adjacent faces share vertices
+		//src dest on the edge of face 0 matches to dest src of face 1
+		for (int EdgeIdx = 0; EdgeIdx < 3; ++EdgeIdx)
+		{
+			const int32 FromFace0 = Entry.IdxBuffer[EdgeIdx];
+			const int32 ToFace0 = Entry.IdxBuffer[(EdgeIdx+1)%3];
+			const int32 Face1EdgeIdx = Entry.AdjEdges[EdgeIdx];
+			const TEPAEntry<T>& Face1 = TetFaces[Entry.AdjFaces[EdgeIdx]];
+			const int32 FromFace1 = Face1.IdxBuffer[Face1EdgeIdx];
+			const int32 ToFace1 = Face1.IdxBuffer[(Face1EdgeIdx+1)%3];
+			EXPECT_EQ(FromFace0, ToFace1);
+			EXPECT_EQ(FromFace1, ToFace0);
+			
+		}
+		switch (Entry.AdjEdges[0])
+		{
+		case 0:
+		{
+			EXPECT_EQ(Entry.IdxBuffer[0], TetFaces[Entry.AdjFaces[0]].IdxBuffer[1]);
+			EXPECT_EQ(Entry.IdxBuffer[1], TetFaces[Entry.AdjFaces[0]].IdxBuffer[0]);
+			break;
+		}
+		case 1:
+		{
+			EXPECT_EQ(Entry.IdxBuffer[0], TetFaces[Entry.AdjFaces[0]].IdxBuffer[2]);
+			EXPECT_EQ(Entry.IdxBuffer[1], TetFaces[Entry.AdjFaces[0]].IdxBuffer[1]);
+			break;
+		}
+		case 2:
+		{
+			EXPECT_EQ(Entry.IdxBuffer[0], TetFaces[Entry.AdjFaces[0]].IdxBuffer[0]);
+			EXPECT_EQ(Entry.IdxBuffer[1], TetFaces[Entry.AdjFaces[0]].IdxBuffer[2]);
+			break;
+		}
+		default: break;
+		}
+
+		EXPECT_LT(TVec3<T>::DotProduct(Verts[Idx], Entry.PlaneNormal), 0);	//normal faces out
+
+		EXPECT_GE(Entry.Distance, 0);	//positive distance since origin is inside tet
+
+		EXPECT_NEAR(Entry.DistanceToPlane(Verts[Entry.IdxBuffer[0]]), 0, 1e-6);
+		EXPECT_NEAR(Entry.DistanceToPlane(Verts[Entry.IdxBuffer[1]]), 0, 1e-6);
+		EXPECT_NEAR(Entry.DistanceToPlane(Verts[Entry.IdxBuffer[2]]), 0, 1e-6);
+	}
+
+	template <typename T>
+	TVec3<T> ErrorSupport(const TVec3<T>& V)
+	{
+		check(false);
+		return TVec3<T>(0);
+	}
+
+	template <typename T>
+	void EPAInitTest()
+	{
+
+		//make sure faces are properly oriented
+		{
+			TArray<TVec3<T>> VertsA = { {-1,-1,1}, {-1,-1,-1}, {-1,1,-1}, {1,1,-1} };
+			TArray<TVec3<T>> VertsB = { TVec3<T>(0), TVec3<T>(0), TVec3<T>(0), TVec3<T>(0) };
+			TArray<TEPAEntry<T>> TetFaces = InitializeEPA(VertsA, VertsB, ErrorSupport<T>, ErrorSupport<T>);
+
+			EXPECT_EQ(TetFaces.Num(), 4);
+			for (int i = 0; i < TetFaces.Num(); ++i)
+			{
+				ValidFace(VertsA.GetData(), TetFaces, i);
+			}
+		}
+
+		{
+			TArray<TVec3<T>> VertsA = { {-1,-1,-1}, {-1,-1,1}, {-1,1,-1}, {1,1,-1} };
+			TArray<TVec3<T>> VertsB = { TVec3<T>(0), TVec3<T>(0), TVec3<T>(0), TVec3<T>(0) };
+			TArray<TEPAEntry<T>> TetFaces = InitializeEPA(VertsA, VertsB, ErrorSupport<T>, ErrorSupport<T>);
+
+			EXPECT_EQ(TetFaces.Num(), 4);
+			for (int i = 0; i < TetFaces.Num(); ++i)
+			{
+				ValidFace(VertsA.GetData(), TetFaces, i);
+			}
+		}
+
+		auto EmptySupport = [](const TVec3<T>& V) { return TVec3<T>(0); };
+
+		//triangle
+		{
+			TVec3<T> AllVerts[] = { {0,-1,1 + 1 / (T)3}, {0,-1,-1 + 1 / (T)3}, {0,1,-1 + 1 / (T)3},{-1,0,0}, {0.5,0,0} };
+
+			auto ASupport = [&](const TVec3<T>& V)
+			{
+				TVec3<T> Best = AllVerts[0];
+				for (const TVec3<T>& Vert : AllVerts)
+				{
+					if (TVec3<T>::DotProduct(Vert, V) > TVec3<T>::DotProduct(Best, V))
+					{
+						Best = Vert;
+					}
+				}
+				return Best;
+			};
+
+			auto ASupportNoPositiveX = [&](const TVec3<T>& V)
+			{
+				TVec3<T> Best = AllVerts[0];
+				for (const TVec3<T>& Vert : AllVerts)
+				{
+					if (Vert.X > 0)
+					{
+						continue;
+					}
+					if (TVec3<T>::DotProduct(Vert, V) > TVec3<T>::DotProduct(Best, V))
+					{
+						Best = Vert;
+					}
+				}
+				return Best;
+			};
+
+			auto ASupportNoX = [&](const TVec3<T>& V)
+			{
+				TVec3<T> Best = AllVerts[0];
+				for (const TVec3<T>& Vert : AllVerts)
+				{
+					if (Vert.X > 0 || Vert.X < 0)
+					{
+						continue;
+					}
+					if (TVec3<T>::DotProduct(Vert, V) > TVec3<T>::DotProduct(Best, V))
+					{
+						Best = Vert;
+					}
+				}
+				return Best;
+			};
+
+			//first winding
+			{
+				TArray<TVec3<T>> VertsA = { AllVerts[0], AllVerts[1], AllVerts[2] };
+				TArray<TVec3<T>> VertsB = { TVec3<T>(0), TVec3<T>(0), TVec3<T>(0) };
+
+				TArray<TEPAEntry<T>> TetFaces = InitializeEPA(VertsA, VertsB, ASupport, EmptySupport);
+				EXPECT_VECTOR_NEAR(VertsA[3], AllVerts[3], 1e-4);
+				EXPECT_VECTOR_NEAR(VertsB[3], TVec3<T>(0), 1e-4);
+
+				EXPECT_EQ(TetFaces.Num(), 4);
+				for (int i = 0; i < TetFaces.Num(); ++i)
+				{
+					ValidFace(VertsA.GetData(), TetFaces, i);
+				}
+
+				T Penetration;
+				TVec3<T> Dir, WitnessA, WitnessB;
+
+				//Try EPA. Note that we are IGNORING the positive x vert to ensure a triangle right on the origin boundary works
+				EPA(VertsA, VertsB, ASupportNoPositiveX, EmptySupport, Penetration, Dir, WitnessA, WitnessB);
+				EXPECT_NEAR(Penetration, 0, 1e-4);
+				EXPECT_VECTOR_NEAR(Dir, TVec3<T>(1,0,0), 1e-4);
+				EXPECT_VECTOR_NEAR(WitnessA, TVec3<T>(0), 1e-4);
+				EXPECT_VECTOR_NEAR(WitnessB, TVec3<T>(0), 1e-4);
+			}
+
+			//other winding
+			{
+				TArray<TVec3<T>> VertsA = { AllVerts[1], AllVerts[0], AllVerts[2] };
+				TArray<TVec3<T>> VertsB = { TVec3<T>(0), TVec3<T>(0), TVec3<T>(0) };
+
+				TArray<TEPAEntry<T>> TetFaces = InitializeEPA(VertsA, VertsB, ASupport, EmptySupport);
+				EXPECT_VECTOR_NEAR(VertsA[3], AllVerts[3], 1e-4);
+				EXPECT_VECTOR_NEAR(VertsB[3], TVec3<T>(0), 1e-4);
+
+				EXPECT_EQ(TetFaces.Num(), 4);
+				for (int i = 0; i < TetFaces.Num(); ++i)
+				{
+					ValidFace(VertsA.GetData(), TetFaces, i);
+				}
+
+				T Penetration;
+				TVec3<T> Dir, WitnessA, WitnessB;
+
+				//Try EPA. Note that we are IGNORING the positive x vert to ensure a triangle right on the origin boundary works
+				EPA(VertsA, VertsB, ASupportNoPositiveX, EmptySupport, Penetration, Dir, WitnessA, WitnessB);
+				EXPECT_NEAR(Penetration, 0, 1e-4);
+				EXPECT_VECTOR_NEAR(Dir, TVec3<T>(1, 0, 0), 1e-4);
+				EXPECT_VECTOR_NEAR(WitnessA, TVec3<T>(0), 1e-4);
+				EXPECT_VECTOR_NEAR(WitnessB, TVec3<T>(0), 1e-4);
+			}
+
+			//touching triangle
+			{
+				TArray<TVec3<T>> VertsA = { AllVerts[1], AllVerts[0], AllVerts[2] };
+				TArray<TVec3<T>> VertsB = { TVec3<T>(0), TVec3<T>(0), TVec3<T>(0) };
+
+				TArray<TEPAEntry<T>> TetFaces = InitializeEPA(VertsA, VertsB, ASupportNoX, EmptySupport);
+				EXPECT_EQ(TetFaces.Num(), 0);
+
+
+				//make sure EPA handles this bad case properly
+				VertsA = { AllVerts[1], AllVerts[0], AllVerts[2] };
+				VertsB = { TVec3<T>(0), TVec3<T>(0), TVec3<T>(0) };
+
+				//touching so penetration 0, normal is 0,0,1
+				T Penetration;
+				TVec3<T> Dir, WitnessA, WitnessB;
+				EXPECT_EQ(EPA(VertsA, VertsB, ASupportNoX, EmptySupport, Penetration, Dir, WitnessA, WitnessB), EPAResult::BadInitialSimplex);
+				EXPECT_EQ(Penetration, 0);
+				EXPECT_VECTOR_NEAR(Dir, TVec3<T>(0, 0, 1), 1e-7);
+				EXPECT_VECTOR_NEAR(WitnessA, TVec3<T>(0), 1e-7);
+				EXPECT_VECTOR_NEAR(WitnessB, TVec3<T>(0), 1e-7);
+			}
+		}
+
+		//line
+		{
+			TVec3<T> AllVerts[] = { {0,-1,1 + 1 / (T)3}, {0,-1,-1 + 1 / (T)3}, {0,1,-1 + 1 / (T)3},{-1,0,0}, {0.5,0,0} };
+
+			auto ASupport = [&](const TVec3<T>& V)
+			{
+				TVec3<T> Best = AllVerts[0];
+				for (const TVec3<T>& Vert : AllVerts)
+				{
+					if (TVec3<T>::DotProduct(Vert, V) > TVec3<T>::DotProduct(Best, V))
+					{
+						Best = Vert;
+					}
+				}
+				return Best;
+			};
+
+			//first winding
+			{
+				TArray<TVec3<T>> VertsA = { AllVerts[0], AllVerts[2] };
+				TArray<TVec3<T>> VertsB = { TVec3<T>(0), TVec3<T>(0) };
+
+				TArray<TEPAEntry<T>> TetFaces = InitializeEPA(VertsA, VertsB, ASupport, EmptySupport);
+				EXPECT_VECTOR_NEAR(VertsA[2], AllVerts[1], 1e-4);
+				EXPECT_VECTOR_NEAR(VertsB[2], TVec3<T>(0), 1e-4);
+
+				EXPECT_VECTOR_NEAR(VertsA[3], AllVerts[3], 1e-4);
+				EXPECT_VECTOR_NEAR(VertsB[3], TVec3<T>(0), 1e-4);
+
+				EXPECT_EQ(TetFaces.Num(), 4);
+				for (int i = 0; i < TetFaces.Num(); ++i)
+				{
+					ValidFace(VertsA.GetData(), TetFaces, i);
+				}
+			}
+
+			//other winding
+			{
+				TArray<TVec3<T>> VertsA = { AllVerts[2], AllVerts[0] };
+				TArray<TVec3<T>> VertsB = { TVec3<T>(0), TVec3<T>(0) };
+
+				TArray<TEPAEntry<T>> TetFaces = InitializeEPA(VertsA, VertsB, ASupport, EmptySupport);
+				EXPECT_VECTOR_NEAR(VertsA[2], AllVerts[1], 1e-4);
+				EXPECT_VECTOR_NEAR(VertsB[2], TVec3<T>(0), 1e-4);
+
+				EXPECT_VECTOR_NEAR(VertsA[3], AllVerts[3], 1e-4);
+				EXPECT_VECTOR_NEAR(VertsB[3], TVec3<T>(0), 1e-4);
+
+				EXPECT_EQ(TetFaces.Num(), 4);
+				for (int i = 0; i < TetFaces.Num(); ++i)
+				{
+					ValidFace(VertsA.GetData(), TetFaces, i);
+				}
+			}
+
+			//touching triangle
+			{
+				auto ASupportNoX = [&](const TVec3<T>& V)
+				{
+					TVec3<T> Best = AllVerts[0];
+					for (const TVec3<T>& Vert : AllVerts)
+					{
+						if (Vert.X > 0 || Vert.X < 0)
+						{
+							continue;
+						}
+						if (TVec3<T>::DotProduct(Vert, V) > TVec3<T>::DotProduct(Best, V))
+						{
+							Best = Vert;
+						}
+					}
+					return Best;
+				};
+
+				TArray<TVec3<T>> VertsA = { AllVerts[2], AllVerts[0] };
+				TArray<TVec3<T>> VertsB = { TVec3<T>(0), TVec3<T>(0) };
+
+				TArray<TEPAEntry<T>> TetFaces = InitializeEPA(VertsA, VertsB, ASupportNoX, EmptySupport);
+				EXPECT_EQ(TetFaces.Num(), 0);
+			}
+
+			//touching line
+			{
+
+				auto ASupportNoXOrZ = [&](const TVec3<T>& V)
+				{
+					TVec3<T> Best = AllVerts[0];
+					for (const TVec3<T>& Vert : AllVerts)
+					{
+						if (Vert.X > 0 || Vert.X < 0 || Vert.Z > 0)
+						{
+							continue;
+						}
+						if (TVec3<T>::DotProduct(Vert, V) > TVec3<T>::DotProduct(Best, V))
+						{
+							Best = Vert;
+						}
+					}
+					return Best;
+				};
+
+				TArray<TVec3<T>> VertsA = { AllVerts[2], AllVerts[0] };
+				TArray<TVec3<T>> VertsB = { TVec3<T>(0), TVec3<T>(0) };
+
+				TArray<TEPAEntry<T>> TetFaces = InitializeEPA(VertsA, VertsB, ASupportNoXOrZ, EmptySupport);
+				EXPECT_EQ(TetFaces.Num(), 0);
+			}
+		}
+	}
+
+	template void EPAInitTest<float>();
+	
+	template <typename T>
+	void EPASimpleTest()
+	{
+		auto ZeroSupport = [](const auto& V) { return TVec3<T>(0); };
+
+		{
+
+			//simple box hull. 0.5 depth on x, 1 depth on y, 1 depth on z. Made z non symmetric to avoid v on tet close to 0 for this case
+			TVec3<T> HullVerts[8] = { {-0.5, -1, -1}, {2, -1, -1}, {-0.5, 1, -1}, {2, 1, -1},
+									  {-0.5, -1, 2}, {2, -1, 2}, {-0.5, 1, 2}, {2, 1, 2} };
+
+			auto SupportA = [&HullVerts](const auto& V)
+			{
+				auto MaxDist = TNumericLimits<T>::Lowest();
+				auto BestVert = HullVerts[0];
+				for (const auto& Vert : HullVerts)
+				{
+					const auto Dist = TVec3<T>::DotProduct(V, Vert);
+					if (Dist > MaxDist)
+					{
+						MaxDist = Dist;
+						BestVert = Vert;
+					}
+				}
+				return BestVert;
+			};
+
+			TArray <TVec3<T>> Tetrahedron = { HullVerts[0], HullVerts[2], HullVerts[3], HullVerts[4] };
+			TArray <TVec3<T>> Zeros = { TVec3<T>(0), TVec3<T>(0), TVec3<T>(0), TVec3<T>(0) };
+
+			T Penetration;
+			TVec3<T> Dir, WitnessA, WitnessB;
+			EXPECT_EQ(EPA(Tetrahedron, Zeros, SupportA, ZeroSupport, Penetration, Dir, WitnessA, WitnessB), EPAResult::Ok);
+			EXPECT_NEAR(Penetration, 0.5, 1e-4);
+			EXPECT_NEAR(Dir[0], -1, 1e-4);
+			EXPECT_NEAR(Dir[1], 0, 1e-4);
+			EXPECT_NEAR(Dir[2], 0, 1e-4);
+			EXPECT_NEAR(WitnessA[0], -0.5, 1e-4);
+			EXPECT_NEAR(WitnessA[1], 0, 1e-4);
+			EXPECT_NEAR(WitnessA[2], 0, 1e-4);
+			EXPECT_NEAR(WitnessB[0], 0, 1e-4);
+			EXPECT_NEAR(WitnessB[1], 0, 1e-4);
+			EXPECT_NEAR(WitnessB[2], 0, 1e-4);
+		}
+
+		{
+			//sphere with deep penetration to make sure we have max iterations
+			TSphere<T,3> Sphere(TVec3<T>(0), 10);
+			auto Support = [&Sphere](const auto& V)
+			{
+				return Sphere.Support(V, 0);
+			};
+
+			TArray <TVec3<T>> Tetrahedron = { Support(FVec3(-1,0,0)), Support(FVec3(1,0,0)),
+				Support(FVec3(0,1,0)), Support(FVec3(0,0,1))};
+			TArray <TVec3<T>> Zeros = { TVec3<T>(0), TVec3<T>(0), TVec3<T>(0), TVec3<T>(0) };
+
+			T Penetration;
+			TVec3<T> Dir, WitnessA, WitnessB;
+			EXPECT_EQ(EPA(Tetrahedron, Zeros, Support, ZeroSupport, Penetration, Dir, WitnessA, WitnessB), EPAResult::MaxIterations);
+			EXPECT_GT(Penetration, 9);
+			EXPECT_LE(Penetration, 10);
+			EXPECT_GT(WitnessA.Size(), 9);	//don't know exact point, but should be 9 away from origin
+			EXPECT_LE(WitnessA.Size(), 10);	//point should be interior to sphere
+		}
+
+		{
+			//capsule with origin in middle
+			TCapsule<T> Capsule(TVec3<T>(0, 0, 10), TVec3<T>(0, 0, -10), 3);
+			auto Support = [&Capsule](const auto& V)
+			{
+				return Capsule.Support(V, 0);
+			};
+
+			TArray <TVec3<T>> Tetrahedron = { Support(FVec3(-1,0,0)), Support(FVec3(1,0,0)),
+				Support(FVec3(0,1,0)), Support(FVec3(0,0,1)) };
+			TArray <TVec3<T>> Zeros = { TVec3<T>(0), TVec3<T>(0), TVec3<T>(0), TVec3<T>(0) };
+
+			T Penetration;
+			TVec3<T> Dir, WitnessA, WitnessB;
+			EXPECT_EQ(EPA(Tetrahedron, Zeros, Support, ZeroSupport, Penetration, Dir, WitnessA, WitnessB), EPAResult::Ok);
+			EXPECT_NEAR(Penetration, 3, 1e-1);
+			EXPECT_NEAR(Dir[2], 0, 1e-1);	//don't know direction, but it should be in xy plane
+			EXPECT_NEAR(WitnessA.Size(), 3, 1e-1);	//don't know exact point, but should be 3 away from origin
+		}
+		{
+			//capsule with origin near top
+			TCapsule<T> Capsule(TVec3<T>(0, 0, -2), TVec3<T>(0, 0, -12), 3);
+			auto Support = [&Capsule](const auto& V)
+			{
+				return Capsule.Support(V, 0);
+			};
+
+			TArray <TVec3<T>> Tetrahedron = { Support(FVec3(-1,0,0)), Support(FVec3(1,0,0)),
+				Support(FVec3(0,1,0)), Support(FVec3(0,0,1)) };
+
+			TArray <TVec3<T>> Zeros = { TVec3<T>(0), TVec3<T>(0), TVec3<T>(0), TVec3<T>(0) };
+
+			T Penetration;
+			TVec3<T> Dir, WitnessA, WitnessB;
+			EXPECT_EQ(EPA(Tetrahedron, Zeros, Support, ZeroSupport, Penetration, Dir, WitnessA, WitnessB), EPAResult::Ok);
+			EXPECT_NEAR(Penetration, 1, 1e-1);
+			EXPECT_NEAR(Dir[0], 0, 1e-1);
+			EXPECT_NEAR(Dir[1], 0, 1e-1);
+			EXPECT_NEAR(Dir[2], 1, 1e-1);
+			EXPECT_NEAR(WitnessA[0], 0, 1e-1);
+			EXPECT_NEAR(WitnessA[1], 0, 1e-1);
+			EXPECT_NEAR(WitnessA[2], 1, 1e-1);
+			EXPECT_NEAR(WitnessB[0], 0, 1e-1);
+			EXPECT_NEAR(WitnessB[1], 0, 1e-1);
+			EXPECT_NEAR(WitnessB[2], 0, 1e-1);
+		}
+
+		{
+			//box is 1,1,1 with origin in the middle to handle cases when origin is right on tetrahedron
+			TVec3<T> HullVerts[8] = { {-1, -1, -1}, {1, -1, -1}, {-1, 1, -1}, {1, 1, -1},
+									  {-1, -1, 1}, {1, -1, 2}, {-1, 1, 1}, {1, 1, 1} };
+
+			auto Support = [&HullVerts](const auto& V)
+			{
+				auto MaxDist = TNumericLimits<T>::Lowest();
+				auto BestVert = HullVerts[0];
+				for (const auto& Vert : HullVerts)
+				{
+					const auto Dist = TVec3<T>::DotProduct(V, Vert);
+					if (Dist > MaxDist)
+					{
+						MaxDist = Dist;
+						BestVert = Vert;
+					}
+				}
+				return BestVert;
+			};
+
+			TArray<TVec3<T>> Tetrahedron = { HullVerts[0], HullVerts[2], HullVerts[3], HullVerts[4] };
+			TArray <TVec3<T>> Zeros = { TVec3<T>(0), TVec3<T>(0), TVec3<T>(0), TVec3<T>(0) };
+
+			T Penetration;
+			TVec3<T> Dir, WitnessA, WitnessB;
+			EXPECT_EQ(EPA(Tetrahedron, Zeros, Support, ZeroSupport, Penetration, Dir, WitnessA, WitnessB), EPAResult::Ok);
+			EXPECT_FLOAT_EQ(Penetration, 1);
+			EXPECT_NEAR(WitnessA.Size(), 1, 1e-1);	//don't know exact point, but should be 1 away from origin
+		}
+	}
+
+	template void EPASimpleTest<float>();
+
+	// Previously failing test cases that we would like to keep testing to prevent regression.
+	GTEST_TEST(EPATests, EPARealFailures_Fixed)
+	{
+		{
+			//get to EPA from GJKPenetration
+			TAABB<float, 3> Box({ -50, -50, -50 }, { 50, 50, 50 });
+
+			const FRigidTransform3 BToATM({ -8.74146843, 4.58291769, -100.029655 }, FRotation3::FromElements(6.63562241e-05, -0.000235952888, 0.00664712908, 0.999977887));
+			FVec3 ClosestA, ClosestB, Normal;
+			float Penetration;
+			GJKPenetration(Box, Box, BToATM, Penetration, ClosestA, ClosestB, Normal);
+
+			EXPECT_GT(Penetration, 0);
+			EXPECT_LT(Penetration, 2);
+		}
+
+		// Problem: EPA was selecting the wrong face on the second box, resulting in a large penetration depth (131cm, but the box is only 20cm thick)
+		// Fixed CL: 10615422
+		{
+			TBox<FReal, 3> A({ -12.5000000, -1.50000000, -12.5000000 }, { 12.5000000, 1.50000000, 12.5000000 });
+			TBox<FReal, 3> B({ -100.000000, -100.000000, -10.0000000 }, { 100.000000, 100.000000, 10.0000000 });
+			const FRigidTransform3 BToATM({ -34.9616776, 64.0135651, -10.9833698 }, FRotation3::FromElements(-0.239406615, -0.664629698, 0.637779951, 0.306901455));
+
+			FVec3 ClosestA, ClosestB, NormalA;
+			float Penetration;
+			GJKPenetration(A, B, BToATM, Penetration, ClosestA, ClosestB, NormalA);
+			FVec3 Normal = BToATM.InverseTransformVector(NormalA);
+
+			EXPECT_NEAR(Penetration, 0.025f, 0.005f);
+			EXPECT_NEAR(Normal.Z, -1.0f, 0.001f);
+		}
+
+		// Problem: EPA was selecting the wrong face on the second box, this was because LastEntry was initialized to first face, not best first face
+		// Fixed CL: 10635151>
+		{
+			TBox<FReal, 3> A({ -12.5000000, -1.50000000, -12.5000000 }, { 12.5000000, 1.50000000, 12.5000000 });
+			TBox<FReal, 3> B({ -100.000000, -100.000000, -10.0000000 }, { 100.000000, 100.000000, 10.0000000 });
+			const FRigidTransform3 BToATM({ -50.4365005, 52.8003693, -35.1415100 }, FRotation3::FromElements(-0.112581111, -0.689017475, 0.657892346, 0.282414317));
+
+			FVec3 ClosestA, ClosestB, NormalA;
+			float Penetration;
+			GJKPenetration(A, B, BToATM, Penetration, ClosestA, ClosestB, NormalA);
+			FVec3 Normal = BToATM.InverseTransformVector(NormalA);
+
+			EXPECT_LT(Penetration, 20);
+			EXPECT_NEAR(Normal.Z, -1.0f, 0.001f);
+		}
+
+		// Do not know what the expected output for this test is, but it is here because it once produced NaN in the V vector in GJKRaycast2.
+		// Turn on NaN diagnostics if you want to be sure to catch the failure. (Fixed now)
+		{
+			TArray<TPlaneConcrete<FReal, 3>> ConvexPlanes(
+				{
+					{{0.000000000, -1024.00000, 2.84217094e-14}, {0.000000000, -1.00000000, 0.000000000}},
+					{{0.000000000, -256.000000, 8.00000000}, {0.000000000, 0.000000000, 1.00000000}},
+					{{0.000000000, -1024.00000, 8.00000000}, {0.000000000, -1.00000000, 0.000000000}},
+					{{0.000000000, -256.000000, 8.00000000}, {-1.00000000, -0.000000000, 0.000000000}},
+					{{768.000000, -1024.00000, 2.84217094e-14}, {-0.000000000, -6.47630076e-17, -1.00000000}},
+					{{0.000000000, -1024.00000, 2.84217094e-14}, {-1.00000000, 0.000000000, 0.000000000}},
+					{{0.000000000, -256.000000, 8.00000000}, {0.000000000, 0.000000000, 1.00000000}},
+					{{768.000000, -1024.00000, 8.00000000}, {1.00000000, -0.000000000, 0.000000000}},
+					{{768.000000, -1024.00000, 2.84217094e-14}, {6.62273836e-09, 6.62273836e-09, -1.00000000}},
+					{{768.000000, -448.000000, 8.00000000}, {1.00000000, 0.000000000, 0.000000000}},
+					{{0.000000000, -256.000000, -2.13162821e-14}, {0.000000000, 1.00000000, 0.000000000}},
+					{{0.000000000, -256.000000, 8.00000000}, {-0.000000000, 0.000000000, 1.00000000}},
+					{{768.000000, -448.000000, 8.00000000}, {0.707106829, 0.707106829, 0.000000000}},
+					{{576.000000, -256.000000, 3.81469727e-06}, {0.000000000, 1.00000000, -0.000000000}},
+					{{768.000000, -448.000000, 8.00000000}, {0.707106829, 0.707106829, 0.000000000}},
+					{{768.000000, -448.000000, 3.81469727e-06}, {6.62273836e-09, 6.62273836e-09, -1.00000000}}
+				});
+
+			TParticles<FReal, 3> SurfaceParticles(
+				{
+					{0.000000000, -1024.00000, 2.84217094e-14},
+					{768.000000, -1024.00000, 2.84217094e-14},
+					{0.000000000, -1024.00000, 8.00000000},
+					{0.000000000, -256.000000, 8.00000000},
+					{768.000000, -1024.00000, 8.00000000},
+					{0.000000000, -256.000000, -2.13162821e-14},
+					{768.000000, -448.000000, 8.00000000},
+					{768.000000, -448.000000, 3.81469727e-06},
+					{576.000000, -256.000000, 3.81469727e-06},
+					{576.000000, -256.000000, 8.00000000}
+				});
+
+			TUniquePtr<FConvex> Convex = MakeUnique<FConvex>(MoveTemp(ConvexPlanes), MoveTemp(SurfaceParticles));
+			TImplicitObjectScaled<FConvex> ScaledConvex(MakeSerializable(Convex), FVec3(1.0f), 0.0f);
+
+			TSphere<FReal, 3> Sphere(FVec3(0.0f), 34.2120171);
+
+			const FRigidTransform3 BToATM({ 568.001648, -535.998352, 8.00000000 }, FRotation3::FromElements(0.000000000, 0.000000000, -0.707105696, 0.707107902));
+			const FVec3 LocalDir(0.000000000, 0.000000000, -1.00000000);
+			const FReal Length = 384.000000;
+			const FReal Thickness = 0.0;
+			const bool bComputeMTD = true;
+			const FVec3 Offset(-536.000000, -568.000000, -8.00000000);
+
+			FReal OutTime = -1.0f;
+			FVec3 LocalPosition(-1.0f);
+			FVec3 LocalNormal(-1.0f);
+
+
+			bool bResult = GJKRaycast2(ScaledConvex, Sphere, BToATM, LocalDir, Length, OutTime, LocalPosition, LocalNormal, Thickness, bComputeMTD, Offset, Thickness);
+
+		}
+
+	}
+
+	// Currently broken EPA edge cases. As they are fixed move them to EPARealFailures_Fixed above so that we can ensure they don't break again.
+	GTEST_TEST(EPATests, EPARealFailures_Broken)
+	{
+
+	}
+}
\ No newline at end of file
diff --git a/Engine/Source/Programs/HeadlessChaos/Private/HeadlessChaosTestForces.cpp b/Engine/Source/Programs/HeadlessChaos/Private/HeadlessChaosTestForces.cpp
new file mode 100644
--- /dev/null
+++ b/Engine/Source/Programs/HeadlessChaos/Private/HeadlessChaosTestForces.cpp
@@ -0,0 +1,29 @@
+// Copyright Epic Games, Inc. All Rights Reserved.
+
+#include "HeadlessChaosTestForces.h"
+
+#include "HeadlessChaos.h"
+#include "HeadlessChaosTestUtility.h"
+#include "Modules/ModuleManager.h"
+#include "Chaos/PBDRigidsEvolution.h"
+#include "Chaos/PBDRigidParticles.h"
+#include "Chaos/PBDRigidsEvolutionGBF.h"
+#include "Chaos/Sphere.h"
+#include "Chaos/Utilities.h"
+
+namespace ChaosTest {
+
+	template<typename TEvolution, typename T>
+	void Gravity()
+	{
+		TPBDRigidsSOAs<T, 3> Particles;
+		TEvolution Evolution(Particles);
+		
+		TArray<TPBDRigidParticleHandle<T,3>*> Dynamics = Evolution.CreateDynamicParticles(1);
+		Evolution.AdvanceOneTimeStep(0.1);
+		EXPECT_LT(Dynamics[0]->X()[2], 0);
+	}
+	template void Gravity<Chaos::TPBDRigidsEvolutionGBF<float, 3>, float>();
+
+
+}
\ No newline at end of file
diff --git a/Engine/Source/Programs/HeadlessChaos/Private/HeadlessChaosTestGJK.cpp b/Engine/Source/Programs/HeadlessChaos/Private/HeadlessChaosTestGJK.cpp
new file mode 100644
--- /dev/null
+++ b/Engine/Source/Programs/HeadlessChaos/Private/HeadlessChaosTestGJK.cpp
@@ -0,0 +1,1432 @@
+// Copyright Epic Games, Inc. All Rights Reserved.
+
+#include "HeadlessChaosTestGJK.h"
+
+#include "HeadlessChaos.h"
+#include "HeadlessChaosTestUtility.h"
+#include "Chaos/GJK.h"
+#include "Chaos/Capsule.h"
+#include "Chaos/Convex.h"
+#include "Chaos/ImplicitObjectScaled.h"
+
+namespace ChaosTest
+{
+	using namespace Chaos;
+
+	//for each simplex test:
+	//- points get removed
+	// - points off simplex return false
+	//- points in simplex return true
+	//- degenerate simplex
+
+	template <typename T>
+	void SimplexLine()
+	{
+		{
+			T Barycentric[4];
+			const TVector<T, 3> Simplex[] = { {-1,-1,-1}, {-1,-1,1} };
+			int32 Idxs[] = { 0,1 };
+			int32 NumVerts = 2;
+			const TVector<T,3> ClosestPoint = LineSimplexFindOrigin(Simplex, Idxs, NumVerts, Barycentric);
+			EXPECT_EQ(NumVerts, 2);
+			EXPECT_FLOAT_EQ(ClosestPoint[0], -1);
+			EXPECT_FLOAT_EQ(ClosestPoint[1], -1);
+			EXPECT_FLOAT_EQ(ClosestPoint[2], 0);
+			EXPECT_FLOAT_EQ(Barycentric[0], 0.5);
+			EXPECT_FLOAT_EQ(Barycentric[1], 0.5);
+		}
+
+		{
+			T Barycentric[4];
+			const TVector<T, 3> Simplex[] = { {-1,-1,-1}, {1,1,1} };
+			int32 Idxs[] = { 0,1 };
+			int32 NumVerts = 2;
+			const TVector<T, 3> ClosestPoint = LineSimplexFindOrigin(Simplex, Idxs, NumVerts, Barycentric);
+			EXPECT_EQ(NumVerts, 2);
+			EXPECT_FLOAT_EQ(ClosestPoint[0], 0);
+			EXPECT_FLOAT_EQ(ClosestPoint[1], 0);
+			EXPECT_FLOAT_EQ(ClosestPoint[2], 0);
+			EXPECT_FLOAT_EQ(Barycentric[0], 0.5);
+			EXPECT_FLOAT_EQ(Barycentric[1], 0.5);
+		}
+
+		{
+			T Barycentric[4];
+			const TVector<T, 3> Simplex[] = { {1,1,1}, {1,2,3} };
+			int32 Idxs[] = { 0,1 };
+			int32 NumVerts = 2;
+			const TVector<T, 3> ClosestPoint = LineSimplexFindOrigin(Simplex, Idxs, NumVerts, Barycentric);
+			EXPECT_EQ(NumVerts, 1);
+			EXPECT_FLOAT_EQ(ClosestPoint[0], 1);
+			EXPECT_FLOAT_EQ(ClosestPoint[1], 1);
+			EXPECT_FLOAT_EQ(ClosestPoint[2], 1);
+			EXPECT_FLOAT_EQ(Barycentric[0], 1);
+			EXPECT_EQ(Idxs[0], 0);
+		}
+
+		{
+			T Barycentric[4];
+			const TVector<T, 3> Simplex[] = { {10,11,12}, {1,2,3} };
+			int32 Idxs[] = { 0,1 };
+			int32 NumVerts = 2;
+			const TVector<T, 3> ClosestPoint = LineSimplexFindOrigin(Simplex, Idxs, NumVerts, Barycentric);
+			EXPECT_EQ(NumVerts, 1);
+			EXPECT_FLOAT_EQ(ClosestPoint[0], 1);
+			EXPECT_FLOAT_EQ(ClosestPoint[1], 2);
+			EXPECT_FLOAT_EQ(ClosestPoint[2], 3);
+			EXPECT_FLOAT_EQ(Barycentric[1], 1);
+			EXPECT_EQ(Idxs[0], 1);
+		}
+
+		{
+			T Barycentric[4];
+			const TVector<T, 3> Simplex[] = { {1,1,1}, {1,1,1} };
+			int32 Idxs[] = { 0,1 };
+			int32 NumVerts = 2;
+			const TVector<T, 3> ClosestPoint = LineSimplexFindOrigin(Simplex, Idxs, NumVerts, Barycentric);
+			EXPECT_EQ(NumVerts, 1);
+			EXPECT_FLOAT_EQ(ClosestPoint[0], 1);
+			EXPECT_FLOAT_EQ(ClosestPoint[1], 1);
+			EXPECT_FLOAT_EQ(ClosestPoint[2], 1);
+			EXPECT_FLOAT_EQ(Barycentric[0], 1);
+			EXPECT_EQ(Idxs[0], 0);
+		}
+
+		{
+			T Barycentric[4];
+			const TVector<T, 3> Simplex[] = { {1,-1e-16,1}, {1,1e-16,1} };
+			int32 Idxs[] = { 0,1 };
+			int32 NumVerts = 2;
+			const TVector<T, 3> ClosestPoint = LineSimplexFindOrigin(Simplex, Idxs, NumVerts, Barycentric);
+			EXPECT_EQ(NumVerts, 2);
+			EXPECT_FLOAT_EQ(ClosestPoint[0], 1);
+			EXPECT_FLOAT_EQ(ClosestPoint[1], 0);
+			EXPECT_FLOAT_EQ(ClosestPoint[2], 1);
+			EXPECT_FLOAT_EQ(Barycentric[0], 0.5);
+			EXPECT_FLOAT_EQ(Barycentric[1], 0.5);
+			EXPECT_EQ(Idxs[0], 0);
+			EXPECT_EQ(Idxs[1], 1);
+		}
+	}
+
+	template <typename T>
+	void SimplexTriangle()
+	{
+		{
+			T Barycentric[4];
+			const TVector<T, 3> Simplex[] = { {-1,-1,-1}, {-1,1,-1}, {-2,1,-1} };
+			FSimplex Idxs = { 0,1, 2 };
+			
+			const TVector<T, 3> ClosestPoint = TriangleSimplexFindOrigin(Simplex, Idxs, Barycentric);
+			EXPECT_EQ(Idxs.NumVerts, 2);
+			EXPECT_FLOAT_EQ(ClosestPoint[0], -1);
+			EXPECT_FLOAT_EQ(ClosestPoint[1], 0);
+			EXPECT_FLOAT_EQ(ClosestPoint[2], -1);
+			EXPECT_EQ(Idxs[0], 0);
+			EXPECT_EQ(Idxs[1], 1);
+			EXPECT_FLOAT_EQ(Barycentric[0], 0.5);
+			EXPECT_FLOAT_EQ(Barycentric[1], 0.5);
+		}
+
+		{
+			T Barycentric[4];
+			const TVector<T, 3> Simplex[] = { {-1,-1,-1},{-2,1,-1}, {-1,1,-1} };
+			FSimplex Idxs = { 0,1, 2 };
+			const TVector<T, 3> ClosestPoint = TriangleSimplexFindOrigin(Simplex, Idxs, Barycentric);
+			EXPECT_EQ(Idxs.NumVerts, 2);
+			EXPECT_FLOAT_EQ(ClosestPoint[0], -1);
+			EXPECT_FLOAT_EQ(ClosestPoint[1], 0);
+			EXPECT_FLOAT_EQ(ClosestPoint[2], -1);
+			EXPECT_EQ(Idxs[0], 0);
+			EXPECT_EQ(Idxs[1], 2);
+			EXPECT_FLOAT_EQ(Barycentric[0], 0.5);
+			EXPECT_FLOAT_EQ(Barycentric[2], 0.5);
+		}
+
+		{
+			//corner
+			T Barycentric[4];
+			const TVector<T, 3> Simplex[] = { {1,1,1},{2,1,1}, {2,2,1} };
+			FSimplex Idxs = { 1,0, 2 };
+			const TVector<T, 3> ClosestPoint = TriangleSimplexFindOrigin(Simplex, Idxs, Barycentric);
+			EXPECT_EQ(Idxs.NumVerts, 1);
+			EXPECT_FLOAT_EQ(ClosestPoint[0], 1);
+			EXPECT_FLOAT_EQ(ClosestPoint[1], 1);
+			EXPECT_FLOAT_EQ(ClosestPoint[2], 1);
+			EXPECT_EQ(Idxs[0], 0);
+			EXPECT_FLOAT_EQ(Barycentric[0], 1);
+		}
+
+		{
+			//corner equal
+			T Barycentric[4];
+			const TVector<T, 3> Simplex[] = { {0,0,0},{2,1,1}, {2,2,1} };
+			FSimplex Idxs = { 0,1, 2 };
+			const TVector<T, 3> ClosestPoint = TriangleSimplexFindOrigin(Simplex, Idxs, Barycentric);
+			EXPECT_EQ(Idxs.NumVerts, 1);
+			EXPECT_FLOAT_EQ(ClosestPoint[0], 0);
+			EXPECT_FLOAT_EQ(ClosestPoint[1], 0);
+			EXPECT_FLOAT_EQ(ClosestPoint[2], 0);
+			EXPECT_EQ(Idxs[0], 0);
+			EXPECT_FLOAT_EQ(Barycentric[0], 1);
+		}
+
+		{
+			//edge equal
+			T Barycentric[4];
+			const TVector<T, 3> Simplex[] = { {-1,0,0},{1,0,0}, {0,2,0} };
+			FSimplex Idxs = { 2,0, 1 };
+			const TVector<T, 3> ClosestPoint = TriangleSimplexFindOrigin(Simplex, Idxs, Barycentric);
+			EXPECT_EQ(Idxs.NumVerts, 2);
+			EXPECT_FLOAT_EQ(ClosestPoint[0], 0);
+			EXPECT_FLOAT_EQ(ClosestPoint[1], 0);
+			EXPECT_FLOAT_EQ(ClosestPoint[2], 0);
+			EXPECT_EQ(Idxs[0], 0);
+			EXPECT_EQ(Idxs[1], 1);
+			EXPECT_FLOAT_EQ(Barycentric[0], 0.5);
+			EXPECT_FLOAT_EQ(Barycentric[1], 0.5);
+		}
+
+		{
+			//triangle equal
+			T Barycentric[4];
+			const TVector<T, 3> Simplex[] = { {-1,0,-1},{1,0,-1}, {0,0,1} };
+			FSimplex Idxs = { 0,1, 2 };
+			const TVector<T, 3> ClosestPoint = TriangleSimplexFindOrigin(Simplex, Idxs, Barycentric);
+			EXPECT_EQ(Idxs.NumVerts, 3);
+			EXPECT_FLOAT_EQ(ClosestPoint[0], 0);
+			EXPECT_FLOAT_EQ(ClosestPoint[1], 0);
+			EXPECT_FLOAT_EQ(ClosestPoint[2], 0);
+			EXPECT_EQ(Idxs[0], 0);
+			EXPECT_EQ(Idxs[1], 1);
+			EXPECT_EQ(Idxs[2], 2);
+			EXPECT_FLOAT_EQ(Barycentric[0], 0.25);
+			EXPECT_FLOAT_EQ(Barycentric[1], 0.25);
+			EXPECT_FLOAT_EQ(Barycentric[2], 0.5);
+		}
+
+		{
+			//co-linear
+			T Barycentric[4];
+			const TVector<T, 3> Simplex[] = { {-1,-1,-1},{-1,1,-1}, {-1,1.2,-1} };
+			FSimplex Idxs = { 0,1, 2 };
+			const TVector<T, 3> ClosestPoint = TriangleSimplexFindOrigin(Simplex, Idxs, Barycentric);
+			EXPECT_EQ(Idxs.NumVerts, 2);
+			EXPECT_FLOAT_EQ(ClosestPoint[0], -1);
+			EXPECT_FLOAT_EQ(ClosestPoint[1], 0);
+			EXPECT_FLOAT_EQ(ClosestPoint[2], -1);
+			EXPECT_EQ(Idxs[0], 0);
+			EXPECT_EQ(Idxs[1], 1);	//degenerate triangle throws out newest point
+			EXPECT_FLOAT_EQ(Barycentric[0], 0.5);
+			EXPECT_FLOAT_EQ(Barycentric[1], 0.5);
+		}
+
+		{
+			//single point
+			T Barycentric[4];
+			const TVector<T, 3> Simplex[] = { {-1,-1,-1},{-1,-1,-1}, {-1,-1,-1} };
+			FSimplex Idxs = { 0,2, 1 };
+			const TVector<T, 3> ClosestPoint = TriangleSimplexFindOrigin(Simplex, Idxs, Barycentric);
+			EXPECT_EQ(Idxs.NumVerts, 1);
+			EXPECT_FLOAT_EQ(ClosestPoint[0], -1);
+			EXPECT_FLOAT_EQ(ClosestPoint[1], -1);
+			EXPECT_FLOAT_EQ(ClosestPoint[2], -1);
+			EXPECT_EQ(Idxs[0], 0);
+			EXPECT_FLOAT_EQ(Barycentric[0], 1);
+		}
+
+		{
+			//corner perfect split
+			T Barycentric[4];
+			const TVector<T, 3> Simplex[] = { {-1,-1,0},{1,-1,0}, {0,-0.5,0} };
+			FSimplex Idxs = { 0,2, 1 };
+			const TVector<T, 3> ClosestPoint = TriangleSimplexFindOrigin(Simplex, Idxs, Barycentric);
+			EXPECT_EQ(Idxs.NumVerts, 1);
+			EXPECT_FLOAT_EQ(ClosestPoint[0], 0);
+			EXPECT_FLOAT_EQ(ClosestPoint[1], -0.5);
+			EXPECT_FLOAT_EQ(ClosestPoint[2], 0);
+			EXPECT_EQ(Idxs[0], 2);
+			EXPECT_FLOAT_EQ(Barycentric[2], 1);
+		}
+
+		{
+			//triangle face correct distance
+			T Barycentric[4];
+			const TVector<T, 3> Simplex[] = { {-1,-1,-1},{1,-1,-1}, {0,1,-1} };
+			FSimplex Idxs = { 0,1,2 };
+			const TVector<T, 3> ClosestPoint = TriangleSimplexFindOrigin(Simplex, Idxs, Barycentric);
+			EXPECT_EQ(Idxs.NumVerts, 3);
+			EXPECT_FLOAT_EQ(ClosestPoint[0], 0);
+			EXPECT_FLOAT_EQ(ClosestPoint[1], 0);
+			EXPECT_FLOAT_EQ(ClosestPoint[2], -1);
+			EXPECT_EQ(Idxs[0], 0);
+			EXPECT_EQ(Idxs[1], 1);
+			EXPECT_EQ(Idxs[2], 2);
+			EXPECT_FLOAT_EQ(Barycentric[0], 0.25);
+			EXPECT_FLOAT_EQ(Barycentric[1], 0.25);
+			EXPECT_FLOAT_EQ(Barycentric[2], 0.5);
+		}
+
+		{
+			//tiny triangle middle point
+			T Barycentric[4];
+			const TVector<T, 3> Simplex[] = { {-1e-9,-1e-9,-1e-9},{-1e-9,1e-9,-1e-9}, {-1e-9,0,1e-9} };
+			FSimplex Idxs = { 0,1,2 };
+			const TVector<T, 3> ClosestPoint = TriangleSimplexFindOrigin(Simplex, Idxs, Barycentric);
+			EXPECT_EQ(Idxs.NumVerts, 3);
+			EXPECT_FLOAT_EQ(ClosestPoint[0], -1e-9);
+			EXPECT_FLOAT_EQ(ClosestPoint[1], 0);
+			EXPECT_FLOAT_EQ(ClosestPoint[2], 0);
+			EXPECT_EQ(Idxs[0], 0);
+			EXPECT_EQ(Idxs[1], 1);
+			EXPECT_EQ(Idxs[2], 2);
+			EXPECT_FLOAT_EQ(Barycentric[0], 0.25);
+			EXPECT_FLOAT_EQ(Barycentric[1], 0.25);
+			EXPECT_FLOAT_EQ(Barycentric[2], 0.5);
+		}
+
+		{
+			//non cartesian triangle plane
+			T Barycentric[4];
+			const TVector<T, 3> Simplex[] = { {2, 0, -1}, {0, 2, -1}, {1, 1, 1} };
+			FSimplex Idxs = { 0,1,2 };
+			const TVector<T, 3> ClosestPoint = TriangleSimplexFindOrigin(Simplex, Idxs, Barycentric);
+			EXPECT_EQ(Idxs.NumVerts, 3);
+			EXPECT_FLOAT_EQ(ClosestPoint[0], 1);
+			EXPECT_FLOAT_EQ(ClosestPoint[1], 1);
+			EXPECT_FLOAT_EQ(ClosestPoint[2], 0);
+			EXPECT_EQ(Idxs[0], 0);
+			EXPECT_EQ(Idxs[1], 1);
+			EXPECT_EQ(Idxs[2], 2);
+			EXPECT_FLOAT_EQ(Barycentric[0], 0.25);
+			EXPECT_FLOAT_EQ(Barycentric[1], 0.25);
+			EXPECT_FLOAT_EQ(Barycentric[2], 0.5);
+		}
+	}
+
+	template <typename T>
+	void SimplexTetrahedron()
+	{
+		{
+			//top corner
+			T Barycentric[4];
+			const TVector<T, 3> Simplex[] = { {-1,-1,-1}, {1,-1,-1}, {0,1,-1}, {0,0,-0.5} };
+			FSimplex Idxs = { 0,1,2,3 };
+			const TVector<T, 3> ClosestPoint = TetrahedronSimplexFindOrigin(Simplex, Idxs, Barycentric);
+			EXPECT_EQ(Idxs.NumVerts, 1);
+			EXPECT_FLOAT_EQ(ClosestPoint[0], 0);
+			EXPECT_FLOAT_EQ(ClosestPoint[1], 0);
+			EXPECT_FLOAT_EQ(ClosestPoint[2], -0.5);
+			EXPECT_EQ(Idxs[0], 3);
+			EXPECT_FLOAT_EQ(Barycentric[3], 1);
+		}
+
+		{
+			//inside
+			T Barycentric[4];
+			const TVector<T, 3> Simplex[] = { {-1,-1,-1}, {1,-1,-1}, {0,1,-1}, {0,0,0.5} };
+			FSimplex Idxs = { 0,1,2,3 };
+			const TVector<T, 3> ClosestPoint = TetrahedronSimplexFindOrigin(Simplex, Idxs, Barycentric);
+			EXPECT_EQ(Idxs.NumVerts, 4);
+			EXPECT_FLOAT_EQ(ClosestPoint[0], 0);
+			EXPECT_FLOAT_EQ(ClosestPoint[1], 0);
+			EXPECT_FLOAT_EQ(ClosestPoint[2], 0);
+			EXPECT_EQ(Idxs[0], 0);
+			EXPECT_EQ(Idxs[1], 1);
+			EXPECT_EQ(Idxs[2], 2);
+			EXPECT_EQ(Idxs[3], 3);
+			EXPECT_FLOAT_EQ(Barycentric[0] + Barycentric[1] + Barycentric[2] + Barycentric[3], 1);
+		}
+
+		{
+			//face
+			T Barycentric[4];
+			const TVector<T, 3> Simplex[] = { {0,0,-1.5}, {-1,-1,-1}, {1,-1,-1}, {0,1,-1} };
+			FSimplex Idxs = { 0,1,2,3 }; 
+			const TVector<T, 3> ClosestPoint = TetrahedronSimplexFindOrigin(Simplex, Idxs, Barycentric);
+			EXPECT_EQ(Idxs.NumVerts, 3);
+			EXPECT_FLOAT_EQ(ClosestPoint[0], 0);
+			EXPECT_FLOAT_EQ(ClosestPoint[1], 0);
+			EXPECT_FLOAT_EQ(ClosestPoint[2], -1);
+			EXPECT_EQ(Idxs[0], 1);
+			EXPECT_EQ(Idxs[1], 2);
+			EXPECT_EQ(Idxs[2], 3);
+			EXPECT_FLOAT_EQ(Barycentric[1] + Barycentric[2] + Barycentric[3], 1);
+		}
+
+		{
+			//edge
+			T Barycentric[4];
+			const TVector<T, 3> Simplex[] = { {-1,-1,0}, {1,-1,0}, {0,-1,-1}, {0, -2, -1} };
+			FSimplex Idxs = { 0,1,2,3 };
+			const TVector<T, 3> ClosestPoint = TetrahedronSimplexFindOrigin(Simplex, Idxs, Barycentric);
+			EXPECT_EQ(Idxs.NumVerts, 2);
+			EXPECT_FLOAT_EQ(ClosestPoint[0], 0);
+			EXPECT_FLOAT_EQ(ClosestPoint[1], -1);
+			EXPECT_FLOAT_EQ(ClosestPoint[2], 0);
+			EXPECT_EQ(Idxs[0], 0);
+			EXPECT_EQ(Idxs[1], 1);
+			EXPECT_FLOAT_EQ(Barycentric[0], 0.5);
+			EXPECT_FLOAT_EQ(Barycentric[1], 0.5);
+		}
+
+		{
+			//degenerate
+			T Barycentric[4];
+			const TVector<T, 3> Simplex[] = { {-1,-1,0}, {1,-1,0}, {0,-1,-1}, {0, -1, -0.5} };
+			FSimplex Idxs = { 0,1,2,3 };
+			const TVector<T, 3> ClosestPoint = TetrahedronSimplexFindOrigin(Simplex, Idxs, Barycentric);
+			EXPECT_EQ(Idxs.NumVerts, 2);
+			EXPECT_FLOAT_EQ(ClosestPoint[0], 0);
+			EXPECT_FLOAT_EQ(ClosestPoint[1], -1);
+			EXPECT_FLOAT_EQ(ClosestPoint[2], 0);
+			EXPECT_EQ(Idxs[0], 0);
+			EXPECT_EQ(Idxs[1], 1);
+			EXPECT_FLOAT_EQ(Barycentric[0], 0.5);
+			EXPECT_FLOAT_EQ(Barycentric[1], 0.5);
+		}
+
+		{
+			//wide angle, bad implementation would return edge but it's really a face
+			T Barycentric[4];
+			const TVector<T, 3> Simplex[] = { {-10000,-1,10000}, {1,-1,10000}, {4,-3,10000}, {1, -1, -10000} };
+			FSimplex Idxs = { 0,1,2,3 };
+			const TVector<T, 3> ClosestPoint = TetrahedronSimplexFindOrigin(Simplex, Idxs, Barycentric);
+			EXPECT_EQ(Idxs.NumVerts, 3);
+			EXPECT_FLOAT_EQ(ClosestPoint[0], 0);
+			EXPECT_FLOAT_EQ(ClosestPoint[1], -1);
+			EXPECT_FLOAT_EQ(ClosestPoint[2], 0);
+			EXPECT_EQ(Idxs[0], 0);
+			EXPECT_EQ(Idxs[1], 1);
+			EXPECT_EQ(Idxs[2], 3);
+			EXPECT_FLOAT_EQ(Barycentric[0] + Barycentric[1] + Barycentric[3], 1);
+		}
+	}
+
+	//For each gjk test we should test:
+	// - thickness
+	// - transformed geometry
+	// - rotated geometry
+	// - degenerate cases
+	// - near miss, near hit
+	// - multiple initial dir
+
+	template <typename T>
+	void GJKSphereSphereTest()
+	{
+		TSphere<T, 3> A(TVector<T, 3>(10, 0, 0), 5);
+		TSphere<T, 3> B(TVector<T, 3>(4, 0, 0), 2);
+
+		TVector<T, 3> InitialDirs[] = { TVector<T,3>(1,0,0), TVector<T,3>(-1,0,0), TVector<T,3>(0,1,0), TVector<T,3>(0,-1,0), TVector<T,3>(0,0,1), TVector<T,3>(0,0,-1) };
+
+		for (const TVector<T, 3>& InitialDir : InitialDirs)
+		{
+			EXPECT_TRUE(GJKIntersection<T>(A, B, TRigidTransform<T, 3>::Identity, 0, InitialDir));
+
+			//miss
+			EXPECT_FALSE(GJKIntersection<T>(A, B, TRigidTransform<T, 3>(TVector<T, 3>(-1.1, 0, 0), TRotation<T, 3>::Identity), 0, InitialDir));
+
+			//hit from thickness
+			EXPECT_TRUE(GJKIntersection<T>(A, B, TRigidTransform<T, 3>(TVector<T, 3>(-1.1, 0, 0), TRotation<T, 3>::Identity), 0.105, InitialDir));
+
+			//miss with thickness
+			EXPECT_FALSE(GJKIntersection<T>(A, B, TRigidTransform<T, 3>(TVector<T, 3>(-1.1, 0, 0), TRotation<T, 3>::Identity), 0.095, InitialDir));
+
+			//hit with rotation
+			EXPECT_TRUE(GJKIntersection<T>(A, B, TRigidTransform<T, 3>(TVector<T, 3>(6.5, 0, 0), TRotation<T, 3>::FromVector(TVector<T, 3>(0, 0, PI))), 1, InitialDir));
+
+			//miss with rotation
+			EXPECT_FALSE(GJKIntersection<T>(A, B, TRigidTransform<T, 3>(TVector<T, 3>(6.5, 0, 0), TRotation<T, 3>::FromVector(TVector<T, 3>(0, 0, PI))), 0.01, InitialDir));
+
+			//hit tiny
+			TSphere<T, 3> Tiny(TVector<T, 3>(0), 1e-2);
+			EXPECT_TRUE(GJKIntersection<T>(A, Tiny, TRigidTransform<T, 3>(TVector<T, 3>(15, 0, 0), TRotation<T, 3>::Identity), 0, InitialDir));
+
+			//miss tiny
+			EXPECT_FALSE(GJKIntersection<T>(A, Tiny, TRigidTransform<T, 3>(TVector<T, 3>(15 + 1e-1, 0, 0), TRotation<T, 3>::Identity), 0, InitialDir));
+		}
+	}
+
+
+	template <typename T>
+	void GJKSphereBoxTest()
+	{
+		TSphere<T, 3> A(TVector<T, 3>(10, 0, 0), 5);
+		TAABB<T, 3> B(TVector<T, 3>(-4, -2, -4), TVector<T,3>(4,2,4));
+
+		TVector<T, 3> InitialDirs[] = { TVector<T,3>(1,0,0), TVector<T,3>(-1,0,0), TVector<T,3>(0,1,0), TVector<T,3>(0,-1,0), TVector<T,3>(0,0,1), TVector<T,3>(0,0,-1) };
+
+		for (const TVector<T, 3>& InitialDir : InitialDirs)
+		{
+			EXPECT_TRUE(GJKIntersection<T>(A, B, TRigidTransform<T, 3>(TVector<T, 3>(1, 0, 0), TRotation<T, 3>::Identity), 0, InitialDir));
+
+			//miss
+			EXPECT_FALSE(GJKIntersection<T>(A, B, TRigidTransform<T, 3>(TVector<T, 3>(0.9, 0, 0), TRotation<T, 3>::Identity), 0, InitialDir));
+
+			//rotate and hit
+			EXPECT_TRUE(GJKIntersection<T>(A, B, TRigidTransform<T, 3>(TVector<T, 3>(3.1, 0, 0), TRotation<T, 3>::FromVector(TVector<T,3>(0,0,PI*0.5))), 0, InitialDir));
+
+			//rotate and miss
+			EXPECT_FALSE(GJKIntersection<T>(A, B, TRigidTransform<T, 3>(TVector<T, 3>(2.9, 0, 0), TRotation<T, 3>::FromVector(TVector<T, 3>(0, 0, PI*0.5))), 0, InitialDir));
+
+			//rotate and hit from thickness
+			EXPECT_TRUE(GJKIntersection<T>(A, B, TRigidTransform<T, 3>(TVector<T, 3>(2.9, 0, 0), TRotation<T, 3>::FromVector(TVector<T, 3>(0, 0, PI*0.5))), 0.1, InitialDir));
+
+			//hit thin
+			TAABB<T, 3> Thin(TVector<T, 3>(4, -2, -4), TVector<T, 3>(4, 2, 4));
+			EXPECT_TRUE(GJKIntersection<T>(A, Thin, TRigidTransform<T, 3>(TVector<T, 3>(1+1e-2, 0, 0), TRotation<T, 3>::Identity), 0, InitialDir));
+
+			//miss
+			EXPECT_FALSE(GJKIntersection<T>(A, Thin, TRigidTransform<T, 3>(TVector<T, 3>(1 - 1e-2, 0, 0), TRotation<T, 3>::Identity), 0, InitialDir));
+
+			//hit line
+			TAABB<T, 3> Line(TVector<T, 3>(4, -2, 0), TVector<T, 3>(4, 2, 0));
+			EXPECT_TRUE(GJKIntersection<T>(A, Line, TRigidTransform<T, 3>(TVector<T, 3>(1 + 1e-2, 0, 0), TRotation<T, 3>::Identity), 0, InitialDir));
+
+			//miss
+			EXPECT_FALSE(GJKIntersection<T>(A, Line, TRigidTransform<T, 3>(TVector<T, 3>(1 - 1e-2, 0, 0), TRotation<T, 3>::Identity), 0, InitialDir));
+		}
+	}
+
+
+	template <typename T>
+	void GJKSphereCapsuleTest()
+	{
+		TSphere<T, 3> A(TVector<T, 3>(10, 0, 0), 5);
+		TCapsule<T> B(TVector<T, 3>(0, 0, -3), TVector<T, 3>(0, 0, 3), 3);
+
+		TVector<T, 3> InitialDirs[] = { TVector<T,3>(1,0,0), TVector<T,3>(-1,0,0), TVector<T,3>(0,1,0), TVector<T,3>(0,-1,0), TVector<T,3>(0,0,1), TVector<T,3>(0,0,-1) };
+
+		for (const TVector<T, 3>& InitialDir : InitialDirs)
+		{
+			EXPECT_TRUE(GJKIntersection<T>(A, B, TRigidTransform<T, 3>(TVector<T, 3>(2, 0, 0), TRotation<T, 3>::Identity), 0, InitialDir));
+
+			//miss
+			EXPECT_FALSE(GJKIntersection<T>(A, B, TRigidTransform<T, 3>(TVector<T, 3>(2-1e-2, 0, 0), TRotation<T, 3>::Identity), 0, InitialDir));
+
+			//thickness
+			EXPECT_TRUE(GJKIntersection<T>(A, B, TRigidTransform<T, 3>(TVector<T, 3>(1, 0, 0), TRotation<T, 3>::Identity), 1.01, InitialDir));
+
+			//miss
+			EXPECT_FALSE(GJKIntersection<T>(A, B, TRigidTransform<T, 3>(TVector<T, 3>(1, 0, 0), TRotation<T, 3>::Identity), 0.99, InitialDir));
+
+			//rotation hit
+			EXPECT_TRUE(GJKIntersection<T>(A, B, TRigidTransform<T, 3>(TVector<T, 3>(-1+1e-2, 0, 0), TRotation<T, 3>::FromVector(TVector<T,3>(0,PI*0.5,0))), 0, InitialDir));
+
+			//miss
+			EXPECT_FALSE(GJKIntersection<T>(A, B, TRigidTransform<T, 3>(TVector<T, 3>(-1-1e-2, 0, 0), TRotation<T, 3>::FromVector(TVector<T, 3>(0, PI*0.5, 0))), 0, InitialDir));
+
+			//degenerate
+			TCapsule<T> Line(TVector<T, 3>(0, 0, -3), TVector<T, 3>(0, 0, 3), 0);
+			EXPECT_TRUE(GJKIntersection<T>(A, Line, TRigidTransform<T, 3>(TVector<T, 3>(5+1e-2, 0, 0), TRotation<T, 3>::Identity), 0, InitialDir));
+
+			//miss
+			EXPECT_FALSE(GJKIntersection<T>(A, Line, TRigidTransform<T, 3>(TVector<T, 3>(5 - 1e-2, 0, 0), TRotation<T, 3>::Identity), 0, InitialDir));
+		}
+	}
+
+
+	template <typename T>
+	void GJKSphereConvexTest()
+	{
+		TVector<T, 3> InitialDirs[] = { TVector<T,3>(1,0,0), TVector<T,3>(-1,0,0), TVector<T,3>(0,1,0), TVector<T,3>(0,-1,0), TVector<T,3>(0,0,1), TVector<T,3>(0,0,-1) };
+		TSphere<T, 3> A(TVector<T, 3>(10, 0, 0), 5);
+
+		{
+			//Tetrahedron
+			TParticles<T, 3> HullParticles;
+			HullParticles.AddParticles(4);
+			HullParticles.X(0) = { -1,-1,-1 };
+			HullParticles.X(1) = { 1,-1,-1 };
+			HullParticles.X(2) = { 0,1,-1 };
+			HullParticles.X(3) = { 0,0,1 };
+			FConvex B(HullParticles);
+
+			for (const TVector<T, 3>& InitialDir : InitialDirs)
+			{
+				//hit
+				EXPECT_TRUE(GJKIntersection<T>(A, B, TRigidTransform<T, 3>(TVector<T, 3>(5, 0, 0), TRotation<T, 3>::Identity), 0, InitialDir));
+
+				//near hit
+				EXPECT_TRUE(GJKIntersection<T>(A, B, TRigidTransform<T, 3>(TVector<T, 3>(4 + 1e-4, 1, 1), TRotation<T, 3>::Identity), 0, InitialDir));
+
+				//near miss
+				EXPECT_FALSE(GJKIntersection<T>(A, B, TRigidTransform<T, 3>(TVector<T, 3>(4 - 1e-2, 1, 1), TRotation<T, 3>::Identity), 0, InitialDir));
+
+				//rotated hit
+				EXPECT_TRUE(GJKIntersection<T>(A, B, TRigidTransform<T, 3>(TVector<T, 3>(4 + 1e-4, 0, 1), TRotation<T, 3>::FromVector(TVector<T, 3>(0, 0, PI*0.5))), 0, InitialDir));
+
+				//rotated miss
+				EXPECT_FALSE(GJKIntersection<T>(A, B, TRigidTransform<T, 3>(TVector<T, 3>(4 - 1e-2, 0, 1), TRotation<T, 3>::FromVector(TVector<T, 3>(0, 0, PI*0.5))), 0, InitialDir));
+
+				//rotated and inflated hit
+				EXPECT_TRUE(GJKIntersection<T>(A, B, TRigidTransform<T, 3>(TVector<T, 3>(3.5, 0, 1), TRotation<T, 3>::FromVector(TVector<T, 3>(0, 0, PI*0.5))), 0.5 + 1e-4, InitialDir));
+
+				//rotated and inflated miss
+				EXPECT_FALSE(GJKIntersection<T>(A, B, TRigidTransform<T, 3>(TVector<T, 3>(3.5, 0, 1), TRotation<T, 3>::FromVector(TVector<T, 3>(0, 0, PI*0.5))), 0.5 - 1e-2, InitialDir));
+			}
+		}
+
+		{
+			//Triangle
+			TParticles<T, 3> TriangleParticles;
+			TriangleParticles.AddParticles(3);
+			TriangleParticles.X(0) = { -1,-1,-1 };
+			TriangleParticles.X(1) = { 1,-1,-1 };
+			TriangleParticles.X(2) = { 0,1,-1 };
+			FConvex B(TriangleParticles);
+
+			//triangle
+			for (const TVector<T, 3>& InitialDir : InitialDirs)
+			{
+				//hit
+				EXPECT_TRUE(GJKIntersection<T>(A, B, TRigidTransform<T, 3>(TVector<T, 3>(5, 0, 0), TRotation<T, 3>::Identity), 0, InitialDir));
+
+				//near hit
+				EXPECT_TRUE(GJKIntersection<T>(A, B, TRigidTransform<T, 3>(TVector<T, 3>(4 + 1e-2, 1, 1), TRotation<T, 3>::Identity), 0, InitialDir));
+
+				//near miss
+				EXPECT_FALSE(GJKIntersection<T>(A, B, TRigidTransform<T, 3>(TVector<T, 3>(4 - 1e-2, 1, 1), TRotation<T, 3>::Identity), 0, InitialDir));
+
+				//rotated hit
+				EXPECT_TRUE(GJKIntersection<T>(A, B, TRigidTransform<T, 3>(TVector<T, 3>(4 + 1e-2, 0, 1), TRotation<T, 3>::FromVector(TVector<T, 3>(0, 0, PI*0.5))), 0, InitialDir));
+
+				//rotated miss
+				EXPECT_FALSE(GJKIntersection<T>(A, B, TRigidTransform<T, 3>(TVector<T, 3>(4 - 1e-2, 0, 1), TRotation<T, 3>::FromVector(TVector<T, 3>(0, 0, PI*0.5))), 0, InitialDir));
+
+				//rotated and inflated hit
+				EXPECT_TRUE(GJKIntersection<T>(A, B, TRigidTransform<T, 3>(TVector<T, 3>(3.5, 0, 1), TRotation<T, 3>::FromVector(TVector<T, 3>(0, 0, PI*0.5))), 0.5 + 1e-2, InitialDir));
+
+				//rotated and inflated miss
+				EXPECT_FALSE(GJKIntersection<T>(A, B, TRigidTransform<T, 3>(TVector<T, 3>(3.5, 0, 1), TRotation<T, 3>::FromVector(TVector<T, 3>(0, 0, PI*0.5))), 0.5 - 1e-2, InitialDir));
+			}
+		}
+	}
+
+
+	template <typename T>
+	void GJKSphereScaledSphereTest()
+	{
+		TSphere<T, 3> A(TVector<T, 3>(10, 0, 0), 5);
+		TUniquePtr<TSphere<T, 3>> Sphere = MakeUnique<TSphere<T, 3>>(TVector<T, 3>(4, 0, 0), 2);
+		TImplicitObjectScaled<TSphere<T, 3>> Unscaled(MakeSerializable(Sphere), TVector<T,3>(1));
+		TImplicitObjectScaled<TSphere<T, 3>> UniformScaled(MakeSerializable(Sphere), TVector<T, 3>(2));
+		TImplicitObjectScaled<TSphere<T, 3>> NonUniformScaled(MakeSerializable(Sphere), TVector<T, 3>(2,1,1));
+
+		TVector<T, 3> InitialDirs[] = { TVector<T,3>(1,0,0), TVector<T,3>(-1,0,0), TVector<T,3>(0,1,0), TVector<T,3>(0,-1,0), TVector<T,3>(0,0,1), TVector<T,3>(0,0,-1) };
+
+		for (const TVector<T, 3>& InitialDir : InitialDirs)
+		{
+			EXPECT_TRUE(GJKIntersection<T>(A, Unscaled, TRigidTransform<T, 3>::Identity, 0, InitialDir));
+			EXPECT_TRUE(GJKIntersection<T>(A, UniformScaled, TRigidTransform<T, 3>::Identity, 0, InitialDir));
+			//EXPECT_TRUE(GJKIntersection<T>(A, NonUniformScaled, TRigidTransform<T, 3>::Identity, 0, InitialDir));
+
+			//miss
+			EXPECT_FALSE(GJKIntersection<T>(A, Unscaled, TRigidTransform<T, 3>(TVector<T, 3>(-1.1, 0, 0), TRotation<T, 3>::Identity), 0, InitialDir));
+			EXPECT_FALSE(GJKIntersection<T>(A, UniformScaled, TRigidTransform<T, 3>(TVector<T, 3>(-7.1, 0, 0), TRotation<T, 3>::Identity), 0, InitialDir));
+			//EXPECT_FALSE(GJKIntersection<T>(A, NonUniformScaled, TRigidTransform<T, 3>(TVector<T, 3>(-7.1, 0, 0), TRotation<T, 3>::Identity), 0, InitialDir));
+
+			//hit from thickness
+			EXPECT_TRUE(GJKIntersection<T>(A, Unscaled, TRigidTransform<T, 3>(TVector<T, 3>(-1.1, 0, 0), TRotation<T, 3>::Identity), 0.105, InitialDir));
+			EXPECT_TRUE(GJKIntersection<T>(A, UniformScaled, TRigidTransform<T, 3>(TVector<T, 3>(-7.1, 0, 0), TRotation<T, 3>::Identity), 0.105, InitialDir));
+			//EXPECT_TRUE(GJKIntersection<T>(A, NonUniformScaled, TRigidTransform<T, 3>(TVector<T, 3>(-7.1, 0, 0), TRotation<T, 3>::Identity), 0.105, InitialDir));
+
+			//miss with thickness
+			EXPECT_FALSE(GJKIntersection<T>(A, Unscaled, TRigidTransform<T, 3>(TVector<T, 3>(-1.1, 0, 0), TRotation<T, 3>::Identity), 0.095, InitialDir));
+			EXPECT_FALSE(GJKIntersection<T>(A, UniformScaled, TRigidTransform<T, 3>(TVector<T, 3>(-7.1, 0, 0), TRotation<T, 3>::Identity), 0.095, InitialDir));
+			//EXPECT_FALSE(GJKIntersection<T>(A, NonUniformScaled, TRigidTransform<T, 3>(TVector<T, 3>(-7.1, 0, 0), TRotation<T, 3>::Identity), 0.095, InitialDir));
+
+			//hit with rotation
+			EXPECT_TRUE(GJKIntersection<T>(A, Unscaled, TRigidTransform<T, 3>(TVector<T, 3>(6.5, 0, 0), TRotation<T, 3>::FromVector(TVector<T, 3>(0, 0, PI))), 1, InitialDir));
+			EXPECT_TRUE(GJKIntersection<T>(A, UniformScaled, TRigidTransform<T, 3>(TVector<T, 3>(8.1, 0, 0), TRotation<T, 3>::FromVector(TVector<T, 3>(0, 0, PI))), 1, InitialDir));
+			//EXPECT_TRUE(GJKIntersection<T>(A, NonUniformScaled, TRigidTransform<T, 3>(TVector<T, 3>(8.1, 0, 0), TRotation<T, 3>::FromVector(TVector<T, 3>(0, 0, PI))), 1, InitialDir));
+
+			//miss with rotation
+			EXPECT_FALSE(GJKIntersection<T>(A, Unscaled, TRigidTransform<T, 3>(TVector<T, 3>(6.5, 0, 0), TRotation<T, 3>::FromVector(TVector<T, 3>(0, 0, PI))), 0.01, InitialDir));
+			EXPECT_FALSE(GJKIntersection<T>(A, UniformScaled, TRigidTransform<T, 3>(TVector<T, 3>(8.1, 0, 0), TRotation<T, 3>::FromVector(TVector<T, 3>(0, 0, PI))), 0.01, InitialDir));
+			//EXPECT_FALSE(GJKIntersection<T>(A, NonUniformScaled, TRigidTransform<T, 3>(TVector<T, 3>(8.1, 0, 0), TRotation<T, 3>::FromVector(TVector<T, 3>(0, 0, PI))), 0.01, InitialDir));
+		}
+	}
+
+	//For each gjkraycast test we should test:
+	// - thickness
+	// - initial overlap
+	// - transformed geometry
+	// - rotated geometry
+	// - offset transform
+	// - degenerate cases
+	// - near miss, near hit
+	// - multiple initial dir
+
+	template <typename T>
+	void GJKSphereSphereSweep()
+	{
+		typedef TVector<T, 3> TVector3;
+		TSphere<T, 3> A(TVector<T, 3>(10, 0, 0), 5);
+		TSphere<T, 3> B(TVector<T, 3>(1, 0, 0), 2);
+
+		TVector<T, 3> InitialDirs[] = { TVector<T,3>(1,0,0), TVector<T,3>(-1,0,0), TVector<T,3>(0,1,0), TVector<T,3>(0,-1,0), TVector<T,3>(0,0,1), TVector<T,3>(0,0,-1) };
+
+		constexpr T Eps = 1e-1;
+
+		for (const TVector<T, 3>& InitialDir : InitialDirs)
+		{
+			T Time;
+			TVector<T, 3> Position;
+			TVector<T, 3> Normal;
+
+			//hit
+			EXPECT_TRUE(GJKRaycast<T>(A, B, TRigidTransform<T, 3>::Identity, TVector<T, 3>(1, 0, 0), 4, Time, Position, Normal, 0, InitialDir));
+			EXPECT_NEAR(Time, 2, Eps);
+			EXPECT_VECTOR_NEAR(Normal, TVector3(-1, 0, 0), Eps);
+			EXPECT_VECTOR_NEAR(Position, TVector3(5, 0, 0), Eps);
+
+			//hit offset
+			EXPECT_TRUE(GJKRaycast<T>(A, B, TRigidTransform<T, 3>(TVector<T,3>(1,0,0), TRotation<T,3>::Identity), TVector<T, 3>(1, 0, 0), 4, Time, Position, Normal, 0, InitialDir));
+			EXPECT_NEAR(Time, 1, Eps);
+			EXPECT_VECTOR_NEAR(Normal, TVector3(-1, 0, 0), Eps);
+			EXPECT_VECTOR_NEAR(Position, TVector3(5, 0, 0), Eps);
+
+			//initial overlap
+			EXPECT_TRUE(GJKRaycast2<T>(A, B, TRigidTransform<T, 3>(TVector<T, 3>(7, 0, 0), TRotation<T, 3>::Identity), TVector<T, 3>(1, 0, 0), 4, Time, Position, Normal, 0, false, InitialDir));
+			EXPECT_FLOAT_EQ(Time, 0);
+
+			//MTD
+			EXPECT_TRUE(GJKRaycast2<T>(A, B, TRigidTransform<T, 3>(TVector<T, 3>(7, 0, 0), TRotation<T, 3>::Identity), TVector<T, 3>(1, 0, 0), 4, Time, Position, Normal, 0, true, InitialDir));
+			EXPECT_FLOAT_EQ(Time, -5);
+			EXPECT_VECTOR_NEAR(Position, TVector3(5,0,0), Eps);
+			EXPECT_VECTOR_NEAR(Normal, TVector3(-1, 0, 0), Eps);
+			
+			//EPA
+			EXPECT_TRUE(GJKRaycast2<T>(A, B, TRigidTransform<T, 3>(TVector<T, 3>(9, 0, 0), TRotation<T, 3>::Identity), TVector<T, 3>(1, 0, 0), 4, Time, Position, Normal, 0, true, InitialDir));
+			EXPECT_FLOAT_EQ(Time, -7);	//perfect overlap, will default to 0,0,1 normal
+			EXPECT_VECTOR_NEAR(Position, TVec3<T>(10,0,5), Eps);
+			EXPECT_VECTOR_NEAR(Normal, TVec3<T>(0, 0, 1), Eps);
+
+			//miss
+			EXPECT_FALSE(GJKRaycast<T>(A, B, TRigidTransform<T, 3>(TVector<T, 3>(0, 0, 7.1), TRotation<T, 3>::Identity), TVector<T, 3>(1, 0, 0), 20, Time, Position, Normal, 0, InitialDir));
+
+			//hit with thickness
+			EXPECT_TRUE(GJKRaycast<T>(A, B, TRigidTransform<T, 3>(TVector<T, 3>(0, 0, 7.1), TRotation<T, 3>::Identity), TVector<T, 3>(1, 0, 0), 20, Time, Position, Normal, 0.2, InitialDir));
+
+			//hit rotated
+			const TRotation<T, 3> RotatedDown(TRotation<T, 3>::FromVector(TVector<T, 3>(0, PI * 0.5, 0)));
+			EXPECT_TRUE(GJKRaycast<T>(A, B, TRigidTransform<T, 3>(TVector<T, 3>(0, 0, 7.9), RotatedDown), TVector<T, 3>(1, 0, 0), 20, Time, Position, Normal, 0, InitialDir));
+
+			//miss rotated
+			EXPECT_FALSE(GJKRaycast<T>(A, B, TRigidTransform<T, 3>(TVector<T, 3>(0, 0, 8.1), RotatedDown), TVector<T, 3>(1, 0, 0), 20, Time, Position, Normal, 0, InitialDir));
+
+			//hit rotated with inflation
+			EXPECT_TRUE(GJKRaycast<T>(A, B, TRigidTransform<T, 3>(TVector<T, 3>(0, 0, 7.9), RotatedDown), TVector<T, 3>(1, 0, 0), 20, Time, Position, Normal, 0.2, InitialDir));
+
+			//near hit
+			EXPECT_TRUE(GJKRaycast<T>(A, B, TRigidTransform<T, 3>(TVector<T, 3>(0, 0, 7 - 1e-2), TRotation<T, 3>::Identity), TVector<T, 3>(1, 0, 0), 20, Time, Position, Normal, 0, InitialDir));
+
+			//near miss
+			EXPECT_FALSE(GJKRaycast<T>(A, B, TRigidTransform<T, 3>(TVector<T, 3>(0, 0, 7 + 1e-2), TRotation<T, 3>::Identity), TVector<T, 3>(1, 0, 0), 20, Time, Position, Normal, 0, InitialDir));
+
+			//degenerate
+			TSphere<T, 3> Tiny(TVector<T, 3>(1, 0, 0), 1e-8);
+			EXPECT_TRUE(GJKRaycast<T>(A, Tiny, TRigidTransform<T, 3>::Identity, TVector<T, 3>(1, 0, 0), 8, Time, Position, Normal, 0, InitialDir));
+			EXPECT_NEAR(Time, 4, Eps);
+			EXPECT_VECTOR_NEAR(Normal, TVector3(-1, 0, 0), Eps);
+			EXPECT_VECTOR_NEAR(Position, TVector3(5, 0, 0), Eps);
+
+			//right at end
+			EXPECT_TRUE(GJKRaycast<T>(A, B, TRigidTransform<T, 3>::Identity, TVector<T, 3>(1, 0, 0), 2, Time, Position, Normal, 0, InitialDir));
+			EXPECT_NEAR(Time, 2, Eps);
+
+			// not far enough
+			EXPECT_FALSE(GJKRaycast<T>(A, B, TRigidTransform<T, 3>::Identity, TVector<T, 3>(1, 0, 0), 2 - 1e-2, Time, Position, Normal, 0, InitialDir));
+		}
+	}
+
+	template <typename T>
+	void GJKSphereBoxSweep()
+	{
+		typedef TVector<T, 3> TVector3;
+		TAABB<T, 3> A(TVector<T, 3>(3, -1, 0), TVector<T, 3>(4, 1, 4));
+		TSphere<T, 3> B(TVector<T, 3>(0, 0, 0), 1);
+
+		TVector<T, 3> InitialDirs[] = { TVector<T,3>(1,0,0), TVector<T,3>(-1,0,0), TVector<T,3>(0,1,0), TVector<T,3>(0,-1,0), TVector<T,3>(0,0,1), TVector<T,3>(0,0,-1) };
+
+		constexpr T Eps = 1e-1;
+
+		for (const TVector<T, 3>& InitialDir : InitialDirs)
+		{
+			T Time;
+			TVector<T, 3> Position;
+			TVector<T, 3> Normal;
+
+			//hit
+			EXPECT_TRUE(GJKRaycast<T>(A, B, TRigidTransform<T, 3>::Identity, TVector<T, 3>(1, 0, 0), 4, Time, Position, Normal, 0, InitialDir));
+			EXPECT_NEAR(Time, 2, Eps);
+			EXPECT_VECTOR_NEAR(Normal, TVector3(-1, 0, 0), Eps);
+			EXPECT_VECTOR_NEAR(Position, TVector3(3, 0, 0), Eps);
+
+			//hit offset
+			EXPECT_TRUE(GJKRaycast<T>(A, B, TRigidTransform<T, 3>(TVector<T, 3>(1.5, 0, 0), TRotation<T, 3>::Identity), TVector<T, 3>(1, 0, 0), 4, Time, Position, Normal, 0, InitialDir));
+			EXPECT_NEAR(Time, 0.5, Eps);
+			EXPECT_VECTOR_NEAR(Normal, TVector3(-1, 0, 0), Eps);
+			EXPECT_VECTOR_NEAR(Position, TVector3(3, 0, 0), Eps);
+
+			//initial overlap
+			EXPECT_TRUE(GJKRaycast2<T>(A, B, TRigidTransform<T, 3>(TVector<T, 3>(4, 0, 4), TRotation<T, 3>::Identity), TVector<T, 3>(1, 0, 0), 4, Time, Position, Normal, 0, false, InitialDir));
+			EXPECT_FLOAT_EQ(Time, 0);
+
+			//MTD without EPA
+			EXPECT_TRUE(GJKRaycast2<T>(A, B, TRigidTransform<T, 3>(TVector<T, 3>(4.25, 0, 2), TRotation<T, 3>::Identity), TVector<T, 3>(1, 0, 0), 4, Time, Position, Normal, 0, true, InitialDir));
+			EXPECT_FLOAT_EQ(Time, -0.75);
+			EXPECT_VECTOR_NEAR(Position, TVector3(4, 0, 2), Eps);
+			EXPECT_VECTOR_NEAR(Normal, TVector3(1, 0, 0), Eps);
+
+			//MTD with EPA
+			EXPECT_TRUE(GJKRaycast2<T>(A, B, TRigidTransform<T, 3>(TVector<T, 3>(4, 0, 2), TRotation<T, 3>::Identity), TVector<T, 3>(1, 0, 0), 4, Time, Position, Normal, 0, true, InitialDir));
+			EXPECT_FLOAT_EQ(Time, -1);
+			EXPECT_VECTOR_NEAR(Position, TVector3(4, 0, 2), Eps);
+			EXPECT_VECTOR_NEAR(Normal, TVector3(1, 0, 0), Eps);
+
+			//MTD with EPA
+			EXPECT_TRUE(GJKRaycast2<T>(A, B, TRigidTransform<T, 3>(TVector<T, 3>(3.25, 0, 2), TRotation<T, 3>::Identity), TVector<T, 3>(1, 0, 0), 4, Time, Position, Normal, 0, true, InitialDir));
+			EXPECT_FLOAT_EQ(Time, -1.25);
+			EXPECT_VECTOR_NEAR(Position, TVector3(3, 0, 2), Eps);
+			EXPECT_VECTOR_NEAR(Normal, TVector3(-1, 0, 0), Eps);
+
+			//MTD with EPA
+			EXPECT_TRUE(GJKRaycast2<T>(A, B, TRigidTransform<T, 3>(TVector<T, 3>(3.4, 0, 3.75), TRotation<T, 3>::Identity), TVector<T, 3>(1, 0, 0), 4, Time, Position, Normal, 0, true, InitialDir));
+			EXPECT_FLOAT_EQ(Time, -1.25);
+			EXPECT_VECTOR_NEAR(Position, TVector3(3.4, 0, 4), Eps);
+			EXPECT_VECTOR_NEAR(Normal, TVector3(0, 0, 1), Eps);
+
+			//hit
+			EXPECT_TRUE(GJKRaycast<T>(A, B, TRigidTransform<T, 3>(TVector<T, 3>(1, 0, 6), TRotation<T, 3>::Identity), TVector<T, 3>(1, 0, -1).GetUnsafeNormal(), 4, Time, Position, Normal, 0, InitialDir));
+			const T ExpectedTime = ((TVector<T, 3>(3, 0, 4) - TVector<T, 3>(1, 0, 6)).Size() - 1);
+			EXPECT_NEAR(Time, ExpectedTime, Eps);
+			EXPECT_VECTOR_NEAR(Normal, TVector3(-sqrt(2) / 2, 0, sqrt(2) / 2), Eps);
+			EXPECT_VECTOR_NEAR(Position, TVector3(3, 0, 4), Eps);
+
+			//near miss
+			EXPECT_FALSE(GJKRaycast<T>(A, B, TRigidTransform<T, 3>(TVector<T, 3>(0, 0, 5+1e-2), TRotation<T, 3>::Identity), TVector<T, 3>(1, 0, 0), 4, Time, Position, Normal, 0, InitialDir));
+
+			//near hit with inflation
+			EXPECT_TRUE(GJKRaycast<T>(A, B, TRigidTransform<T, 3>(TVector<T, 3>(0, 0, 5 + 1e-2), TRotation<T, 3>::Identity), TVector<T, 3>(1, 0, 0), 4, Time, Position, Normal, 2e-2, InitialDir));
+			const T DistanceFromCorner = (Position - TVector<T, 3>(3, 0, 4)).Size();
+			EXPECT_LT(DistanceFromCorner, 1e-1);
+
+			//rotated box
+			const TRotation<T, 3> Rotated(TRotation<T, 3>::FromVector(TVector<T, 3>(0, 0, PI * 0.5)));
+			EXPECT_TRUE(GJKRaycast<T>(B, A, TRigidTransform<T, 3>(TVector<T, 3>(0), Rotated), TVector<T, 3>(0, -1, 0), 4, Time, Position, Normal, 0, InitialDir));
+			EXPECT_NEAR(Time, 2, Eps);
+			EXPECT_VECTOR_NEAR(Normal, TVector3(0, 1, 0), Eps);
+			EXPECT_VECTOR_NEAR(Position, TVector3(0, 1, 0), Eps);
+
+			//degenerate box
+			TAABB<T, 3> Needle(TVector<T, 3>(3, 0, 0), TVector<T, 3>(4, 0, 0));
+			EXPECT_TRUE(GJKRaycast<T>(B, Needle, TRigidTransform<T, 3>(TVector<T, 3>(0), Rotated), TVector<T, 3>(0, -1, 0), 4, Time, Position, Normal, 0, InitialDir));
+			EXPECT_NEAR(Time, 2, Eps);
+			EXPECT_VECTOR_NEAR(Normal, TVector3(0, 1, 0), Eps);
+			EXPECT_VECTOR_NEAR(Position, TVector3(0, 1, 0), Eps);
+		}
+	}
+
+
+	template <typename T>
+	void GJKSphereCapsuleSweep()
+	{
+		typedef TVector<T, 3> TVector3;
+		TSphere<T, 3> A(TVector<T, 3>(10, 0, 0), 5);
+		TCapsule<T> B(TVector<T, 3>(1, 0, 0), TVector<T, 3>(-3, 0, 0), 2);
+
+		TVector<T, 3> InitialDirs[] = { TVector<T,3>(1,0,0), TVector<T,3>(-1,0,0), TVector<T,3>(0,1,0), TVector<T,3>(0,-1,0), TVector<T,3>(0,0,1), TVector<T,3>(0,0,-1) };
+
+		constexpr T Eps = 1e-1;
+
+		for (const TVector<T, 3>& InitialDir : InitialDirs)
+		{
+			T Time;
+			TVector<T, 3> Position;
+			TVector<T, 3> Normal;
+
+			//hit
+			EXPECT_TRUE(GJKRaycast<T>(A, B, TRigidTransform<T, 3>::Identity, TVector<T, 3>(1, 0, 0), 4, Time, Position, Normal, 0, InitialDir));
+			EXPECT_NEAR(Time, 2, Eps);
+			EXPECT_VECTOR_NEAR(Normal, TVector3(-1, 0, 0), Eps);
+			EXPECT_VECTOR_NEAR(Position, TVector3(5, 0, 0), Eps);
+
+			//hit offset
+			EXPECT_TRUE(GJKRaycast<T>(A, B, TRigidTransform<T, 3>(TVector<T, 3>(1, 0, 0), TRotation<T, 3>::Identity), TVector<T, 3>(1, 0, 0), 4, Time, Position, Normal, 0, InitialDir));
+			EXPECT_NEAR(Time, 1, Eps);
+			EXPECT_VECTOR_NEAR(Normal, TVector3(-1, 0, 0), Eps);
+			EXPECT_VECTOR_NEAR(Position, TVector3(5, 0, 0), Eps);
+			
+			//initial overlap
+			EXPECT_TRUE(GJKRaycast2<T>(A, B, TRigidTransform<T, 3>(TVector<T, 3>(7, 0, 0), TRotation<T, 3>::Identity), TVector<T, 3>(1, 0, 0), 4, Time, Position, Normal, 0, false, InitialDir));
+			EXPECT_FLOAT_EQ(Time, 0);
+
+			//MTD
+			EXPECT_TRUE(GJKRaycast2<T>(A, B, TRigidTransform<T, 3>(TVector<T, 3>(7, 0, 0), TRotation<T, 3>::Identity), TVector<T, 3>(1, 0, 0), 4, Time, Position, Normal, 0, true, InitialDir));
+			EXPECT_FLOAT_EQ(Time, -5);
+			EXPECT_VECTOR_NEAR(Position, TVector3(5, 0, 0), Eps);
+			EXPECT_VECTOR_NEAR(Normal, TVector3(-1, 0, 0), Eps);
+
+			//miss
+			EXPECT_FALSE(GJKRaycast<T>(A, B, TRigidTransform<T, 3>(TVector<T, 3>(0, 0, 7.1), TRotation<T, 3>::Identity), TVector<T, 3>(1, 0, 0), 20, Time, Position, Normal, 0, InitialDir));
+
+			//hit with thickness
+			EXPECT_TRUE(GJKRaycast<T>(A, B, TRigidTransform<T, 3>(TVector<T, 3>(0, 0, 7.1), TRotation<T, 3>::Identity), TVector<T, 3>(1, 0, 0), 20, Time, Position, Normal, 0.2, InitialDir));
+
+			//hit rotated
+			const TRotation<T, 3> RotatedDown(TRotation<T, 3>::FromVector(TVector<T, 3>(0, PI * 0.5, 0)));
+			EXPECT_TRUE(GJKRaycast<T>(A, B, TRigidTransform<T, 3>(TVector<T, 3>(0, 0, 7.9), RotatedDown), TVector<T, 3>(1, 0, 0), 20, Time, Position, Normal, 0, InitialDir));
+
+			//miss rotated
+			EXPECT_FALSE(GJKRaycast<T>(A, B, TRigidTransform<T, 3>(TVector<T, 3>(0, 0, 8.1), RotatedDown), TVector<T, 3>(1, 0, 0), 20, Time, Position, Normal, 0, InitialDir));
+
+			//hit rotated with inflation
+			EXPECT_TRUE(GJKRaycast<T>(A, B, TRigidTransform<T, 3>(TVector<T, 3>(0, 0, 7.9), RotatedDown), TVector<T, 3>(1, 0, 0), 20, Time, Position, Normal, 0.2, InitialDir));
+
+			//near hit
+			EXPECT_TRUE(GJKRaycast<T>(A, B, TRigidTransform<T, 3>(TVector<T, 3>(0, 0, 7 - 1e-2), TRotation<T, 3>::Identity), TVector<T, 3>(1, 0, 0), 20, Time, Position, Normal, 0, InitialDir));
+
+			//near miss
+			EXPECT_FALSE(GJKRaycast<T>(A, B, TRigidTransform<T, 3>(TVector<T, 3>(0, 0, 7 + 1e-2), TRotation<T, 3>::Identity), TVector<T, 3>(1, 0, 0), 20, Time, Position, Normal, 0, InitialDir));
+
+			//degenerate
+			TSphere<T, 3> Tiny(TVector<T, 3>(1, 0, 0), 1e-8);
+			EXPECT_TRUE(GJKRaycast<T>(A, Tiny, TRigidTransform<T, 3>::Identity, TVector<T, 3>(1, 0, 0), 8, Time, Position, Normal, 0, InitialDir));
+			EXPECT_NEAR(Time, 4, Eps);
+			EXPECT_VECTOR_NEAR(Normal, TVector3(-1, 0, 0), Eps);
+			EXPECT_VECTOR_NEAR(Position, TVector3(5, 0, 0), Eps);
+
+			//right at end
+			EXPECT_TRUE(GJKRaycast<T>(A, B, TRigidTransform<T, 3>::Identity, TVector<T, 3>(1, 0, 0), 2, Time, Position, Normal, 0, InitialDir));
+			EXPECT_NEAR(Time, 2, Eps);
+
+			// not far enough
+			EXPECT_FALSE(GJKRaycast<T>(A, B, TRigidTransform<T, 3>::Identity, TVector<T, 3>(1, 0, 0), 2 - 1e-2, Time, Position, Normal, 0, InitialDir));
+		}
+	}
+
+
+	template <typename T>
+	void GJKSphereConvexSweep()
+	{
+		typedef TVector<T, 3> TVector3;
+		//Tetrahedron
+		TParticles<T, 3> HullParticles;
+		HullParticles.AddParticles(4);
+		HullParticles.X(0) = { 3,0,4 };
+		HullParticles.X(1) = { 3,1,0 };
+		HullParticles.X(2) = { 3,-1,0 };
+		HullParticles.X(3) = { 4,0,2 };
+		FConvex A(HullParticles);
+		TSphere<T, 3> B(TVector<T, 3>(0, 0, 0), 1);
+
+		TVector<T, 3> InitialDirs[] = { TVector<T,3>(1,0,0), TVector<T,3>(-1,0,0), TVector<T,3>(0,1,0), TVector<T,3>(0,-1,0), TVector<T,3>(0,0,1), TVector<T,3>(0,0,-1) };
+
+		constexpr T Eps = 1e-1;
+
+		for (const TVector<T, 3>& InitialDir : InitialDirs)
+		{
+			T Time;
+			TVector<T, 3> Position;
+			TVector<T, 3> Normal;
+
+			//hit
+			EXPECT_TRUE(GJKRaycast<T>(A, B, TRigidTransform<T, 3>::Identity, TVector<T, 3>(1, 0, 0), 4, Time, Position, Normal, 0, InitialDir));
+			EXPECT_NEAR(Time, 2, Eps);
+			EXPECT_VECTOR_NEAR(Normal, TVector3(-1, 0, 0), Eps);
+			EXPECT_VECTOR_NEAR(Position, TVector3(3, 0, 0), Eps);
+
+			//hit offset
+			EXPECT_TRUE(GJKRaycast<T>(A, B, TRigidTransform<T, 3>(TVector<T, 3>(1.5, 0, 0), TRotation<T, 3>::Identity), TVector<T, 3>(1, 0, 0), 4, Time, Position, Normal, 0, InitialDir));
+			EXPECT_NEAR(Time, 0.5, Eps);
+			EXPECT_VECTOR_NEAR(Normal, TVector3(-1, 0, 0), Eps);
+			EXPECT_VECTOR_NEAR(Position, TVector3(3, 0, 0), Eps);
+
+			//initial overlap
+			EXPECT_TRUE(GJKRaycast2<T>(A, B, TRigidTransform<T, 3>(TVector<T, 3>(4, 0, 4), TRotation<T, 3>::Identity), TVector<T, 3>(1, 0, 0), 4, Time, Position, Normal, 0, false, InitialDir));
+			EXPECT_FLOAT_EQ(Time, 0);
+
+			//MTD
+			EXPECT_TRUE(GJKRaycast2<T>(A, B, TRigidTransform<T, 3>(TVector<T, 3>(2.5, 0, 2), TRotation<T, 3>::Identity), TVector<T, 3>(1, 0, 0), 4, Time, Position, Normal, 0, true, InitialDir));
+			EXPECT_FLOAT_EQ(Time, -0.5);
+			EXPECT_VECTOR_NEAR(Normal, TVector3(-1, 0, 0).GetUnsafeNormal(), Eps);
+
+			//MTD
+			T Penetration;
+			TVec3<T> ClosestA, ClosestB;
+			EXPECT_TRUE(GJKPenetration<T>(A, B, TRigidTransform<T, 3>(TVector<T, 3>(2.5, 0, 2), TRotation<T, 3>::Identity), Penetration, ClosestA, ClosestB, Normal, 0, InitialDir));
+			EXPECT_FLOAT_EQ(Penetration, 0.5);
+			EXPECT_VECTOR_NEAR(Normal, TVector3(-1, 0, 0).GetUnsafeNormal(), Eps);
+			EXPECT_NEAR(ClosestA[0], 3, Eps);	//could be any point on face, but should have x == 3
+			EXPECT_VECTOR_NEAR(ClosestB, TVector3(3.5, 0, 2), Eps);
+
+			//hit
+			EXPECT_TRUE(GJKRaycast<T>(A, B, TRigidTransform<T, 3>(TVector<T, 3>(1, 0, 6), TRotation<T, 3>::Identity), TVector<T, 3>(1, 0, -1).GetUnsafeNormal(), 4, Time, Position, Normal, 0, InitialDir));
+			const T ExpectedTime = ((TVector<T, 3>(3, 0, 4) - TVector<T, 3>(1, 0, 6)).Size() - 1);
+			EXPECT_NEAR(Time, ExpectedTime, Eps);
+			EXPECT_VECTOR_NEAR(Normal, TVector3(-sqrt(2) / 2, 0, sqrt(2) / 2), Eps);
+			EXPECT_VECTOR_NEAR(Position, TVector3(3, 0, 4), Eps);
+
+			//near miss
+			EXPECT_FALSE(GJKRaycast<T>(A, B, TRigidTransform<T, 3>(TVector<T, 3>(0, 0, 5 + 1e-2), TRotation<T, 3>::Identity), TVector<T, 3>(1, 0, 0), 4, Time, Position, Normal, 0, InitialDir));
+
+			//near hit with inflation
+			EXPECT_TRUE(GJKRaycast<T>(A, B, TRigidTransform<T, 3>(TVector<T, 3>(0, 0, 5 + 1e-2), TRotation<T, 3>::Identity), TVector<T, 3>(1, 0, 0), 4, Time, Position, Normal, 2e-2, InitialDir));
+			const T DistanceFromCorner = (Position - TVector<T, 3>(3, 0, 4)).Size();
+			EXPECT_LT(DistanceFromCorner, 1e-1);
+
+			//rotated box
+			const TRotation<T, 3> Rotated(TRotation<T, 3>::FromVector(TVector<T, 3>(0, 0, PI * 0.5)));
+			EXPECT_TRUE(GJKRaycast<T>(B, A, TRigidTransform<T, 3>(TVector<T, 3>(0), Rotated), TVector<T, 3>(0, -1, 0), 4, Time, Position, Normal, 0, InitialDir));
+			EXPECT_NEAR(Time, 2, Eps);
+			EXPECT_NEAR(Normal.X, 0, Eps);
+			EXPECT_NEAR(Normal.Y, 1, Eps);
+			//EXPECT_NEAR(Normal.Z, 0, Eps);
+			EXPECT_VECTOR_NEAR(Position, TVector3(0, 1, 0), Eps);
+
+			//degenerate box
+			TAABB<T, 3> Needle(TVector<T, 3>(3, 0, 0), TVector<T, 3>(4, 0, 0));
+			EXPECT_TRUE(GJKRaycast<T>(B, Needle, TRigidTransform<T, 3>(TVector<T, 3>(0), Rotated), TVector<T, 3>(0, -1, 0), 4, Time, Position, Normal, 0, InitialDir));
+			EXPECT_NEAR(Time, 2, Eps);
+			EXPECT_VECTOR_NEAR(Normal, TVector3(0, 1, 0), Eps);
+			EXPECT_VECTOR_NEAR(Position, TVector3(0, 1, 0), Eps);
+		}
+	}
+
+	template <typename T>
+	void GJKSphereScaledSphereSweep()
+	{
+		typedef TVector<T, 3> TVector3;
+		TSphere<T, 3> A(TVector<T, 3>(10, 0, 0), 5);
+		TUniquePtr<TSphere<T, 3>> Sphere = MakeUnique<TSphere<T, 3>>(TVector<T, 3>(0, 0, 0), 2);
+		TImplicitObjectScaled<TSphere<T, 3>> Unscaled(MakeSerializable(Sphere), TVector<T, 3>(1));
+		TImplicitObjectScaled<TSphere<T, 3>> UniformScaled(MakeSerializable(Sphere), TVector<T, 3>(2));
+		TImplicitObjectScaled<TSphere<T, 3>> NonUniformScaled(MakeSerializable(Sphere), TVector<T, 3>(2, 1, 1));
+
+		TVector<T, 3> InitialDirs[] = { TVector<T,3>(1,0,0), TVector<T,3>(-1,0,0), TVector<T,3>(0,1,0), TVector<T,3>(0,-1,0), TVector<T,3>(0,0,1), TVector<T,3>(0,0,-1) };
+
+		constexpr T Eps = 1e-1;
+
+		for (const TVector<T, 3>& InitialDir : InitialDirs)
+		{
+			T Time;
+			TVector<T, 3> Position;
+			TVector<T, 3> Normal;
+
+			//hit
+			EXPECT_TRUE(GJKRaycast<T>(A, Unscaled, TRigidTransform<T, 3>::Identity, TVector<T, 3>(1, 0, 0), 4, Time, Position, Normal, 0, InitialDir));
+			EXPECT_NEAR(Time, 3, Eps);
+			EXPECT_VECTOR_NEAR(Normal, TVector3(-1, 0, 0), Eps);
+			EXPECT_VECTOR_NEAR(Position, TVector3(5, 0, 0), Eps);
+
+			EXPECT_TRUE(GJKRaycast<T>(A, UniformScaled, TRigidTransform<T, 3>::Identity, TVector<T, 3>(1, 0, 0), 6, Time, Position, Normal, 0, InitialDir));
+			EXPECT_NEAR(Time, 1, Eps);
+			EXPECT_VECTOR_NEAR(Normal, TVector3(-1, 0, 0), Eps);
+			EXPECT_VECTOR_NEAR(Position, TVector3(5, 0, 0), Eps);
+
+			EXPECT_TRUE(GJKRaycast<T>(A, NonUniformScaled, TRigidTransform<T, 3>::Identity, TVector<T, 3>(1, 0, 0), 4, Time, Position, Normal, 0, InitialDir));
+			EXPECT_NEAR(Time, 1, Eps);
+			EXPECT_VECTOR_NEAR(Normal, TVector3(-1, 0, 0), Eps);
+			EXPECT_VECTOR_NEAR(Position, TVector3(5, 0, 0), Eps);
+
+			//hit offset
+			EXPECT_TRUE(GJKRaycast<T>(A, Unscaled, TRigidTransform<T, 3>(TVector<T, 3>(1, 0, 0), TRotation<T, 3>::Identity), TVector<T, 3>(1, 0, 0), 4, Time, Position, Normal, 0, InitialDir));
+			EXPECT_NEAR(Time, 2, Eps);
+			EXPECT_VECTOR_NEAR(Normal, TVector3(-1, 0, 0), Eps);
+			EXPECT_VECTOR_NEAR(Position, TVector3(5, 0, 0), Eps);
+
+			EXPECT_TRUE(GJKRaycast<T>(A, UniformScaled, TRigidTransform<T, 3>(TVector<T, 3>(1, 0, 0), TRotation<T, 3>::Identity), TVector<T, 3>(1, 0, 0), 4, Time, Position, Normal, 0, InitialDir));
+			EXPECT_NEAR(Time, 0, Eps);
+			EXPECT_VECTOR_NEAR(Normal, TVector3(-1, 0, 0), Eps);
+			EXPECT_VECTOR_NEAR(Position, TVector3(5, 0, 0), Eps);
+
+			EXPECT_TRUE(GJKRaycast<T>(A, NonUniformScaled, TRigidTransform<T, 3>(TVector<T, 3>(1, 0, 0), TRotation<T, 3>::Identity), TVector<T, 3>(1, 0, 0), 4, Time, Position, Normal, 0, InitialDir));
+			EXPECT_NEAR(Time, 0, Eps);
+			EXPECT_VECTOR_NEAR(Normal, TVector3(-1, 0, 0), Eps);
+			EXPECT_VECTOR_NEAR(Position, TVector3(5, 0, 0), Eps);
+
+			//initial overlap
+			EXPECT_TRUE(GJKRaycast<T>(A, Unscaled, TRigidTransform<T, 3>(TVector<T, 3>(8, 0, 0), TRotation<T, 3>::Identity), TVector<T, 3>(1, 0, 0), 4, Time, Position, Normal, 0, InitialDir));
+			EXPECT_FLOAT_EQ(Time, 0);
+			EXPECT_TRUE(GJKRaycast<T>(A, UniformScaled, TRigidTransform<T, 3>(TVector<T, 3>(6, 0, 0), TRotation<T, 3>::Identity), TVector<T, 3>(1, 0, 0), 4, Time, Position, Normal, 0, InitialDir));
+			EXPECT_FLOAT_EQ(Time, 0);
+			EXPECT_TRUE(GJKRaycast<T>(A, NonUniformScaled, TRigidTransform<T, 3>(TVector<T, 3>(6, 0, 0), TRotation<T, 3>::Identity), TVector<T, 3>(1, 0, 0), 4, Time, Position, Normal, 0, InitialDir));
+			EXPECT_FLOAT_EQ(Time, 0);
+
+			//miss
+			EXPECT_FALSE(GJKRaycast<T>(A, Unscaled, TRigidTransform<T, 3>(TVector<T, 3>(0, 0, 7.1), TRotation<T, 3>::Identity), TVector<T, 3>(1, 0, 0), 20, Time, Position, Normal, 0, InitialDir));
+			EXPECT_FALSE(GJKRaycast<T>(A, UniformScaled, TRigidTransform<T, 3>(TVector<T, 3>(0, 0, 9.1), TRotation<T, 3>::Identity), TVector<T, 3>(1, 0, 0), 20, Time, Position, Normal, 0, InitialDir));
+			EXPECT_FALSE(GJKRaycast<T>(A, NonUniformScaled, TRigidTransform<T, 3>(TVector<T, 3>(0, 0, 7.1), TRotation<T, 3>::Identity), TVector<T, 3>(1, 0, 0), 20, Time, Position, Normal, 0, InitialDir));
+
+			//hit with thickness
+			EXPECT_TRUE(GJKRaycast<T>(A, Unscaled, TRigidTransform<T, 3>(TVector<T, 3>(0, 0, 7.1), TRotation<T, 3>::Identity), TVector<T, 3>(1, 0, 0), 20, Time, Position, Normal, 0.2, InitialDir));
+			EXPECT_TRUE(GJKRaycast<T>(A, UniformScaled, TRigidTransform<T, 3>(TVector<T, 3>(0, 0, 9.1), TRotation<T, 3>::Identity), TVector<T, 3>(1, 0, 0), 20, Time, Position, Normal, 0.2, InitialDir));
+			EXPECT_TRUE(GJKRaycast<T>(A, NonUniformScaled, TRigidTransform<T, 3>(TVector<T, 3>(0, 0, 7.1), TRotation<T, 3>::Identity), TVector<T, 3>(1, 0, 0), 20, Time, Position, Normal, 0.2, InitialDir));
+
+			//hit rotated
+			const TRotation<T, 3> RotatedInPlace(TRotation<T, 3>::FromVector(TVector<T, 3>(0, PI * 0.5, 0)));
+			EXPECT_TRUE(GJKRaycast<T>(A, Unscaled, TRigidTransform<T, 3>(TVector<T, 3>(0, 0, 0), RotatedInPlace), TVector<T, 3>(1, 0, 0), 20, Time, Position, Normal, 0, InitialDir));
+			EXPECT_TRUE(GJKRaycast<T>(A, UniformScaled, TRigidTransform<T, 3>(TVector<T, 3>(0, 0, 0), RotatedInPlace), TVector<T, 3>(1, 0, 0), 20, Time, Position, Normal, 0, InitialDir));
+			EXPECT_TRUE(GJKRaycast<T>(A, NonUniformScaled, TRigidTransform<T, 3>(TVector<T, 3>(0, 0, 0), RotatedInPlace), TVector<T, 3>(1, 0, 0), 20, Time, Position, Normal, 0, InitialDir));
+
+			//miss rotated
+			EXPECT_FALSE(GJKRaycast<T>(A, Unscaled, TRigidTransform<T, 3>(TVector<T, 3>(0, 0, 7.1), RotatedInPlace), TVector<T, 3>(1, 0, 0), 20, Time, Position, Normal, 0, InitialDir));
+			EXPECT_FALSE(GJKRaycast<T>(A, UniformScaled, TRigidTransform<T, 3>(TVector<T, 3>(0, 0, 9.1), RotatedInPlace), TVector<T, 3>(1, 0, 0), 20, Time, Position, Normal, 0, InitialDir));
+			EXPECT_FALSE(GJKRaycast<T>(A, NonUniformScaled, TRigidTransform<T, 3>(TVector<T, 3>(0, 0, 9.1), RotatedInPlace), TVector<T, 3>(1, 0, 0), 20, Time, Position, Normal, 0, InitialDir));
+
+			//near hit
+			EXPECT_TRUE(GJKRaycast<T>(A, Unscaled, TRigidTransform<T, 3>(TVector<T, 3>(0, 0, 7 - 1e-2), TRotation<T, 3>::Identity), TVector<T, 3>(1, 0, 0), 20, Time, Position, Normal, 0, InitialDir));
+
+			//near miss
+			EXPECT_FALSE(GJKRaycast<T>(A, Unscaled, TRigidTransform<T, 3>(TVector<T, 3>(0, 0, 7 + 1e-2), TRotation<T, 3>::Identity), TVector<T, 3>(1, 0, 0), 20, Time, Position, Normal, 0, InitialDir));
+
+			//degenerate
+			TSphere<T, 3> Tiny(TVector<T, 3>(1, 0, 0), 1e-8);
+			EXPECT_TRUE(GJKRaycast<T>(A, Tiny, TRigidTransform<T, 3>::Identity, TVector<T, 3>(1, 0, 0), 8, Time, Position, Normal, 0, InitialDir));
+			EXPECT_NEAR(Time, 4, Eps);
+			EXPECT_VECTOR_NEAR(Normal, TVector3(-1, 0, 0), Eps);
+			EXPECT_VECTOR_NEAR(Position, TVector3(5, 0, 0), Eps);
+
+			//right at end
+			EXPECT_TRUE(GJKRaycast<T>(A, Unscaled, TRigidTransform<T, 3>::Identity, TVector<T, 3>(1, 0, 0), 3, Time, Position, Normal, 0, InitialDir));
+			EXPECT_NEAR(Time, 3, Eps);
+
+			// not far enough
+			EXPECT_FALSE(GJKRaycast<T>(A, Unscaled, TRigidTransform<T, 3>::Identity, TVector<T, 3>(1, 0, 0), 3 - 1e-2, Time, Position, Normal, 0, InitialDir));
+		}
+	}
+
+
+	template <typename T>
+	void GJKSphereTransformedSphereSweep()
+	{
+		typedef TVector<T, 3> TVector3;
+		TSphere<T, 3> A(TVector<T, 3>(10, 0, 0), 5);
+
+		TSphere<T, 3> Sphere(TVector<T, 3>(0), 2);
+		TSphere<T, 3> Translated(Sphere.GetCenter() + TVector<T, 3>(1, 0, 0), Sphere.GetRadius());
+		TSphere<T,3> Transformed(TRigidTransform<T, 3>(TVector<T, 3>(1, 0, 0), TRotation<T, 3>::FromVector(TVector<T, 3>(0, 0, PI))).TransformPosition(Sphere.GetCenter()), Sphere.GetRadius());
+
+		TVector<T, 3> InitialDirs[] = { TVector<T,3>(1,0,0), TVector<T,3>(-1,0,0), TVector<T,3>(0,1,0), TVector<T,3>(0,-1,0), TVector<T,3>(0,0,1), TVector<T,3>(0,0,-1) };
+
+		constexpr T Eps = 1e-1;
+
+		for (const TVector<T, 3>& InitialDir : InitialDirs)
+		{
+			T Time;
+			TVector<T, 3> Position;
+			TVector<T, 3> Normal;
+
+			//hit
+			EXPECT_TRUE(GJKRaycast<T>(A, Translated, TRigidTransform<T, 3>::Identity, TVector<T, 3>(1, 0, 0), 4, Time, Position, Normal, 0, InitialDir));
+			EXPECT_NEAR(Time, 2, Eps);
+			EXPECT_VECTOR_NEAR(Normal, TVector3(-1, 0, 0), Eps);
+			EXPECT_VECTOR_NEAR(Position, TVector3(5, 0, 0), Eps);
+			EXPECT_TRUE(GJKRaycast<T>(A, Transformed, TRigidTransform<T, 3>::Identity, TVector<T, 3>(1, 0, 0), 4, Time, Position, Normal, 0, InitialDir));
+			EXPECT_NEAR(Time, 2, Eps);
+			EXPECT_VECTOR_NEAR(Normal, TVector3(-1, 0, 0), Eps);
+			EXPECT_VECTOR_NEAR(Position, TVector3(5, 0, 0), Eps);
+
+			//hit offset
+			EXPECT_TRUE(GJKRaycast<T>(A, Translated, TRigidTransform<T, 3>(TVector<T, 3>(1, 0, 0), TRotation<T, 3>::Identity), TVector<T, 3>(1, 0, 0), 4, Time, Position, Normal, 0, InitialDir));
+			EXPECT_NEAR(Time, 1, Eps);
+			EXPECT_VECTOR_NEAR(Normal, TVector3(-1, 0, 0), Eps);
+			EXPECT_VECTOR_NEAR(Position, TVector3(5, 0, 0), Eps);
+			EXPECT_TRUE(GJKRaycast<T>(A, Transformed, TRigidTransform<T, 3>(TVector<T, 3>(1, 0, 0), TRotation<T, 3>::Identity), TVector<T, 3>(1, 0, 0), 4, Time, Position, Normal, 0, InitialDir));
+			EXPECT_NEAR(Time, 1, Eps);
+			EXPECT_VECTOR_NEAR(Normal, TVector3(-1, 0, 0), Eps);
+			EXPECT_VECTOR_NEAR(Position, TVector3(5, 0, 0), Eps);
+
+			//initial overlap
+			EXPECT_TRUE(GJKRaycast<T>(A, Translated, TRigidTransform<T, 3>(TVector<T, 3>(7, 0, 0), TRotation<T, 3>::Identity), TVector<T, 3>(1, 0, 0), 4, Time, Position, Normal, 0, InitialDir));
+			EXPECT_FLOAT_EQ(Time, 0);
+			EXPECT_TRUE(GJKRaycast<T>(A, Transformed, TRigidTransform<T, 3>(TVector<T, 3>(7, 0, 0), TRotation<T, 3>::Identity), TVector<T, 3>(1, 0, 0), 4, Time, Position, Normal, 0, InitialDir));
+			EXPECT_FLOAT_EQ(Time, 0);
+
+			//miss
+			EXPECT_FALSE(GJKRaycast<T>(A, Translated, TRigidTransform<T, 3>(TVector<T, 3>(0, 0, 7.1), TRotation<T, 3>::Identity), TVector<T, 3>(1, 0, 0), 20, Time, Position, Normal, 0, InitialDir));
+			EXPECT_FALSE(GJKRaycast<T>(A, Transformed, TRigidTransform<T, 3>(TVector<T, 3>(0, 0, 7.1), TRotation<T, 3>::Identity), TVector<T, 3>(1, 0, 0), 20, Time, Position, Normal, 0, InitialDir));
+
+			//hit with thickness
+			EXPECT_TRUE(GJKRaycast<T>(A, Translated, TRigidTransform<T, 3>(TVector<T, 3>(0, 0, 7.1), TRotation<T, 3>::Identity), TVector<T, 3>(1, 0, 0), 20, Time, Position, Normal, 0.2, InitialDir));
+			EXPECT_TRUE(GJKRaycast<T>(A, Transformed, TRigidTransform<T, 3>(TVector<T, 3>(0, 0, 7.1), TRotation<T, 3>::Identity), TVector<T, 3>(1, 0, 0), 20, Time, Position, Normal, 0.2, InitialDir));
+
+			//hit rotated
+			const TRotation<T, 3> RotatedDown(TRotation<T, 3>::FromVector(TVector<T, 3>(0, PI * 0.5, 0)));
+			EXPECT_TRUE(GJKRaycast<T>(A, Translated, TRigidTransform<T, 3>(TVector<T, 3>(0, 0, 7.9), RotatedDown), TVector<T, 3>(1, 0, 0), 20, Time, Position, Normal, 0, InitialDir));
+			EXPECT_TRUE(GJKRaycast<T>(A, Transformed, TRigidTransform<T, 3>(TVector<T, 3>(0, 0, 7.9), RotatedDown), TVector<T, 3>(1, 0, 0), 20, Time, Position, Normal, 0, InitialDir));
+
+			//miss rotated
+			EXPECT_FALSE(GJKRaycast<T>(A, Translated, TRigidTransform<T, 3>(TVector<T, 3>(0, 0, 8.1), RotatedDown), TVector<T, 3>(1, 0, 0), 20, Time, Position, Normal, 0, InitialDir));
+			EXPECT_FALSE(GJKRaycast<T>(A, Transformed, TRigidTransform<T, 3>(TVector<T, 3>(0, 0, 8.1), RotatedDown), TVector<T, 3>(1, 0, 0), 20, Time, Position, Normal, 0, InitialDir));
+
+			//hit rotated with inflation
+			EXPECT_TRUE(GJKRaycast<T>(A, Translated, TRigidTransform<T, 3>(TVector<T, 3>(0, 0, 7.9), RotatedDown), TVector<T, 3>(1, 0, 0), 20, Time, Position, Normal, 0.2, InitialDir));
+			EXPECT_TRUE(GJKRaycast<T>(A, Transformed, TRigidTransform<T, 3>(TVector<T, 3>(0, 0, 7.9), RotatedDown), TVector<T, 3>(1, 0, 0), 20, Time, Position, Normal, 0.2, InitialDir));
+
+			//near hit
+			EXPECT_TRUE(GJKRaycast<T>(A, Translated, TRigidTransform<T, 3>(TVector<T, 3>(0, 0, 7 - 1e-2), TRotation<T, 3>::Identity), TVector<T, 3>(1, 0, 0), 20, Time, Position, Normal, 0, InitialDir));
+			EXPECT_TRUE(GJKRaycast<T>(A, Transformed, TRigidTransform<T, 3>(TVector<T, 3>(0, 0, 7 - 1e-2), TRotation<T, 3>::Identity), TVector<T, 3>(1, 0, 0), 20, Time, Position, Normal, 0, InitialDir));
+
+			//near miss
+			EXPECT_FALSE(GJKRaycast<T>(A, Translated, TRigidTransform<T, 3>(TVector<T, 3>(0, 0, 7 + 1e-2), TRotation<T, 3>::Identity), TVector<T, 3>(1, 0, 0), 20, Time, Position, Normal, 0, InitialDir));
+			EXPECT_FALSE(GJKRaycast<T>(A, Transformed, TRigidTransform<T, 3>(TVector<T, 3>(0, 0, 7 + 1e-2), TRotation<T, 3>::Identity), TVector<T, 3>(1, 0, 0), 20, Time, Position, Normal, 0, InitialDir));
+
+			//right at end
+			EXPECT_TRUE(GJKRaycast<T>(A, Translated, TRigidTransform<T, 3>::Identity, TVector<T, 3>(1, 0, 0), 2, Time, Position, Normal, 0, InitialDir));
+			EXPECT_NEAR(Time, 2, Eps);
+			EXPECT_TRUE(GJKRaycast<T>(A, Transformed, TRigidTransform<T, 3>::Identity, TVector<T, 3>(1, 0, 0), 2, Time, Position, Normal, 0, InitialDir));
+			EXPECT_NEAR(Time, 2, Eps);
+
+			// not far enough
+			EXPECT_FALSE(GJKRaycast<T>(A, Translated, TRigidTransform<T, 3>::Identity, TVector<T, 3>(1, 0, 0), 2 - 1e-2, Time, Position, Normal, 0, InitialDir));
+			EXPECT_FALSE(GJKRaycast<T>(A, Transformed, TRigidTransform<T, 3>::Identity, TVector<T, 3>(1, 0, 0), 2 - 1e-2, Time, Position, Normal, 0, InitialDir));
+		}
+	}
+
+
+	template <typename T>
+	void GJKBoxCapsuleSweep()
+	{
+		TAABB<T, 3> A(TVector<T, 3>(3, -1, 0), TVector<T, 3>(4, 1, 4));
+		TCapsule<T> B(TVector<T, 3>(0, 0, -1), TVector<T, 3>(0, 0, 1), 2);
+
+		TVector<T, 3> InitialDirs[] = { TVector<T,3>(1,0,0), TVector<T,3>(-1,0,0), TVector<T,3>(0,1,0), TVector<T,3>(0,-1,0), TVector<T,3>(0,0,1), TVector<T,3>(0,0,-1) };
+
+		constexpr T Eps = 1e-1;
+
+		for (const TVector<T, 3>& InitialDir : InitialDirs)
+		{
+			T Time;
+			TVector<T, 3> Position;
+			TVector<T, 3> Normal;
+
+			//hit
+			EXPECT_TRUE(GJKRaycast<T>(A, B, TRigidTransform<T, 3>::Identity, TVector<T, 3>(1, 0, 0), 2, Time, Position, Normal, 0, InitialDir));
+			EXPECT_NEAR(Time, 1, Eps);
+			EXPECT_NEAR(Normal.X, -1, Eps);
+			EXPECT_NEAR(Normal.Y, 0, Eps);
+			EXPECT_NEAR(Normal.Z, 0, Eps);
+			EXPECT_NEAR(Position.X, 3, Eps);
+			//EXPECT_NEAR(Position.Y, 0, Eps);	//todo: look into inaccuracy here (0.015) instead of <1e-2
+			EXPECT_LE(Position.Z, 1 + Eps);
+			EXPECT_GE(Position.Z, -1 - Eps);
+
+			//hit offset
+			EXPECT_TRUE(GJKRaycast<T>(A, B, TRigidTransform<T, 3>(TVector<T, 3>(0.5, 0, 0), TRotation<T, 3>::Identity), TVector<T, 3>(1, 0, 0), 4, Time, Position, Normal, 0, InitialDir));
+			EXPECT_NEAR(Time, 0.5, Eps);
+			EXPECT_NEAR(Normal.X, -1, Eps);
+			EXPECT_NEAR(Normal.Y, 0, Eps);
+			EXPECT_NEAR(Normal.Z, 0, Eps);
+			EXPECT_NEAR(Position.X, 3, Eps);
+			//EXPECT_NEAR(Position.Y, 0, Eps);	//todo: look into inaccuracy here (0.015) instead of <1e-2
+			EXPECT_LE(Position.Z, 1 + Eps);
+			EXPECT_GE(Position.Z, -1 - Eps);
+
+			//initial overlap
+			EXPECT_TRUE(GJKRaycast2<T>(A, B, TRigidTransform<T, 3>(TVector<T, 3>(3, 0, 0), TRotation<T, 3>::Identity), TVector<T, 3>(1, 0, 0), 2, Time, Position, Normal, 0, false, InitialDir));
+			EXPECT_FLOAT_EQ(Time, 0);
+
+			//MTD
+			EXPECT_TRUE(GJKRaycast2<T>(A, B, TRigidTransform<T, 3>(TVector<T, 3>(2.5, 0, 0), TRotation<T, 3>::Identity), TVector<T, 3>(1, 0, 0), 2, Time, Position, Normal, 0, true, InitialDir));
+			EXPECT_FLOAT_EQ(Time, -1.5);
+			EXPECT_NEAR(Position[0], 3, Eps);	//many possible, but x must be on 3
+			EXPECT_VECTOR_NEAR(Normal, TVec3<T>(-1, 0, 0), Eps);
+
+			//MTD
+			T Penetration;
+			TVec3<T> ClosestA, ClosestB;
+			EXPECT_TRUE(GJKPenetration<T>(A, B, TRigidTransform<T, 3>(TVector<T, 3>(2.5, 0, 0), TRotation<T, 3>::Identity), Penetration, ClosestA, ClosestB, Normal, 0, InitialDir));
+			EXPECT_FLOAT_EQ(Penetration, 1.5);
+			EXPECT_VECTOR_NEAR(Normal, TVec3<T>(-1, 0, 0), Eps);
+			EXPECT_NEAR(ClosestA[0], 3, Eps);	//could be any point on face, but should have x == 3
+			EXPECT_NEAR(ClosestB[0], 4.5, Eps);
+			EXPECT_NEAR(ClosestB[1], 0, Eps);
+
+			//EPA
+			EXPECT_TRUE(GJKRaycast2<T>(A, B, TRigidTransform<T, 3>(TVector<T, 3>(3, 0, 0), TRotation<T, 3>::Identity), TVector<T, 3>(1, 0, 0), 2, Time, Position, Normal, 0, true, InitialDir));
+			EXPECT_FLOAT_EQ(Time, -2);
+			EXPECT_NEAR(Position[0], 3, Eps);	//many possible, but x must be on 3
+			EXPECT_VECTOR_NEAR(Normal, TVec3<T>(-1, 0, 0), Eps);
+
+			//EPA
+			EXPECT_TRUE(GJKPenetration<T>(A, B, TRigidTransform<T, 3>(TVector<T, 3>(3, 0, 0), TRotation<T, 3>::Identity), Penetration, ClosestA, ClosestB, Normal, 0, InitialDir));
+			EXPECT_NEAR(Penetration, 2, Eps);
+			EXPECT_VECTOR_NEAR(Normal, TVec3<T>(-1, 0, 0), Eps);
+			EXPECT_NEAR(ClosestA[0], 3, Eps);	//could be any point on face, but should have x == 3
+			EXPECT_NEAR(ClosestB[0], 5, Eps);
+			EXPECT_NEAR(ClosestB[1], 0, Eps);
+
+			//EPA
+			EXPECT_TRUE(GJKRaycast2<T>(A, B, TRigidTransform<T, 3>(TVector<T, 3>(3.25, 0, 0), TRotation<T, 3>::Identity), TVector<T, 3>(1, 0, 0), 2, Time, Position, Normal, 0, true, InitialDir));
+			EXPECT_FLOAT_EQ(Time, -2.25);
+			EXPECT_NEAR(Position[0], 3, Eps);	//many possible, but x must be on 3
+			EXPECT_VECTOR_NEAR(Normal, TVec3<T>(-1, 0, 0), Eps);
+
+			//EPA
+			EXPECT_TRUE(GJKPenetration<T>(A, B, TRigidTransform<T, 3>(TVector<T, 3>(3.25, 0, 0), TRotation<T, 3>::Identity), Penetration, ClosestA, ClosestB, Normal, 0, InitialDir));
+			EXPECT_NEAR(Penetration, 2.25, Eps);
+			EXPECT_VECTOR_NEAR(Normal, TVec3<T>(-1, 0, 0), Eps);
+			EXPECT_NEAR(ClosestA[0], 3, Eps);	//could be any point on face, but should have x == 3
+			EXPECT_NEAR(ClosestB[0], 5.25, Eps);
+			EXPECT_NEAR(ClosestB[1], 0, Eps);
+
+			//MTD
+			EXPECT_TRUE(GJKRaycast2<T>(A, B, TRigidTransform<T, 3>(TVector<T, 3>(3.25, 0, -2.875), TRotation<T, 3>::Identity), TVector<T, 3>(1, 0, 0), 2, Time, Position, Normal, 0, true, InitialDir));
+			EXPECT_FLOAT_EQ(Time, -0.125);
+			EXPECT_VECTOR_NEAR(Position, TVec3<T>(3.25, 0, 0), Eps);
+			EXPECT_VECTOR_NEAR(Normal, TVec3<T>(0, 0, -1), Eps);
+
+			//MTD
+			EXPECT_TRUE(GJKPenetration<T>(A, B, TRigidTransform<T, 3>(TVector<T, 3>(3.25, 0, -2.875), TRotation<T, 3>::Identity), Penetration, ClosestA, ClosestB, Normal, 0, InitialDir));
+			EXPECT_NEAR(Penetration, 0.125, Eps);
+			EXPECT_VECTOR_NEAR(Normal, TVec3<T>(0, 0, -1), Eps);
+			EXPECT_VECTOR_NEAR(ClosestA, TVec3<T>(3.25, 0, 0), Eps);
+			EXPECT_VECTOR_NEAR(ClosestB, TVec3<T>(3.25, 0, 0.125), Eps);
+
+			//near miss
+			EXPECT_FALSE(GJKRaycast<T>(A, B, TRigidTransform<T, 3>(TVector<T, 3>(0, 0, 7 + 1e-2), TRotation<T, 3>::Identity), TVector<T, 3>(1, 0, 0), 4, Time, Position, Normal, 0, InitialDir));
+
+			//near hit
+			EXPECT_TRUE(GJKRaycast<T>(A, B, TRigidTransform<T, 3>(TVector<T, 3>(0, 0, 7 - 1e-2), TRotation<T, 3>::Identity), TVector<T, 3>(1, 0, 0), 4, Time, Position, Normal, 0, InitialDir));
+			EXPECT_NEAR(Position.X, 3, Eps);
+			EXPECT_NEAR(Position.Z, 4, 10 * Eps);
+
+			//near hit inflation
+			EXPECT_TRUE(GJKRaycast<T>(A, B, TRigidTransform<T, 3>(TVector<T, 3>(0, 0, 7 - 1e-2), TRotation<T, 3>::Identity), TVector<T, 3>(1, 0, 0), 4, Time, Position, Normal, 2e-2, InitialDir));
+			EXPECT_NEAR(Position.X, 3, Eps);
+			EXPECT_NEAR(Position.Z, 4, 10 * Eps);
+
+			//rotation hit
+			TRotation<T, 3> Rotated(TRotation<T, 3>::FromVector(TVector<T, 3>(0, -PI * 0.5, 0)));
+			EXPECT_TRUE(GJKRaycast<T>(A, B, TRigidTransform<T, 3>(TVector<T, 3>(-0.5, 0, 0), Rotated), TVector<T, 3>(1, 0, 0), 1, Time, Position, Normal, 0, InitialDir));
+			EXPECT_NEAR(Time, 0.5, Eps);
+			EXPECT_NEAR(Position.X, 3, Eps);
+			EXPECT_NEAR(Normal.X, -1, Eps);
+			EXPECT_NEAR(Normal.Y, 0, Eps);
+			EXPECT_NEAR(Normal.Z, 0, Eps);
+
+			//rotation near hit
+			EXPECT_TRUE(GJKRaycast<T>(A, B, TRigidTransform<T, 3>(TVector<T, 3>(0, 0, 6 - 1e-2), Rotated), TVector<T, 3>(1, 0, 0), 10, Time, Position, Normal, 0, InitialDir));
+
+			//rotation near miss
+			EXPECT_FALSE(GJKRaycast<T>(A, B, TRigidTransform<T, 3>(TVector<T, 3>(0, 0, 6 + 1e-2), Rotated), TVector<T, 3>(1, 0, 0), 10, Time, Position, Normal, 0, InitialDir));
+
+			//degenerate capsule
+			TCapsule<T> Needle(TVector<T, 3>(0, 0, -1), TVector<T, 3>(0, 0, 1), 1e-8);
+			EXPECT_TRUE(GJKRaycast<T>(A, Needle, TRigidTransform<T, 3>::Identity, TVector<T, 3>(1, 0, 0), 6, Time, Position, Normal, 0, InitialDir));
+			EXPECT_NEAR(Time, 3, Eps);
+			EXPECT_NEAR(Normal.X, -1, Eps);
+			EXPECT_NEAR(Normal.Y, 0, Eps);
+			EXPECT_NEAR(Normal.Z, 0, Eps);
+			EXPECT_NEAR(Position.X, 3, Eps);
+			//EXPECT_NEAR(Position.Y, 0, Eps);	//todo: look into inaccuracy here (0.015) instead of <1e-2
+			EXPECT_LE(Position.Z, 1 + Eps);
+			EXPECT_GE(Position.Z, -1 - Eps);
+		}
+	}
+
+	template <typename T>
+	void GJKBoxBoxSweep()
+	{
+		{
+			//based on real sweep from game
+			const TAABB<T, 3> A(TVec3<T>(-2560.00000, -268.000031, -768.000122), TVec3<T>(0.000000000, 3.99996948, 0.000000000));
+			const TAABB<T, 3> B(TVec3<T>(-248.000000, -248.000000, -9.99999975e-05), TVec3<T>(248.000000, 248.000000, 9.99999975e-05));
+			const TRigidTransform<T, 3> BToATM(TVec3<T>(-2559.99780, -511.729492, -8.98901367), TRotation<T, 3>::FromElements(1.51728955e-06, 1.51728318e-06, 0.707108259, 0.707105279));
+			const TVec3<T> LocalDir(-4.29153351e-06, 0.000000000, -1.00000000);
+			const T Length = 393.000000;
+			const TVec3<T> SearchDir(511.718750, -2560.00000, 9.00000000);
+
+			T Time;
+			TVec3<T> Pos, Normal;
+			GJKRaycast2<T>(A, B, BToATM, LocalDir, Length, Time, Pos, Normal, 0, true, SearchDir, 0);
+		}
+
+		{
+			//based on real sweep from game
+			TParticles<T, 3> ConvexParticles;
+			ConvexParticles.AddParticles(10);
+
+			ConvexParticles.X(0) = { 51870.2305, 54369.6719, 19200.0000 };
+			ConvexParticles.X(1) = { -91008.5625, -59964.0000, -19199.9629 };
+			ConvexParticles.X(2) = { 51870.2305, 54369.6758, -19199.9668 };
+			ConvexParticles.X(3) = { 22164.4883, 124647.500, -19199.9961 };
+			ConvexParticles.X(4) = { 34478.5000, 123975.492, -19199.9961 };
+			ConvexParticles.X(5) = { -91008.5000, -59963.9375, 19200.0000 };
+			ConvexParticles.X(6) = { -91008.5000, 33715.5625, 19200.0000 };
+			ConvexParticles.X(7) = { 34478.4961, 123975.500, 19200.0000 };
+			ConvexParticles.X(8) = { 22164.4922, 124647.500, 19200.0000 };
+			ConvexParticles.X(9) = { -91008.5000, 33715.5625, -19199.9961 };
+
+
+			const FConvex A(ConvexParticles);
+			const TAABB<T, 3> B(TVec3<T>{ -6.00000000, -248.000000, -9.99999975e-05 }, TVec3<T>{ 6.00000000, 248.000000, 9.99999975e-05 });
+			const TRigidTransform<T, 3> BToATM(TVec3<T>{33470.5000, 41570.5000, -1161.00000}, TRotation<T, 3>::FromIdentity());
+			const TVec3<T> LocalDir(0, 0, -1);
+			const T Length = 393.000000;
+			const TVec3<T> SearchDir{ -33470.5000, -41570.5000, 1161.00000 };
+
+			T Time;
+			TVec3<T> Pos, Normal;
+			GJKRaycast2<T>(A, B, BToATM, LocalDir, Length, Time, Pos, Normal, 0, true, SearchDir, 0);
+		}
+	}
+
+	template <typename T>
+	void GJKCapsuleConvexInitialOverlapSweep()
+	{
+		TParticles<T, 3> ConvexParticles;
+		ConvexParticles.AddParticles(8);
+
+		ConvexParticles.X(0) = {-256.000031, 12.0000601, 384.000061};
+		ConvexParticles.X(1) = {256.000031, 12.0000601, 384.000061};
+		ConvexParticles.X(2) = {256.000031, 12.0000601, 6.10351563e-05};
+		ConvexParticles.X(3) = {-256.000031, -11.9999399, 6.10351563e-05};
+		ConvexParticles.X(4) = {-256.000031, 12.0000601, 6.10351563e-05};
+		ConvexParticles.X(5) = {-256.000031, -11.9999399, 384.000061};
+		ConvexParticles.X(6) = {256.000031, -11.9999399, 6.10351563e-05};
+		ConvexParticles.X(7) = {256.000031, -11.9999399, 384.000061};
+
+		TUniquePtr<FConvex> UniqueConvex = MakeUnique<FConvex>(ConvexParticles);
+		TSerializablePtr<FConvex> AConv(UniqueConvex);
+		const TImplicitObjectScaled<FConvex> A(AConv, TVec3<T>(1.0, 1.0, 1.0));
+
+		const TVec3<T> Pt0(0.0, 0.0, -33.0);
+		TVec3<T> Pt1 = Pt0;
+		Pt1 += (TVec3<T>(0.0, 0.0, 1.0) * 66.0);
+
+		const TCapsule<T> B(Pt0, Pt1, 42.0);
+
+		const TRigidTransform<T, 3> BToATM(TVec3<T>(157.314758, -54.0000839, 76.1436157), TRotation<T, 3>::FromElements(0.0, 0.0, 0.704960823, 0.709246278));
+		const TVec3<T> LocalDir(-0.00641351938, -0.999979556, 0.0);
+		const T Length = 0.0886496082;
+		const TVec3<T> SearchDir(-3.06152344, 166.296631, -76.1436157);
+
+		T Time;
+		TVec3<T> Position, Normal;
+		EXPECT_TRUE(GJKRaycast2<T>(A, B, BToATM, LocalDir, Length, Time, Position, Normal, 0, true, SearchDir, 0));
+		EXPECT_FLOAT_EQ(Time, 0.0);
+	}
+
+	template void SimplexLine<float>();
+	template void SimplexTriangle<float>();
+	template void SimplexTetrahedron<float>();
+	template void GJKSphereSphereTest<float>();
+	template void GJKSphereBoxTest<float>();
+	template void GJKSphereCapsuleTest<float>();
+	template void GJKSphereConvexTest<float>();
+	template void GJKSphereScaledSphereTest<float>();
+	
+	template void GJKSphereSphereSweep<float>();
+	template void GJKSphereBoxSweep<float>();
+	template void GJKSphereCapsuleSweep<float>();
+	template void GJKSphereConvexSweep<float>();
+	template void GJKSphereScaledSphereSweep<float>();
+	template void GJKSphereTransformedSphereSweep<float>();
+	template void GJKBoxCapsuleSweep<float>();
+	template void GJKBoxBoxSweep<float>();
+	template void GJKCapsuleConvexInitialOverlapSweep<float>();
+}
\ No newline at end of file
