Rename/move HeadlessChaos.

diff --git a/Engine/Source/Programs/HeadlessChaos/Private/HeadlessChaosTestGJKDistance.cpp b/Engine/Source/Programs/HeadlessChaos/Private/HeadlessChaosTestGJKDistance.cpp
new file mode 100644
--- /dev/null
+++ b/Engine/Source/Programs/HeadlessChaos/Private/HeadlessChaosTestGJKDistance.cpp
@@ -0,0 +1,379 @@
+// Copyright Epic Games, Inc. All Rights Reserved.
+
+#include "HeadlessChaosTestGJK.h"
+#include "HeadlessChaos.h"
+
+#include "Chaos/Capsule.h"
+#include "Chaos/Convex.h"
+#include "Chaos/GJK.h"
+#include "Chaos/ImplicitObjectScaled.h"
+
+namespace ChaosTest
+{
+	using namespace Chaos;
+
+
+
+	template <typename T>
+	void GJKSphereSphereDistanceTest()
+	{
+		const T Tolerance = (T)1e-3;
+
+		TVector<T, 3> NearestA = { 0,0,0 };
+		TVector<T, 3> NearestB = { 0,0,0 };
+		T Distance = 0;
+		// Fail - overlapping
+		{
+			TSphere<T, 3> A(TVector<T, 3>(12, 0, 0), 5);
+			TSphere<T, 3> B(TVector<T, 3>(4, 0, 0), 2);
+			bool bSuccess = GJKDistance<T>(A, B, TRigidTransform<T, 3>(TVector<T, 3>(2, 0, 0), TRotation<T, 3>::FromIdentity()), Distance, NearestA, NearestB);
+			EXPECT_FALSE(bSuccess);
+		}
+
+		// Success - not overlapping
+		{
+			TSphere<T, 3> A(TVector<T, 3>(12, 0, 0), 5);
+			TSphere<T, 3> B(TVector<T, 3>(4, 0, 0), 2);
+			bool bSuccess = GJKDistance<T>(A, B, TRigidTransform<T, 3>::Identity, Distance, NearestA, NearestB);
+			EXPECT_TRUE(bSuccess);
+			EXPECT_NEAR(Distance, (T)1, Tolerance);
+			EXPECT_NEAR(NearestA.X, (T)7, Tolerance);
+			EXPECT_NEAR(NearestA.Y, (T)0, Tolerance);
+			EXPECT_NEAR(NearestA.Z, (T)0, Tolerance);
+			EXPECT_NEAR(NearestB.X, (T)6, Tolerance);
+			EXPECT_NEAR(NearestB.Y, (T)0, Tolerance);
+			EXPECT_NEAR(NearestB.Z, (T)0, Tolerance);
+		}
+
+		// Success - not overlapping
+		{
+			TSphere<T, 3> A(TVector<T, 3>(0, 0, 0), 2);
+			TSphere<T, 3> B(TVector<T, 3>(0, 0, 0), 2);
+			TVector<T, 3> BPos = TVector<T, 3>(3, 3, 0);
+			bool bSuccess = GJKDistance<T>(A, B, TRigidTransform<T, 3>(BPos, TRotation<T, 3>::FromIdentity()), Distance, NearestA, NearestB);
+			EXPECT_TRUE(bSuccess);
+			TVector<T, 3> CenterDelta = (B.GetCenter() + BPos) - A.GetCenter();
+			TVector<T, 3> CenterDir = CenterDelta.GetSafeNormal();
+			EXPECT_NEAR(Distance, CenterDelta.Size() - (A.GetRadius() + B.GetRadius()), Tolerance);
+			EXPECT_NEAR(NearestA.X, A.GetCenter().X + A.GetRadius() * CenterDir.X, Tolerance);
+			EXPECT_NEAR(NearestA.Y, A.GetCenter().Y + A.GetRadius() * CenterDir.Y, Tolerance);
+			EXPECT_NEAR(NearestA.Z, A.GetCenter().Z + A.GetRadius() * CenterDir.Z, Tolerance);
+			EXPECT_NEAR(NearestB.X, B.GetCenter().X - B.GetRadius() * CenterDir.X, Tolerance);
+			EXPECT_NEAR(NearestB.Y, B.GetCenter().Y - B.GetRadius() * CenterDir.Y, Tolerance);
+			EXPECT_NEAR(NearestB.Z, B.GetCenter().Z - B.GetRadius() * CenterDir.Z, Tolerance);
+		}
+
+		// Success - very close not overlapping
+		{
+			TSphere<T, 3> A(TVector<T, 3>(12, 0, 0), 5);
+			TSphere<T, 3> B(TVector<T, 3>(4, 0, 0), 2);
+			TVector<T, 3> BPos = TVector<T, 3>(0.99, 0, 0);
+			bool bSuccess = GJKDistance<T>(A, B, TRigidTransform<T, 3>(BPos, TRotation<T, 3>::FromIdentity()), Distance, NearestA, NearestB);
+			EXPECT_TRUE(bSuccess);
+			EXPECT_NEAR(Distance, (T)1 - BPos.X, Tolerance);
+			EXPECT_NEAR(NearestA.X, (T)7, Tolerance);
+			EXPECT_NEAR(NearestA.Y, (T)0, Tolerance);
+			EXPECT_NEAR(NearestA.Z, (T)0, Tolerance);
+			EXPECT_NEAR(NearestB.X, (T)6, Tolerance);
+			EXPECT_NEAR(NearestB.Y, (T)0, Tolerance);
+			EXPECT_NEAR(NearestB.Z, (T)0, Tolerance);
+		}
+	}
+
+	TEST(TestGJKDistance, GJKSphereSphereDistanceTest)
+	{
+		GJKSphereSphereDistanceTest<float>();
+	}
+
+	template <typename T>
+	void GJKBoxSphereDistanceTest()
+	{
+		const T Tolerance = (T)2e-3;
+
+		TVector<T, 3> NearestA = { 0,0,0 };
+		TVector<T, 3> NearestB = { 0,0,0 };
+		T Distance = 0;
+
+		// Fail - overlapping
+		{
+			TAABB<T, 3> A(TVector<T, 3>(5, -2, -2), TVector<T, 3>(8, 2, 2));
+			TSphere<T, 3> B(TVector<T, 3>(2, 0, 0), 2);
+			bool bSuccess = GJKDistance<T>(A, B, TRigidTransform<T, 3>(TVector<T, 3>(2, 0, 0), TRotation<T, 3>::FromIdentity()), Distance, NearestA, NearestB);
+			EXPECT_FALSE(bSuccess);
+		}
+
+		// Success - not overlapping - mid-face near point
+		{
+			TAABB<T, 3> A(TVector<T, 3>(5, -2, -2), TVector<T, 3>(8, 2, 2));
+			TSphere<T, 3> B(TVector<T, 3>(2, 0, 0), 2);
+			bool bSuccess = GJKDistance<T>(A, B, TRigidTransform<T, 3>::Identity, Distance, NearestA, NearestB);
+			EXPECT_TRUE(bSuccess);
+			EXPECT_NEAR(Distance, (T)1, Tolerance);
+			EXPECT_NEAR(NearestA.X, (T)5, Tolerance);
+			EXPECT_NEAR(NearestA.Y, (T)0, Tolerance);
+			EXPECT_NEAR(NearestA.Z, (T)0, Tolerance);
+			EXPECT_NEAR(NearestB.X, (T)4, Tolerance);
+			EXPECT_NEAR(NearestB.Y, (T)0, Tolerance);
+			EXPECT_NEAR(NearestB.Z, (T)0, Tolerance);
+		}
+		// Other way round
+		{
+			TAABB<T, 3> A(TVector<T, 3>(5, -2, -2), TVector<T, 3>(8, 2, 2));
+			TSphere<T, 3> B(TVector<T, 3>(2, 0, 0), 2);
+			bool bSuccess = GJKDistance<T>(B, A, TRigidTransform<T, 3>::Identity, Distance, NearestB, NearestA);
+			EXPECT_TRUE(bSuccess);
+			EXPECT_NEAR(Distance, (T)1, Tolerance);
+			EXPECT_NEAR(NearestA.X, (T)5, Tolerance);
+			EXPECT_NEAR(NearestA.Y, (T)0, Tolerance);
+			EXPECT_NEAR(NearestA.Z, (T)0, Tolerance);
+			EXPECT_NEAR(NearestB.X, (T)4, Tolerance);
+			EXPECT_NEAR(NearestB.Y, (T)0, Tolerance);
+			EXPECT_NEAR(NearestB.Z, (T)0, Tolerance);
+		}
+
+		// Success - not overlapping - vertex near point
+		{
+			TAABB<T, 3> A(TVector<T, 3>(5, 2, 2), TVector<T, 3>(8, 4, 4));
+			TSphere<T, 3> B(TVector<T, 3>(2, 0, 0), 2);
+			bool bSuccess = GJKDistance<T>(A, B, TRigidTransform<T, 3>::Identity, Distance, NearestA, NearestB);
+			TVector<T, 3> NearPointOnA = A.Min();
+			TVector<T, 3> SphereNearPointDir = (NearPointOnA - B.GetCenter()).GetSafeNormal();
+			TVector<T, 3> NearPointOnB = B.GetCenter() + SphereNearPointDir * B.GetRadius();
+			EXPECT_TRUE(bSuccess);
+			EXPECT_NEAR(Distance, (NearPointOnA - NearPointOnB).Size(), Tolerance);
+			EXPECT_NEAR(NearestA.X, NearPointOnA.X, Tolerance);
+			EXPECT_NEAR(NearestA.Y, NearPointOnA.Y, Tolerance);
+			EXPECT_NEAR(NearestA.Z, NearPointOnA.Z, Tolerance);
+			EXPECT_NEAR(NearestB.X, NearPointOnB.X, Tolerance);
+			EXPECT_NEAR(NearestB.Y, NearPointOnB.Y, Tolerance);
+			EXPECT_NEAR(NearestB.Z, NearPointOnB.Z, Tolerance);
+		}
+		// Other way round
+		{
+			TAABB<T, 3> A(TVector<T, 3>(5, 2, 2), TVector<T, 3>(8, 4, 4));
+			TSphere<T, 3> B(TVector<T, 3>(2, 0, 0), 2);
+			bool bSuccess = GJKDistance<T>(B, A, TRigidTransform<T, 3>::Identity, Distance, NearestB, NearestA);
+			TVector<T, 3> NearPointOnA = A.Min();
+			TVector<T, 3> SphereNearPointDir = (NearPointOnA - B.GetCenter()).GetSafeNormal();
+			TVector<T, 3> NearPointOnB = B.GetCenter() + SphereNearPointDir * B.GetRadius();
+			EXPECT_TRUE(bSuccess);
+			EXPECT_NEAR(Distance, (NearPointOnA - NearPointOnB).Size(), Tolerance);
+			EXPECT_NEAR(NearestA.X, NearPointOnA.X, Tolerance);
+			EXPECT_NEAR(NearestA.Y, NearPointOnA.Y, Tolerance);
+			EXPECT_NEAR(NearestA.Z, NearPointOnA.Z, Tolerance);
+			EXPECT_NEAR(NearestB.X, NearPointOnB.X, Tolerance);
+			EXPECT_NEAR(NearestB.Y, NearPointOnB.Y, Tolerance);
+			EXPECT_NEAR(NearestB.Z, NearPointOnB.Z, Tolerance);
+		}
+
+		// Rotated
+		{
+			TAABB<T, 3> A(TVector<T, 3>(-2, -2, -2), TVector<T, 3>(4, 4, 4));
+			TSphere<T, 3> B(TVector<T, 3>(0, 0, 0), 2);
+			TRigidTransform<T, 3> BToATm = TRigidTransform<T, 3>(TVector<T, 3>(8, 0, 0), TRotation<T, 3>::FromAxisAngle(TVector<T, 3>(0, 1, 0), FMath::DegreesToRadians(45)));	// Rotation won't affect sphere
+			bool bSuccess = GJKDistance<T>(A, B, BToATm, Distance, NearestA, NearestB);
+			TVector<T, 3> NearPointOnA = TVector<T, 3>(4, 0, 0);
+			TVector<T, 3> BPos = BToATm.TransformPositionNoScale(B.GetCenter());
+			TVector<T, 3> NearPointDir = (NearPointOnA - BPos).GetSafeNormal();
+			TVector<T, 3> NearPointOnB = BPos + NearPointDir * B.GetRadius();
+			TVector<T, 3> NearPointOnBLocal = BToATm.InverseTransformPositionNoScale(NearPointOnB);
+			EXPECT_TRUE(bSuccess);
+			EXPECT_NEAR(Distance, (NearPointOnA - NearPointOnB).Size(), Tolerance);
+			EXPECT_NEAR(NearestA.X, NearPointOnA.X, Tolerance);
+			EXPECT_NEAR(NearestA.Y, NearPointOnA.Y, Tolerance);
+			EXPECT_NEAR(NearestA.Z, NearPointOnA.Z, Tolerance);
+			EXPECT_NEAR(NearestB.X, NearPointOnBLocal.X, Tolerance);
+			EXPECT_NEAR(NearestB.Y, NearPointOnBLocal.Y, Tolerance);
+			EXPECT_NEAR(NearestB.Z, NearPointOnBLocal.Z, Tolerance);
+		}
+		// Other way round
+		{
+			TAABB<T, 3> A(TVector<T, 3>(-2, -2, -2), TVector<T, 3>(4, 4, 4));
+			TSphere<T, 3> B(TVector<T, 3>(0, 0, 0), 2);
+			TRigidTransform<T, 3> BToATm = TRigidTransform<T, 3>(TVector<T, 3>(-8, 0, 0), TRotation<T, 3>::FromAxisAngle(TVector<T, 3>(0, 1, 0), FMath::DegreesToRadians(45)));	// Rotation will affect box
+			bool bSuccess = GJKDistance<T>(B, A, BToATm, Distance, NearestB, NearestA);
+			TVector<T, 3> NearPointOnA = TVector<T, 3>(4, 0, 4);
+			TVector<T, 3> BPos = BToATm.InverseTransformPositionNoScale(B.GetCenter());
+			TVector<T, 3> NearPointDir = (NearPointOnA - BPos).GetSafeNormal();
+			TVector<T, 3> NearPointOnB = BPos + NearPointDir * B.GetRadius();
+			TVector<T, 3> NearPointOnBLocal = BToATm.TransformPositionNoScale(NearPointOnB);
+			EXPECT_TRUE(bSuccess);
+			EXPECT_NEAR(Distance, (NearPointOnA - NearPointOnB).Size(), Tolerance);
+			EXPECT_NEAR(NearestA.X, NearPointOnA.X, Tolerance);
+			EXPECT_NEAR(NearestA.Y, NearPointOnA.Y, Tolerance);
+			EXPECT_NEAR(NearestA.Z, NearPointOnA.Z, Tolerance);
+			EXPECT_NEAR(NearestB.X, NearPointOnBLocal.X, Tolerance);
+			EXPECT_NEAR(NearestB.Y, NearPointOnBLocal.Y, Tolerance);
+			EXPECT_NEAR(NearestB.Z, NearPointOnBLocal.Z, Tolerance);
+		}
+	
+		// Success - specific test case that initially failed (using incorrect initialization of V which works for Overlap but not Distance)
+		{
+			TAABB<T, 3> A(TVector<T, 3>(5, -2, 2), TVector<T, 3>(8, 2, 4));
+			TSphere<T, 3> B(TVector<T, 3>(2, 0, 0), 2);
+
+			bool bOverlap = GJKIntersection<T>(A, B, TRigidTransform<T, 3>::Identity);
+			EXPECT_FALSE(bOverlap);
+
+			bool bSuccess = GJKDistance<T>(A, B, TRigidTransform<T, 3>::Identity, Distance, NearestA, NearestB);
+			TVector<T, 3> NearPointOnA = TVector<T, 3>(5, 0, 2);
+			TVector<T, 3> NearPointDir = (NearPointOnA - B.GetCenter()).GetSafeNormal();
+			TVector<T, 3> NearPointOnB = B.GetCenter() + NearPointDir * B.GetRadius();
+			EXPECT_TRUE(bSuccess);
+			EXPECT_NEAR(Distance, (NearPointOnA - NearPointOnB).Size(), Tolerance);
+			EXPECT_NEAR(NearestA.X, NearPointOnA.X, Tolerance);
+			EXPECT_NEAR(NearestA.Y, NearPointOnA.Y, Tolerance);
+			EXPECT_NEAR(NearestA.Z, NearPointOnA.Z, Tolerance);
+			EXPECT_NEAR(NearestB.X, NearPointOnB.X, Tolerance);
+			EXPECT_NEAR(NearestB.Y, NearPointOnB.Y, Tolerance);
+			EXPECT_NEAR(NearestB.Z, NearPointOnB.Z, Tolerance);
+		}
+	}
+
+	TEST(TestGJKDistance, GJKBoxSphereDistanceTest)
+	{
+		GJKBoxSphereDistanceTest<float>();
+	}
+
+	template <typename T>
+	void GJKBoxCapsuleDistanceTest()
+	{
+		TVector<T, 3> NearestA = { 0,0,0 };
+		TVector<T, 3> NearestB = { 0,0,0 };
+		T Distance = 0;
+
+		// Fail - overlapping
+		{
+			TAABB<T, 3> A(TVector<T, 3>(5, -2, -2), TVector<T, 3>(8, 2, 2));
+			TCapsule<T> B(TVector<T, 3>(2, -2, 0), TVector<T, 3>(2, 2, 0), 2);
+			bool bSuccess = GJKDistance<T>(A, B, TRigidTransform<T, 3>(TVector<T, 3>(2, 0, 0), TRotation<T, 3>::FromIdentity()), Distance, NearestA, NearestB);
+			EXPECT_FALSE(bSuccess);
+		}
+
+		// Success - not overlapping, capsule axis parallel to nearest face (near points on cylinder and box face)
+		{
+			TAABB<T, 3> A(TVector<T, 3>(5, -2, -2), TVector<T, 3>(8, 2, 2));
+			TCapsule<T> B(TVector<T, 3>(2, 0, -1), TVector<T, 3>(2, 0, 2), 2);
+			bool bSuccess = GJKDistance<T>(A, B, TRigidTransform<T, 3>::Identity, Distance, NearestA, NearestB);
+
+			const T Tolerance = (T)2e-3;
+			EXPECT_TRUE(bSuccess);
+			EXPECT_NEAR(Distance, (T)1, Tolerance);
+			EXPECT_NEAR(NearestA.X, (T)5, Tolerance);
+			EXPECT_NEAR(NearestA.Y, (T)0, Tolerance);
+			EXPECT_GT(NearestA.Z, (T)-2-Tolerance);
+			EXPECT_LT(NearestA.Z, (T)2+Tolerance);
+			EXPECT_NEAR(NearestB.X, (T)4, Tolerance);
+			EXPECT_NEAR(NearestB.Y, (T)0, Tolerance);
+			EXPECT_GT(NearestB.Z, (T)-1-Tolerance);
+			EXPECT_LT(NearestB.Z, (T)2+Tolerance);
+		}
+
+		// Success - not overlapping, capsule axis at angle to nearest face (near points on end-cap and box edge)
+		{
+			TAABB<T, 3> A(TVector<T, 3>(5, -2, -2), TVector<T, 3>(8, 2, 2));
+			TCapsule<T> B(TVector<T, 3>(-2, 0, 3), TVector<T, 3>(2, 0, -3), 2);
+			bool bSuccess = GJKDistance<T>(A, B, TRigidTransform<T, 3>::Identity, Distance, NearestA, NearestB);
+			TVector<T, 3> ExpectedNearestA = TVector<T, 3>(5, 0, -2);
+			TVector<T, 3> ExpectedDir = (ExpectedNearestA - B.GetX2()).GetSafeNormal();
+			TVector<T, 3> ExpectedNearestB = B.GetX2() + ExpectedDir * B.GetRadius();
+
+			const T Tolerance = (T)2e-3;
+			EXPECT_TRUE(bSuccess);
+			EXPECT_NEAR(Distance, (ExpectedNearestB - ExpectedNearestA).Size(), Tolerance);
+			EXPECT_NEAR(NearestA.X, (T)ExpectedNearestA.X, Tolerance);
+			EXPECT_NEAR(NearestA.Y, (T)ExpectedNearestA.Y, Tolerance);
+			EXPECT_NEAR(NearestA.Z, (T)ExpectedNearestA.Z, Tolerance);
+			EXPECT_NEAR(NearestB.X, (T)ExpectedNearestB.X, Tolerance);
+			EXPECT_NEAR(NearestB.Y, (T)ExpectedNearestB.Y, Tolerance);
+			EXPECT_NEAR(NearestB.Z, (T)ExpectedNearestB.Z, Tolerance);
+		}
+
+		// Success - not overlapping, near point partway down wall of capsule
+		{
+			TCapsule<T> A(TVector<T, 3>(4, 0, -1), TVector<T, 3>(4, 0, -7), 1);
+			TAABB<T, 3> B(TVector<T, 3>(-2, -2, -2), TVector<T, 3>(2, 2, 2));
+			bool bSuccess = GJKDistance<T>(A, B, TRigidTransform<T, 3>::Identity, Distance, NearestA, NearestB);
+			TVector<T, 3> ExpectedNearestA = TVector<T, 3>(3, 0, (T)-1.5);
+			TVector<T, 3> ExpectedNearestB = TVector<T, 3>(2, 0, (T)-1.5);
+
+			const T Tolerance = (T)2e-3;
+			EXPECT_TRUE(bSuccess);
+			EXPECT_NEAR(Distance, (T)1, Tolerance);
+			EXPECT_NEAR(NearestA.X, (T)ExpectedNearestA.X, Tolerance);
+			EXPECT_NEAR(NearestA.Y, (T)ExpectedNearestA.Y, Tolerance);
+			EXPECT_LT(NearestA.Z, (T)ExpectedNearestA.Z + (T)0.5 + Tolerance);
+			EXPECT_GT(NearestA.Z, (T)ExpectedNearestA.Z - (T)0.5 - Tolerance);
+			EXPECT_NEAR(NearestB.X, (T)ExpectedNearestB.X, Tolerance);
+			EXPECT_NEAR(NearestB.Y, (T)ExpectedNearestB.Y, Tolerance);
+			EXPECT_NEAR(NearestB.Z, (T)NearestA.Z, Tolerance);
+		}
+
+		// Success - not overlapping, near point partway down wall of capsule.
+		// Same result as above, but using transform rather than the shape's built-in offsets.
+		{
+			TCapsule<T> A(TVector<T, 3>(0, 0, -3), TVector<T, 3>(0, 0, 3), 1);
+			TAABB<T, 3> B(TVector<T, 3>(-2, -2, -2), TVector<T, 3>(2, 2, 2));
+			TRigidTransform<T, 3> BToA = TRigidTransform<T, 3>(TVector<T, 3>(-4, 0, 4), TRotation<T, 3>::FromIdentity());
+			bool bSuccess = GJKDistance<T>(A, B, BToA, Distance, NearestA, NearestB);
+			TVector<T, 3> ExpectedNearestA = TVector<T, 3>(-1, 0, (T)2);
+			TVector<T, 3> ExpectedNearestB = TVector<T, 3>(2, 0, (T)-2);
+
+			const T Tolerance = (T)2e-3;
+			EXPECT_TRUE(bSuccess);
+			EXPECT_NEAR(Distance, (T)1, Tolerance);
+			EXPECT_NEAR(NearestA.X, (T)ExpectedNearestA.X, Tolerance);
+			EXPECT_NEAR(NearestA.Y, (T)ExpectedNearestA.Y, Tolerance);
+			EXPECT_LT(NearestA.Z, (T)ExpectedNearestA.Z + (T)0.5 + Tolerance);
+			EXPECT_GT(NearestA.Z, (T)ExpectedNearestA.Z - (T)0.5 - Tolerance);
+			EXPECT_NEAR(NearestB.X, (T)ExpectedNearestB.X, Tolerance);
+			EXPECT_NEAR(NearestB.Y, (T)ExpectedNearestB.Y, Tolerance);
+			EXPECT_NEAR(NearestB.Z + BToA.GetTranslation().Z, (T)NearestA.Z, Tolerance);
+		}
+	}
+
+
+	TEST(TestGJKDistance, GJKBoxCapsuleDistanceTest)
+	{
+		GJKBoxCapsuleDistanceTest<float>();
+	}
+
+
+	template <typename T>
+	void GJKBoxCapsuleDistanceIterationCountTest()
+	{
+		TVector<T, 3> NearestA = { 0,0,0 };
+		TVector<T, 3> NearestB = { 0,0,0 };
+		T Distance = 0;
+
+		// Capsule-box takes number of iterations at the moment (we can improve that with a better the choice of Initial V)
+		// so test that we still get an approximate answer with less iterations
+		{
+			TAABB<T, 3> A(TVector<T, 3>(5, -2, -2), TVector<T, 3>(8, 2, 2));
+			TCapsule<T> B(TVector<T, 3>(-2, 0, 3), TVector<T, 3>(2, 0, -3), 2);
+			T Epsilon = (T)1e-6;
+			int32 MaxIts = 5;
+			bool bSuccess = GJKDistance<T>(A, B, TRigidTransform<T, 3>::Identity, Distance, NearestA, NearestB, Epsilon, MaxIts);
+			TVector<T, 3> ExpectedNearestA = TVector<T, 3>(5, 0, -2);
+			TVector<T, 3> ExpectedDir = (ExpectedNearestA - B.GetX2()).GetSafeNormal();
+			TVector<T, 3> ExpectedNearestB = B.GetX2() + ExpectedDir * B.GetRadius();
+
+			const T Tolerance = (T)0.3;
+			EXPECT_TRUE(bSuccess);
+			EXPECT_NEAR(Distance, (ExpectedNearestB - ExpectedNearestA).Size(), Tolerance);
+			EXPECT_NEAR(NearestA.X, (T)ExpectedNearestA.X, Tolerance);
+			EXPECT_NEAR(NearestA.Y, (T)ExpectedNearestA.Y, Tolerance);
+			EXPECT_NEAR(NearestA.Z, (T)ExpectedNearestA.Z, Tolerance);
+			EXPECT_NEAR(NearestB.X, (T)ExpectedNearestB.X, Tolerance);
+			EXPECT_NEAR(NearestB.Y, (T)ExpectedNearestB.Y, Tolerance);
+			EXPECT_NEAR(NearestB.Z, (T)ExpectedNearestB.Z, Tolerance);
+		}
+	}
+
+
+	TEST(TestGJKDistance, GJKBoxCapsuleDistanceIterationCountTest)
+	{
+		GJKBoxCapsuleDistanceIterationCountTest<float>();
+	}
+
+}
\ No newline at end of file
diff --git a/Engine/Source/Programs/HeadlessChaos/Private/HeadlessChaosTestGraph.cpp b/Engine/Source/Programs/HeadlessChaos/Private/HeadlessChaosTestGraph.cpp
new file mode 100644
--- /dev/null
+++ b/Engine/Source/Programs/HeadlessChaos/Private/HeadlessChaosTestGraph.cpp
@@ -0,0 +1,898 @@
+// Copyright Epic Games, Inc. All Rights Reserved.
+
+#include "CoreMinimal.h"
+
+#include "Chaos/PBDConstraintColor.h"
+#include "Chaos/PBDConstraintGraph.h"
+#include "Chaos/PBDRigidsEvolution.h"
+#include "Chaos/PBDRigidsEvolutionGBF.h"
+#include "Chaos/Utilities.h"
+#include "HeadlessChaos.h"
+#include "HeadlessChaosTestUtility.h"
+#include "Modules/ModuleManager.h"
+
+namespace ChaosTest {
+
+	using namespace Chaos;
+
+	template<int32 T_TYPEID>
+	class TMockGraphConstraints;
+
+
+	template<int32 T_TYPEID>
+	class TMockGraphConstraintHandle : public FConstraintHandle
+	{
+	public:
+		TMockGraphConstraintHandle(TMockGraphConstraints<T_TYPEID>* InConstraintContainer, int32 ConstraintIndex)
+			: FConstraintHandle(ConstraintIndex)
+			, ConstraintContainer(InConstraintContainer)
+		{
+		}
+
+		TMockGraphConstraints<T_TYPEID>* ConstraintContainer;
+	};
+
+	/**
+	 * Constraint Container with minimal API required to test the Graph.
+	 * We can pretend we have many constraint containers of different types
+	 * by using containers with different T_TYPEIDs.
+	 */
+	template<int32 T_TYPEID>
+	class TMockGraphConstraints
+	{
+	public:
+		using FConstraintContainerHandle = TMockGraphConstraintHandle<T_TYPEID>;
+		struct FMockConstraint
+		{
+			TVector<int32, 2> ConstrainedParticles;
+		};
+
+		int32 NumConstraints() const { return Constraints.Num(); }
+
+		TVector<int32, 2> ConstraintParticleIndices(const int32 ConstraintIndex) const
+		{
+			return Constraints[ConstraintIndex].ConstrainedParticles;
+		}
+
+		void AddConstraint(const TVector<int32, 2>& InConstraintedParticles)
+		{
+			Constraints.Emplace(FMockConstraint({ InConstraintedParticles }));
+			Handles.Emplace(HandleAllocator.AllocHandle(this, Handles.Num()));
+		}
+
+		bool AreConstraintsIndependent(const TPBDRigidParticles<FReal, 3>& InParticles, const TArray<FConstraintHandle*>& InConstraintHandles)
+		{
+			bool bIsValid = true;
+			TSet<int32> ParticleSet;
+			for (FConstraintHandle* ConstraintHandle : InConstraintHandles)
+			{
+				FMockConstraint& Constraint = Constraints[ConstraintHandle->GetConstraintIndex()];
+				if (ParticleSet.Contains(Constraint.ConstrainedParticles[0]) && !InParticles.HasInfiniteMass(Constraint.ConstrainedParticles[0]))
+				{
+					bIsValid = false;
+				}
+				if (ParticleSet.Contains(Constraint.ConstrainedParticles[1]) && !InParticles.HasInfiniteMass(Constraint.ConstrainedParticles[1]))
+				{
+					bIsValid = false;
+				}
+				ParticleSet.Add(Constraint.ConstrainedParticles[0]);
+				ParticleSet.Add(Constraint.ConstrainedParticles[1]);
+			}
+			return bIsValid;
+		}
+
+		bool AreConstraintsIndependent(const TArray<TGeometryParticleHandle<FReal, 3>*>& Particles, const TArray<FConstraintHandle*>& InConstraintHandles)
+		{
+			bool bIsValid = true;
+			TSet<TGeometryParticleHandle<FReal, 3>*> ParticleSet;
+			for (FConstraintHandle* ConstraintHandle : InConstraintHandles)
+			{
+				FMockConstraint& Constraint = Constraints[ConstraintHandle->GetConstraintIndex()];
+				TGeometryParticleHandle<FReal, 3>* Particle0 = Particles[Constraint.ConstrainedParticles[0]];
+				if (ParticleSet.Contains(Particle0) && Particle0->CastToRigidParticle() && Particle0->ObjectState() == EObjectStateType::Dynamic)
+				{
+					bIsValid = false;
+				}
+
+				TGeometryParticleHandle<FReal, 3>* Particle1 = Particles[Constraint.ConstrainedParticles[1]];
+				if (ParticleSet.Contains(Particle1) && Particle1->CastToRigidParticle() && Particle1->ObjectState() == EObjectStateType::Dynamic)
+				{
+					bIsValid = false;
+				}
+				ParticleSet.Add(Particle0);
+				ParticleSet.Add(Particle1);
+			}
+			return bIsValid;
+		}
+
+
+		TArray<FMockConstraint> Constraints;
+		TArray<FConstraintContainerHandle*> Handles;
+		TConstraintHandleAllocator<TMockGraphConstraints<T_TYPEID>> HandleAllocator;
+
+	};
+	template class TMockGraphConstraints<0>;
+	template class TMockGraphConstraints<1>;
+
+	void GraphIslands()
+	{
+		// Create some dynamic particles - doesn't matter what position or other state they have
+		TArray<TGeometryParticleHandle<FReal, 3>*> AllParticles;
+
+		TPBDRigidsSOAs<FReal, 3> SOAs;
+		TArray<TPBDRigidParticleHandle<FReal,3>*> Dynamics = SOAs.CreateDynamicParticles(17);
+		for (auto Dyn : Dynamics) { AllParticles.Add(Dyn);}
+
+		// Make a static particle. Islands should not merge across these.
+		TArray<TGeometryParticleHandle<FReal, 3>*> Statics = SOAs.CreateStaticParticles(1);
+		AllParticles.Append(Statics);
+
+		Dynamics = SOAs.CreateDynamicParticles(3);
+		for (auto Dyn : Dynamics) { AllParticles.Add(Dyn);}
+		
+		// Create some constraints between the particles
+		TArray<TVector<int32, 2>> ConstrainedParticles0 =
+		{
+			//
+			{0, 1},
+			{0, 2},
+			{0, 3},
+			{3, 4},
+			{3, 5},
+			{6, 4},
+			//
+			{8, 7},
+			{8, 9},
+			//
+			{13, 18},
+			//
+			{20, 17},
+		};
+
+		TArray<TVector<int32, 2>> ConstrainedParticles1 =
+		{
+			//
+			{0, 1},
+			{2, 1},
+			//
+			{9, 10},
+			{11, 10},
+			{11, 13},
+			//
+			{14, 15},
+			{16, 14},
+			{17, 14},
+		};
+
+		TMockGraphConstraints<0> ConstraintsOfType0;
+		for (const auto& ConstrainedParticles : ConstrainedParticles0)
+		{
+			ConstraintsOfType0.AddConstraint(ConstrainedParticles);
+		}
+
+		TMockGraphConstraints<1> ConstraintsOfType1;
+		for (const auto& ConstrainedParticles : ConstrainedParticles1)
+		{
+			ConstraintsOfType0.AddConstraint(ConstrainedParticles);
+		}
+
+		// Set up the particle graph
+		FPBDConstraintGraph Graph;
+		Graph.InitializeGraph(SOAs.GetNonDisabledView());
+
+		Graph.ReserveConstraints(ConstraintsOfType0.NumConstraints());
+		for (int32 ConstraintIndex = 0; ConstraintIndex < ConstraintsOfType0.NumConstraints(); ++ConstraintIndex)
+		{
+			TVector<int32, 2> Indices = ConstraintsOfType0.ConstraintParticleIndices(ConstraintIndex);
+			Graph.AddConstraint(0, ConstraintsOfType0.Handles[ConstraintIndex], TVector<TGeometryParticleHandle<FReal, 3>*,2>(AllParticles[Indices[0]], AllParticles[Indices[1]]));
+		}
+
+		Graph.ReserveConstraints(ConstraintsOfType1.NumConstraints());
+		for (int32 ConstraintIndex = 0; ConstraintIndex < ConstraintsOfType1.NumConstraints(); ++ConstraintIndex)
+		{
+			TVector<int32, 2> Indices = ConstraintsOfType1.ConstraintParticleIndices(ConstraintIndex);
+			Graph.AddConstraint(1, ConstraintsOfType1.Handles[ConstraintIndex], TVector<TGeometryParticleHandle<FReal, 3>*, 2>(AllParticles[Indices[0]], AllParticles[Indices[1]]));
+		}
+		
+		// Generate the constraint/particle islands
+		Graph.UpdateIslands(SOAs.GetNonDisabledDynamicView(), SOAs);
+		
+		// Islands should be end up with the following particles (note: particle 17 is infinite mass and can appear in multiple islands)
+		TArray<TSet<TGeometryParticleHandle<FReal, 3>*>> ExpectedIslandParticles = {
+			{AllParticles[0], AllParticles[1], AllParticles[2], AllParticles[3], AllParticles[4], AllParticles[5], AllParticles[6]},
+			{AllParticles[7], AllParticles[8], AllParticles[9], AllParticles[10], AllParticles[11], AllParticles[13], AllParticles[18]},
+			{AllParticles[12]},
+			{AllParticles[14], AllParticles[15], AllParticles[16], AllParticles[17]},
+			{AllParticles[19]},
+			{AllParticles[17], AllParticles[20]},
+		};
+
+		// Get the Island indices which map to the ExpectedIslandParticles
+		const TArray<int32> CalculatedIslandIndices =
+		{
+			AllParticles[0]->CastToRigidParticle()->Island(),
+			AllParticles[7]->CastToRigidParticle()->Island(),
+			AllParticles[12]->CastToRigidParticle()->Island(),
+			AllParticles[14]->CastToRigidParticle()->Island(),
+			AllParticles[19]->CastToRigidParticle()->Island(),
+			AllParticles[20]->CastToRigidParticle()->Island(),
+		};
+
+		// All non-static partcles should still be Active
+		EXPECT_EQ(SOAs.GetActiveParticlesView().Num(), 20);
+		for (auto& Particle : SOAs.GetActiveParticlesView())
+		{
+			EXPECT_NE(Particle.Handle(), AllParticles[17]);
+		}
+		
+		// Each calculated island should contain the particles we expected and no others
+		for (int32 ExpectedIslandIndex = 0; ExpectedIslandIndex < ExpectedIslandParticles.Num(); ++ExpectedIslandIndex)
+		{
+			const int32 CalculatedIslandIndex = CalculatedIslandIndices[ExpectedIslandIndex];
+			const TArray<TGeometryParticleHandle<FReal, 3>*>& CalculatedIslandParticles = Graph.GetIslandParticles(CalculatedIslandIndex);
+
+			EXPECT_EQ(CalculatedIslandParticles.Num(), ExpectedIslandParticles[ExpectedIslandIndex].Num());
+			for (TGeometryParticleHandle<FReal, 3>* CalculatedIslandParticleIndex : CalculatedIslandParticles)
+			{
+				EXPECT_TRUE(ExpectedIslandParticles[ExpectedIslandIndex].Contains(CalculatedIslandParticleIndex));
+			}
+		}
+	}
+
+	void CheckIslandIntegrity(TArray<TSet<TGeometryParticleHandle<FReal,3>*>>& ExpectedIslandParticles, const TArray<int32> CalculatedIslandIndices, FPBDConstraintGraph& Graph)
+	{
+		for (int32 ExpectedIslandIndex = 0; ExpectedIslandIndex < ExpectedIslandParticles.Num(); ++ExpectedIslandIndex)
+		{
+			const int32 CalculatedIslandIndex = CalculatedIslandIndices[ExpectedIslandIndex];
+			const TArray<TGeometryParticleHandle<FReal, 3>*>& CalculatedIslandParticles = Graph.GetIslandParticles(CalculatedIslandIndex);
+
+			EXPECT_EQ(CalculatedIslandParticles.Num(), ExpectedIslandParticles[ExpectedIslandIndex].Num());
+			for (TGeometryParticleHandle<FReal, 3>* CalculatedIslandParticleIndex : CalculatedIslandParticles)
+			{
+				EXPECT_TRUE(ExpectedIslandParticles[ExpectedIslandIndex].Contains(CalculatedIslandParticleIndex));
+
+				auto* DynParticle = CalculatedIslandParticleIndex->CastToRigidParticle();
+				if (DynParticle && DynParticle->ObjectState() == EObjectStateType::Dynamic)
+				{
+					EXPECT_EQ(DynParticle->Island(), CalculatedIslandIndex);
+				}
+			}
+		}
+	}
+
+	struct FIterationData 
+	{
+		TArray<TVector<int32, 2>> ConstrainedParticles;
+
+		TArray<TSet<int32>> ExpectedIslandParticleIndices;
+
+		TArray<int32> ExpectedIslandEdges;
+
+		TArray<TSet<TGeometryParticleHandle<FReal, 3>*>> ExpectedIslandParticles;
+
+		TArray<int32> MaxLevel;
+		TArray<int32> MaxColor;
+	};
+
+	void GraphIslandsPersistence()
+	{
+		// Create some dynamic particles - doesn't matter what position or other state they have
+		TArray<TGeometryParticleHandle<FReal, 3>*> AllParticles;
+
+		TPBDRigidsSOAs<FReal, 3> SOAs;
+		TArray<TPBDRigidParticleHandle<FReal, 3>*> Dynamics = SOAs.CreateDynamicParticles(2);
+		for (auto Dyn : Dynamics) { AllParticles.Add(Dyn); }
+
+		// Make a static particle. Islands should not merge across these.
+		TArray<TGeometryParticleHandle<FReal, 3>*> Statics = SOAs.CreateStaticParticles(1);
+		AllParticles.Append(Statics);
+
+		Dynamics = SOAs.CreateDynamicParticles(3);
+		for (auto Dyn : Dynamics) { AllParticles.Add(Dyn); }
+
+		//////////////////////////////////////////////////////////////////////////
+		// Iteration 0
+		TArray<FIterationData> IterationData;
+
+		FIterationData Data0;
+		Data0.ExpectedIslandParticleIndices = {
+			{0},
+			{1},
+			{3},
+			{4},
+			{5},
+		};
+
+		Data0.ExpectedIslandEdges = {
+			0,
+			0,
+			0,
+			0,
+			0
+		};
+
+		Data0.MaxLevel = { -1, -1, -1, -1, -1 };
+		Data0.MaxColor = { -1, -1, -1, -1, -1 };
+
+		IterationData.Push(Data0);
+
+		//////////////////////////////////////////////////////////////////////////
+		// Iteration 1
+		FIterationData Data1;
+		Data1.ConstrainedParticles = {
+			//
+			{0, 1},
+			{1, 2},
+			//
+			{3,4},
+			{4,5}
+		};
+
+		// Islands should be end up with the following particles
+		Data1.ExpectedIslandParticleIndices = {
+			{0,1,2},
+			{3,4,5}
+		};
+
+		Data1.ExpectedIslandEdges = {
+			2,
+			2
+		};
+
+		Data1.MaxLevel = { 1, 0 };
+		Data1.MaxColor = { 1, 1 };
+
+		IterationData.Push(Data1);
+
+		//////////////////////////////////////////////////////////////////////////
+		// Iteration 2
+		FIterationData Data2;
+		Data2.ConstrainedParticles =
+		{
+			//
+			{0,1},
+			{1,2},
+			{2,3},
+			{3,4},
+			{4,5}
+		};
+
+		Data2.ExpectedIslandParticleIndices = {
+			{0,1,2},
+			{2,3,4,5},
+		};
+
+		Data2.ExpectedIslandEdges = {
+			2,
+			3
+		};
+
+		Data2.MaxLevel = { 1, 2 };
+		Data2.MaxColor = { 1, 1 };
+
+		IterationData.Push(Data2);
+		//////////////////////////////////////////////////////////////////////////
+		// Iteration 3
+		FIterationData Data3;
+		Data3.ConstrainedParticles =
+		{
+			//
+			{0,1},
+			//
+			{2,3},
+			//
+			{4,5}
+		};
+
+		Data3.ExpectedIslandParticleIndices = {
+			{0,1},
+			{2,3},
+			{4,5}
+		};
+
+		Data3.ExpectedIslandEdges = {
+			1,
+			1,
+			1
+		};
+
+		Data3.MaxLevel = { 0, 0, 0 };
+		Data3.MaxColor = { 0, 0, 0 };
+
+		IterationData.Push(Data3);
+		//////////////////////////////////////////////////////////////////////////
+		// Iteration 4
+		FIterationData Data4;
+		Data4.ConstrainedParticles =
+		{
+			//
+			{0,1},
+			//
+			{2,3},
+		};
+
+		Data4.ExpectedIslandParticleIndices = {
+			{0,1},
+			{2,3},
+			{4},
+			{5}
+		};
+
+		Data4.ExpectedIslandEdges = {
+			1,
+			1,
+			0,
+			0
+		};
+
+		Data4.MaxLevel = { 0, 0, -1, -1 };
+		Data4.MaxColor = { 0, 0, -1, -1 };
+
+		IterationData.Push(Data4);
+		//////////////////////////////////////////////////////////////////////////
+		// Iteration 5
+		FIterationData Data5;
+		Data5.ConstrainedParticles =
+		{
+			//
+			{0,2},
+			//
+			{3,2},
+		};
+
+		Data5.ExpectedIslandParticleIndices = {
+			{1},
+			{0,2},
+			{2,3},
+			{4},
+			{5}
+		};
+
+		Data5.ExpectedIslandEdges = {
+			0,
+			1,
+			1,
+			0,
+			0
+		};
+
+		Data5.MaxLevel = { 0, -1, 0, -1, -1 };
+		Data5.MaxColor = { 0, -1, 0, -1, -1 };
+
+		IterationData.Push(Data5);
+		////////////////////////////////////////////////
+
+		for (int Iteration = 0; Iteration < IterationData.Num(); Iteration++)
+		{
+			FIterationData& IterData = IterationData[Iteration];
+
+			// convert ExpectedIslandParticleIndices to ExpectedIslandParticles for ease of later comparison
+			for (int32 I = 0; I < IterData.ExpectedIslandParticleIndices.Num(); I++)
+			{
+				const TSet<int32>& Set = IterData.ExpectedIslandParticleIndices[I];
+
+				TSet<TGeometryParticleHandle<FReal, 3>*> HandlesSet;
+				for (int32 Index : Set)
+				{
+					HandlesSet.Add(AllParticles[Index]);
+				}
+				IterData.ExpectedIslandParticles.Add(HandlesSet);
+			}
+			FPBDConstraintGraph Graph;
+			FPBDConstraintColor GraphColor;
+			const int32 ContainerId = 0;
+
+			// Set up the particle graph
+			Graph.InitializeGraph(SOAs.GetNonDisabledView());
+
+			// add constraints
+			TMockGraphConstraints<0> ConstraintsOfType0;
+			for (const auto& ConstrainedParticles : IterData.ConstrainedParticles)
+			{
+				ConstraintsOfType0.AddConstraint(ConstrainedParticles);
+			}
+
+			Graph.ReserveConstraints(ConstraintsOfType0.NumConstraints());
+			for (int32 ConstraintIndex = 0; ConstraintIndex < ConstraintsOfType0.NumConstraints(); ++ConstraintIndex)
+			{
+				TVector<int32, 2> Indices = ConstraintsOfType0.ConstraintParticleIndices(ConstraintIndex);
+				Graph.AddConstraint(ContainerId, ConstraintsOfType0.Handles[ConstraintIndex], TVector<TGeometryParticleHandle<FReal, 3>*, 2>(AllParticles[Indices[0]], AllParticles[Indices[1]]));
+			}
+
+			// Generate the constraint/particle islands
+			Graph.UpdateIslands(SOAs.GetNonDisabledDynamicView(), SOAs);
+
+			// Assign color to constraints
+			GraphColor.InitializeColor(Graph);
+			for (int32 Island = 0; Island < Graph.NumIslands(); ++Island)
+			{
+				GraphColor.ComputeColor(Island, Graph, ContainerId);
+			}
+
+			// get the generated island indices
+			TArray<int32> CalculatedIslandIndices;
+			for (TSet<int32>& EIP : IterData.ExpectedIslandParticleIndices)
+			{	
+				auto SetAsArray = EIP.Array();
+				int32 FoundIsland = INDEX_NONE;
+				for (int32 ParticleIdx : SetAsArray)
+				{
+					auto RigidParticle = AllParticles[ParticleIdx]->CastToRigidParticle();
+					if (RigidParticle && RigidParticle->ObjectState() == EObjectStateType::Dynamic)
+					{
+						if (FoundIsland == INDEX_NONE)
+						{
+							FoundIsland = RigidParticle->Island();
+							CalculatedIslandIndices.Push(FoundIsland);
+						}
+						else
+						{
+							EXPECT_EQ(FoundIsland, RigidParticle->Island());
+						}
+					}
+				}
+				check(FoundIsland != INDEX_NONE);
+			}
+
+			// check the number of edges matches what we are expecting for this island
+			int Index = 0;
+			for (int32 Island : CalculatedIslandIndices)
+			{
+				const TArray<int32>& ConstraintDataIndices = Graph.GetIslandConstraintData(Island);
+				EXPECT_EQ(ConstraintDataIndices.Num(), IterData.ExpectedIslandEdges[Index++]);
+			}
+
+			CheckIslandIntegrity(IterData.ExpectedIslandParticles, CalculatedIslandIndices, Graph);
+
+			// check level/color integrity
+			TSet<FConstraintHandle*> ConstraintUnionSet;
+			for (int32 Island = 0; Island < Graph.NumIslands(); ++Island)
+			{
+				const typename FPBDConstraintColor::FLevelToColorToConstraintListMap& LevelToColorToConstraintListMap = GraphColor.GetIslandLevelToColorToConstraintListMap(Island);
+				const int32 MaxLevel = GraphColor.GetIslandMaxLevel(Island);
+				const int32 MaxColor = GraphColor.GetIslandMaxColor(Island);
+
+				EXPECT_EQ(IterData.MaxLevel[Island], MaxLevel);
+				EXPECT_EQ(IterData.MaxColor[Island], MaxColor);
+			}
+		}
+	}
+
+
+	/**
+	 * Create some constrained sets of particles, some of which meet the sleep criteria, and 
+	 * verify that they sleep when expected while the others do not.
+	 */
+	
+	void GraphSleep()
+	{
+		TUniquePtr<FChaosPhysicsMaterial> PhysicalMaterial = MakeUnique<FChaosPhysicsMaterial>();
+		PhysicalMaterial->Friction = 0;
+		PhysicalMaterial->Restitution = 0;
+		PhysicalMaterial->SleepingLinearThreshold = 10;
+		PhysicalMaterial->SleepingAngularThreshold = 10;
+		PhysicalMaterial->DisabledLinearThreshold = 0;
+		PhysicalMaterial->DisabledAngularThreshold = 0;
+
+		// Create some dynamic particles
+		int32 NumParticles = 6;
+		TPBDRigidsSOAs<FReal, 3> SOAs;
+		TArray<TPBDRigidParticleHandle<FReal, 3>*> DynParticles = SOAs.CreateDynamicParticles(NumParticles);
+		TArray<TGeometryParticleHandle<FReal, 3>*> Particles;
+		TArrayCollectionArray<TSerializablePtr<FChaosPhysicsMaterial>> PhysicsMaterials;
+		SOAs.GetParticleHandles().AddArray(&PhysicsMaterials);
+
+		for (int32 Idx = 0; Idx < NumParticles; ++Idx)
+		{
+			Particles.Add(DynParticles[Idx]);
+			PhysicsMaterials[Idx] = MakeSerializable(PhysicalMaterial);
+		}
+		// Ensure some particles will not sleep
+		DynParticles[0]->V() = FVec3(100);
+		DynParticles[1]->V() = FVec3(100);
+		DynParticles[2]->V() = FVec3(100);
+
+		// Ensure others will sleep but only if sleep threshold is actually considered
+		DynParticles[3]->V() = FVec3(1);
+		DynParticles[4]->V() = FVec3(1);
+
+		// Create some constraints between the particles
+		TArray<TVector<int32, 2>> ConstrainedParticles =
+		{
+			//
+			{0, 1},
+			//
+			{3, 4},
+		};
+
+		TMockGraphConstraints<0> Constraints;
+		for (const auto& ConstrainedParticleIndices : ConstrainedParticles)
+		{
+			Constraints.AddConstraint(ConstrainedParticleIndices);
+		}
+		FPBDConstraintGraph Graph;
+		for (int32 LoopIndex = 0; LoopIndex < 5 + FPBDConstraintGraph::SleepCountThreshold; ++LoopIndex)
+		{
+			Graph.InitializeGraph(SOAs.GetNonDisabledView());
+
+			Graph.ReserveConstraints(Constraints.NumConstraints());
+			for (int32 ConstraintIndex = 0; ConstraintIndex < Constraints.NumConstraints(); ++ConstraintIndex)
+			{
+				const TVector<int32, 2> Indices = Constraints.ConstraintParticleIndices(ConstraintIndex);
+				Graph.AddConstraint(0, Constraints.Handles[ConstraintIndex], TVector<TGeometryParticleHandle<FReal, 3>*, 2>(Particles[Indices[0]], Particles[Indices[1]]));
+			}
+
+			Graph.UpdateIslands(SOAs.GetNonDisabledDynamicView(), SOAs);
+
+			for (int32 IslandIndex = 0; IslandIndex < Graph.NumIslands(); ++IslandIndex)
+			{
+				bool bSleeped = Graph.SleepInactive(IslandIndex, PhysicsMaterials);
+
+				if (bSleeped)
+				{
+					for (TGeometryParticleHandle<FReal, 3>* Particle : Graph.GetIslandParticles(IslandIndex))
+					{
+						SOAs.DeactivateParticle(Particle);
+					}
+				}
+			}
+
+			auto ContainsHelper = [](const TParticleView<TPBDRigidParticles<FReal, 3>>& View, const TGeometryParticleHandle<FReal, 3>* InParticle)
+			{
+				for (auto& Particle : View)
+				{
+					if (Particle.Handle() == InParticle)
+					{
+						return true;
+					}
+				}
+				return false;
+			};
+			
+			// Particles 0-2 are always awake
+			EXPECT_FALSE(Particles[0]->Sleeping());
+			EXPECT_FALSE(Particles[1]->Sleeping());
+			EXPECT_FALSE(Particles[2]->Sleeping());
+			EXPECT_TRUE(ContainsHelper(SOAs.GetActiveParticlesView(), Particles[0]));
+			EXPECT_TRUE(ContainsHelper(SOAs.GetActiveParticlesView(), Particles[1]));
+			EXPECT_TRUE(ContainsHelper(SOAs.GetActiveParticlesView(), Particles[2]));
+			// Particles 3-5 should sleep when we hit the frame count threshold and then stay asleep
+			bool bSomeShouldSleep = (LoopIndex > FPBDConstraintGraph::SleepCountThreshold);
+			EXPECT_EQ(Particles[3]->Sleeping(), bSomeShouldSleep);
+			EXPECT_EQ(Particles[4]->Sleeping(), bSomeShouldSleep);
+			EXPECT_EQ(Particles[5]->Sleeping(), bSomeShouldSleep);
+			EXPECT_NE(ContainsHelper(SOAs.GetActiveParticlesView(), Particles[3]), bSomeShouldSleep);
+			EXPECT_NE(ContainsHelper(SOAs.GetActiveParticlesView(), Particles[4]), bSomeShouldSleep);
+			EXPECT_NE(ContainsHelper(SOAs.GetActiveParticlesView(), Particles[5]), bSomeShouldSleep);
+		}
+	}
+
+
+	/**
+	 * Arrange particles in a MxN grid with constraints connecting adjacent pairs.
+	 * Outer particles are static.
+	 * Support multiple constraints between each adjacent pair.
+	 * Support randomization of constraint order.
+	 * Arrange particles in a grid, and create a number of constraints connecting adjacent pairs.
+	 * Verify that no two same-colored constraints affect the same particle.
+	 * Verify that we need no more than 2 x Node-Multiplicity + 1 colors. 
+	 *
+	 *        X   X   X   X   X
+	 *            |   |   |
+	 *        X - o - o - o - X
+	 *            |   |   |
+	 *        X - o - o - o - X
+	 *            |   |   |
+	 *        X - o - o - o - X
+	 *            |   |   |
+	 *        X   X   X   X   X
+	 *
+	 */
+
+	void GraphColorGrid(const int32 NumParticlesX, const int32 NumParticlesY, const int32 Multiplicity, const bool bRandomize)
+	{
+		// Create a grid of particles
+		TPBDRigidsSOAs<FReal, 3> SOAs;
+		int32 NumParticles = NumParticlesX * NumParticlesY;
+		TArray<TGeometryParticleHandle<FReal,3>*> AllParticles;
+
+		for (int32 ParticleIndexY = 0; ParticleIndexY < NumParticlesY; ++ParticleIndexY)
+		{
+			for (int32 ParticleIndexX = 0; ParticleIndexX < NumParticlesX; ++ParticleIndexX)
+			{
+				if ((ParticleIndexX == 0) || (ParticleIndexX == NumParticlesX - 1) || (ParticleIndexY == 0) || (ParticleIndexY == NumParticlesY - 1))
+				{
+					//border so static
+					AllParticles.Add(SOAs.CreateStaticParticles(1)[0]);
+				}
+				else
+				{
+					AllParticles.Add(SOAs.CreateDynamicParticles(1)[0]);
+				}
+
+				TGeometryParticleHandle<FReal, 3>* Particle = AllParticles.Last();
+				Particle->X() = FVec3((FReal)ParticleIndexX * 200, (FReal)ParticleIndexY * 200, (FReal)0);
+			}
+		}
+
+		// Determine which particle pairs should be constrained.
+		// Connect all adjacent pairs in a grid with one or more constraints for each pair
+		// making sure no two non-dynamic particles are connected to each other.
+		TArray<TVector<int32, 2>> ConstrainedParticles;
+		// X-Direction constraints
+		for (int32 ParticleIndexY = 0; ParticleIndexY < NumParticlesY; ++ParticleIndexY)
+		{
+			for (int32 ParticleIndexX = 0; ParticleIndexX < NumParticlesX - 1; ++ParticleIndexX)
+			{
+				int32 ParticleIndex0 = ParticleIndexX + ParticleIndexY * NumParticlesX;
+				int32 ParticleIndex1 = ParticleIndexX + ParticleIndexY * NumParticlesX + 1;
+
+				auto RigidParticle0 = AllParticles[ParticleIndex0]->CastToRigidParticle();
+				auto RigidParticle1 = AllParticles[ParticleIndex1]->CastToRigidParticle();
+				if (   (RigidParticle0 && RigidParticle0->ObjectState() == EObjectStateType::Dynamic)
+					|| (RigidParticle1 && RigidParticle1->ObjectState() == EObjectStateType::Dynamic))
+				{
+					for (int32 MultiplicityIndex = 0; MultiplicityIndex < Multiplicity; ++MultiplicityIndex)
+					{
+						ConstrainedParticles.Add({ ParticleIndex0, ParticleIndex1 });
+					}
+				}
+			}
+		}
+		// Y-Direction Constraints
+		for (int32 ParticleIndexY = 0; ParticleIndexY < NumParticlesY - 1; ++ParticleIndexY)
+		{
+			for (int32 ParticleIndexX = 0; ParticleIndexX < NumParticlesX; ++ParticleIndexX)
+			{
+				int32 ParticleIndex0 = ParticleIndexX + ParticleIndexY * NumParticlesX;
+				int32 ParticleIndex1 = ParticleIndexX + (ParticleIndexY + 1) * NumParticlesX;
+
+				auto RigidParticle0 = AllParticles[ParticleIndex0]->CastToRigidParticle();
+				auto RigidParticle1 = AllParticles[ParticleIndex1]->CastToRigidParticle();
+				if (   (RigidParticle0 && RigidParticle0->ObjectState() == EObjectStateType::Dynamic)
+					|| (RigidParticle1 && RigidParticle1->ObjectState() == EObjectStateType::Dynamic))
+				{
+					for (int32 MultiplicityIndex = 0; MultiplicityIndex < Multiplicity; ++MultiplicityIndex)
+					{
+						ConstrainedParticles.Add({ ParticleIndex0, ParticleIndex1 });
+					}
+				}
+			}
+		}
+
+		// Randomize the constraint order
+		if (bRandomize)
+		{
+			FMath::RandInit(3354786483);
+			for (int32 RandIndex = 0; RandIndex < 2 * ConstrainedParticles.Num(); ++RandIndex)
+			{
+				int32 Rand0 = FMath::RandRange(0, ConstrainedParticles.Num() - 1);
+				int32 Rand1 = FMath::RandRange(0, ConstrainedParticles.Num() - 1);
+				Swap(ConstrainedParticles[Rand0], ConstrainedParticles[Rand1]);
+			}
+		}
+
+		// Generate the constraints
+		TMockGraphConstraints<0> Constraints;
+		for (const auto& ConstrainedParticleIndices : ConstrainedParticles)
+		{
+			Constraints.AddConstraint(ConstrainedParticleIndices);
+		}
+		
+		// Build the connectivity graph and islands
+		FPBDConstraintGraph Graph;
+		FPBDConstraintColor GraphColor;
+		const int32 ContainerId = 0;
+		Graph.InitializeGraph(SOAs.GetNonDisabledView());
+		Graph.ReserveConstraints(Constraints.NumConstraints());
+		for (int32 ConstraintIndex = 0; ConstraintIndex < Constraints.NumConstraints(); ++ConstraintIndex)
+		{
+			const TVector<int32, 2>& Indices = Constraints.ConstraintParticleIndices(ConstraintIndex);
+			Graph.AddConstraint(ContainerId, Constraints.Handles[ConstraintIndex], TVector<TGeometryParticleHandle<FReal, 3>*,2>(AllParticles[Indices[0]], AllParticles[Indices[1]]) );
+		}
+		Graph.UpdateIslands(SOAs.GetNonDisabledDynamicView(), SOAs);
+
+		// It's a connected grid, so only one island
+		EXPECT_EQ(Graph.NumIslands(), 1);
+		
+		// Assign color to constraints
+		GraphColor.InitializeColor(Graph);
+		for (int32 Island = 0; Island < Graph.NumIslands(); ++Island)
+		{
+			GraphColor.ComputeColor(Island, Graph, ContainerId);
+		}
+		// Check colors
+		// No constraints should appear twice in the level-color-constraint map
+		// No particles should be influenced by more than one constraint in any individual level+color
+		TSet<FConstraintHandle*> ConstraintUnionSet;
+		for (int32 Island = 0; Island < Graph.NumIslands(); ++Island)
+		{
+			const typename FPBDConstraintColor::FLevelToColorToConstraintListMap& LevelToColorToConstraintListMap = GraphColor.GetIslandLevelToColorToConstraintListMap(Island);
+			const int32 MaxLevel = GraphColor.GetIslandMaxLevel(Island);
+			const int32 MaxColor = GraphColor.GetIslandMaxColor(Island);
+			for (int32 Level = 0; Level <= MaxLevel; ++Level)
+			{
+				for (int Color = 0; Color <= MaxColor; ++Color)
+				{
+					if (LevelToColorToConstraintListMap[Level].Contains(Color))
+					{
+						// Build the set of constraint and particle indices for this color
+						TSet<FConstraintHandle*> ConstraintSet = TSet<FConstraintHandle*>(LevelToColorToConstraintListMap[Level][Color]);
+
+						// See if we have any constraints that were also in a prior color
+						TSet<FConstraintHandle*> ConstraintIntersectSet = ConstraintUnionSet.Intersect(ConstraintSet);
+						EXPECT_EQ(ConstraintIntersectSet.Num(), 0);
+						ConstraintUnionSet = ConstraintUnionSet.Union(ConstraintSet);
+
+						// See if any particles are being modified by more than one constraint at this level+color
+						EXPECT_TRUE(Constraints.AreConstraintsIndependent(AllParticles, LevelToColorToConstraintListMap[Level][Color]));
+					}
+				}
+			}
+		}
+
+		// Verify that we have created a reasonable number of colors. For the greedy edge coloring algorithm this is
+		//		NumColors >= MaxNodeMultiplicity
+		//		NumColors <= MaxNodeMultiplicity * 2 - 1
+		// Each node connects to 4 neighbors with 'Multiplicity' connections, but some of them will be to static particles and we ignore those.
+		// For grid dimensions less that 4 each particle on has 2 non-static connections, otherwise there are up to 4
+		const int32 MaxMultiplicity = Multiplicity * (((NumParticlesX <= 4) || (NumParticlesY <= 4)) ? 2 : 4);
+		const int32 MinNumGreedyColors = MaxMultiplicity;
+		const int32 MaxNumGreedyColors = 2 * MaxMultiplicity - 1;
+		for (int32 Island = 0; Island < Graph.NumIslands(); ++Island)
+		{
+			const typename FPBDConstraintColor::FLevelToColorToConstraintListMap& LevelToColorToConstraintListMap = GraphColor.GetIslandLevelToColorToConstraintListMap(Island);
+			const int32 MaxIslandColor = GraphColor.GetIslandMaxColor(Island);
+			EXPECT_GE(MaxIslandColor, MinNumGreedyColors - 1);
+			// We are consistently slightly worse that the greedy algorithm, but hopefully doesn't matter
+			//EXPECT_LE(MaxIslandColor, MaxNumGreedyColors - 1);
+			EXPECT_LE(MaxIslandColor, MaxNumGreedyColors);
+		}
+	}
+
+
+	TEST(GraphTests, TestGraphIslands)
+	{
+		GraphIslands();
+	}
+
+	TEST(GraphTests, TestGraphIslandsPersistence)
+	{
+		GraphIslandsPersistence();
+	}
+		
+	TEST(GraphTests, TestGraphSleep)
+	{
+		GraphSleep();
+	}
+
+	TEST(GraphTests, TestGraphColor)
+	{
+		Chaos::TVector<int32, 2> GridDims[] = {
+			{3, 3},
+			{4, 4},
+			{5, 5},
+			{6, 6},
+			{7, 7},
+			{8, 8},
+			{9, 9},
+			{10, 10},
+			{20, 3},
+			{20, 10}
+		};
+		for (int32 Randomize = 0; Randomize < 2; ++Randomize)
+		{
+			bool bRandomize = (Randomize > 0);
+			for (int32 Multiplicity = 1; Multiplicity < 4; ++Multiplicity)
+			{
+				for (int32 Grid = 0; Grid < UE_ARRAY_COUNT(GridDims); ++Grid)
+				{
+					GraphColorGrid(GridDims[Grid][0], GridDims[Grid][1], Multiplicity, bRandomize);
+				}
+			}
+		}
+
+		SUCCEED();
+	}
+}
+
