Rename/move HeadlessChaos.

diff --git a/Engine/Source/Programs/HeadlessChaos/Private/HeadlessChaosTestImplicits.cpp b/Engine/Source/Programs/HeadlessChaos/Private/HeadlessChaosTestImplicits.cpp
new file mode 100644
--- /dev/null
+++ b/Engine/Source/Programs/HeadlessChaos/Private/HeadlessChaosTestImplicits.cpp
@@ -0,0 +1,1552 @@
+// Copyright Epic Games, Inc. All Rights Reserved.
+
+#include "HeadlessChaosTestImplicits.h"
+
+#include "HeadlessChaos.h"
+#include "HeadlessChaosCollisionConstraints.h"
+#include "Modules/ModuleManager.h"
+#include "Chaos/PBDRigidsEvolution.h"
+#include "Chaos/PBDRigidParticles.h"
+#include "Chaos/Box.h"
+#include "Chaos/Plane.h"
+#include "Chaos/Sphere.h"
+#include "Chaos/Cylinder.h"
+#include "Chaos/TaperedCylinder.h"
+#include "Chaos/Capsule.h"
+#include "Chaos/ImplicitObject.h"
+#include "Chaos/ImplicitObjectTransformed.h"
+#include "Chaos/ImplicitObjectUnion.h"
+#include "Chaos/ImplicitObjectScaled.h"
+#include "Chaos/ImplicitObjectIntersection.h"
+#include "Chaos/Levelset.h"
+#include "Chaos/UniformGrid.h"
+#include "Chaos/Utilities.h"
+#include "Chaos/Convex.h"
+#include "Math/RandomStream.h"
+#include "Chaos/ErrorReporter.h"
+
+#define RUN_KNOWN_BROKEN_TESTS 0
+
+namespace ChaosTest {
+
+	using namespace Chaos;
+
+	DEFINE_LOG_CATEGORY_STATIC(LogChaosTestImplicits, Verbose, All);
+	
+
+	/* HELPERS */
+
+
+	/* Takes an ImplicitObject of unit size (circumscribed inside a 2x2 cube centered on the origin). 
+	   Tests the .Normal() function and the .SignedDistance() function. */
+	template<class T>
+	void UnitImplicitObjectNormalsInternal(FImplicitObject &Subject, FString Caller)
+	{
+		typedef TVector<T, 3> TVector3;
+		FString Error = FString("Called by ") + Caller + FString(".");
+
+#if RUN_KNOWN_BROKEN_TESTS
+		// Normal when equally close to many points (currently inconsistent between geometries)
+		EXPECT_VECTOR_NEAR_ERR(Subject.Normal(TVector3(0)), TVector3(0, 0, 0), KINDA_SMALL_NUMBER, Error);
+#endif
+
+		// inside normal
+		EXPECT_VECTOR_NEAR_ERR(Subject.Normal(TVector3(0, 0, 1 / 2.)), (TVector3(0, 0, 1)), KINDA_SMALL_NUMBER, Error);
+		EXPECT_VECTOR_NEAR_ERR(Subject.Normal(TVector3(0, 0, -1 / 2.)), (TVector3(0, 0, -1)), KINDA_SMALL_NUMBER, Error);
+		EXPECT_VECTOR_NEAR_ERR(Subject.Normal(TVector3(0, 1 / 2., 0)), (TVector3(0, 1, 0)), KINDA_SMALL_NUMBER, Error);
+		EXPECT_VECTOR_NEAR_ERR(Subject.Normal(TVector3(0, -1 / 2., 0)), (TVector3(0, -1, 0)), KINDA_SMALL_NUMBER, Error);
+		EXPECT_VECTOR_NEAR_ERR(Subject.Normal(TVector3(1 / 2., 0, 0)), (TVector3(1, 0, 0)), KINDA_SMALL_NUMBER, Error);
+		EXPECT_VECTOR_NEAR_ERR(Subject.Normal(TVector3(-1 / 2., 0, 0)), (TVector3(-1, 0, 0)), KINDA_SMALL_NUMBER, Error);
+
+		// inside phi
+		EXPECT_EQ(Subject.SignedDistance(TVector3(0, 0, 1 / 2.)), -1 / 2.) << *Error;
+		EXPECT_EQ(Subject.SignedDistance(TVector3(0, 0, -1 / 2.)), -1 / 2.) << *Error;
+		EXPECT_EQ(Subject.SignedDistance(TVector3(0, 1 / 2., 0)), -1 / 2.) << *Error;
+		EXPECT_EQ(Subject.SignedDistance(TVector3(0, -1 / 2., 0)), -1 / 2.) << *Error;
+		EXPECT_EQ(Subject.SignedDistance(TVector3(1 / 2., 0, 0)), -1 / 2.) << *Error;
+		EXPECT_EQ(Subject.SignedDistance(TVector3(-1 / 2., 0, 0)), -1 / 2.) << *Error;
+	}
+
+	template<class T>
+	void UnitImplicitObjectNormalsExternal(FImplicitObject &Subject, FString Caller)
+	{
+		typedef TVector<T, 3> TVector3;
+		FString Error = FString("Called by ") + Caller + FString(".");
+
+		// outside normal 
+		EXPECT_VECTOR_NEAR_ERR(Subject.Normal(TVector3(0, 0, 3 / 2.)), (TVector3(0, 0, 1)), KINDA_SMALL_NUMBER, Error);
+		EXPECT_VECTOR_NEAR_ERR(Subject.Normal(TVector3(0, 0, -3 / 2.)), (TVector3(0, 0, -1)), KINDA_SMALL_NUMBER, Error);
+		EXPECT_VECTOR_NEAR_ERR(Subject.Normal(TVector3(0, 3 / 2., 0)), (TVector3(0, 1, 0)), KINDA_SMALL_NUMBER, Error);
+		EXPECT_VECTOR_NEAR_ERR(Subject.Normal(TVector3(0, -3 / 2., 0)), (TVector3(0, -1, 0)), KINDA_SMALL_NUMBER, Error);
+		EXPECT_VECTOR_NEAR_ERR(Subject.Normal(TVector3(3 / 2., 0, 0)), (TVector3(1, 0, 0)), KINDA_SMALL_NUMBER, Error);
+		EXPECT_VECTOR_NEAR_ERR(Subject.Normal(TVector3(-3 / 2., 0, 0)), (TVector3(-1, 0, 0)), KINDA_SMALL_NUMBER, Error);
+
+		// outside phi
+		EXPECT_EQ(Subject.SignedDistance(TVector3(0, 0, 3 / 2.)), 1 / 2.) << *Error;
+		EXPECT_EQ(Subject.SignedDistance(TVector3(0, 0, -3 / 2.)), 1 / 2.) << *Error;
+		EXPECT_EQ(Subject.SignedDistance(TVector3(0, 3 / 2., 0)), 1 / 2.) << *Error;
+		EXPECT_EQ(Subject.SignedDistance(TVector3(0, -3 / 2., 0)), 1 / 2.) << *Error;
+		EXPECT_EQ(Subject.SignedDistance(TVector3(3 / 2., 0, 0)), 1 / 2.) << *Error;
+		EXPECT_EQ(Subject.SignedDistance(TVector3(-3 / 2., 0, 0)), 1 / 2.) << *Error;
+	}
+
+
+	/* Given an ImplicitObject and an InputPoint, verifies that when that point is reflected across the surface of the object, the point of 
+	   intersection between those two points is ExpectedPoint. */
+	template<class T>
+	void TestFindClosestIntersection(FImplicitObject& Subject, TVector<T, 3> InputPoint, TVector<T, 3> ExpectedPoint, FString Caller)
+	{
+		FString Error = FString("Called by ") + Caller + FString(".");
+		typedef TVector<T, 3> TVector3;
+		T SamplePhi = Subject.SignedDistance(InputPoint);
+		TVector3 SampleNormal = Subject.Normal(InputPoint);
+		TVector3 EndPoint = InputPoint + SampleNormal * SamplePhi*-2.;
+		Pair<TVector3, bool> Result = Subject.FindClosestIntersection(InputPoint, EndPoint, KINDA_SMALL_NUMBER);
+		EXPECT_VECTOR_NEAR_ERR(Result.First, ExpectedPoint, 0.001, Error);
+	}
+
+
+	/* Takes an ImplicitObject of unit size (circumscribed inside a 2x2 cube centered on the origin).
+	   Tests the FindClosestIntersection functionality on a point near the top of the unit object. */
+	template<class T>
+	void UnitImplicitObjectIntersections(FImplicitObject &Subject, FString Caller)
+	{
+		// closest point near origin (+)
+		TestFindClosestIntersection(Subject, TVector<T, 3>(0, 0, 2), TVector<T, 3>(0, 0, 1), Caller);
+
+		// closest point near origin (-)
+		TestFindClosestIntersection(Subject, TVector<T, 3>(0, 0, 1 / 2.), TVector<T, 3>(0, 0, 1), Caller);
+	}
+
+
+	/* Takes an ImplicitObject of unit size (circumscribed inside a 2x2 cube centered on the origin).
+	   Tests the .Support() function. */
+	template<class T>
+	void UnitImplicitObjectSupportPhis(FImplicitObject &Subject, FString Caller)
+	{
+		typedef TVector<T, 3> TVector3;
+		FString Error = FString("Called by ") + Caller + FString(".");
+
+		// support phi
+		EXPECT_VECTOR_NEAR_ERR(Subject.Support(TVector3(0, 0, 1), T(0)), (TVector3(0, 0, 1)), KINDA_SMALL_NUMBER, Error);
+		EXPECT_VECTOR_NEAR_ERR(Subject.Support(TVector3(0, 0, -1), T(0)), (TVector3(0, 0, -1)), KINDA_SMALL_NUMBER, Error);
+		EXPECT_VECTOR_NEAR_ERR(Subject.Support(TVector3(0, 1, 0), T(0)), (TVector3(0, 1, 0)), KINDA_SMALL_NUMBER, Error);
+		EXPECT_VECTOR_NEAR_ERR(Subject.Support(TVector3(0, -1, 0), T(0)), (TVector3(0, -1, 0)), KINDA_SMALL_NUMBER, Error);
+		EXPECT_VECTOR_NEAR_ERR(Subject.Support(TVector3(1, 0, 0), T(0)), (TVector3(1, 0, 0)), KINDA_SMALL_NUMBER, Error);
+		EXPECT_VECTOR_NEAR_ERR(Subject.Support(TVector3(-1, 0, 0), T(0)), (TVector3(-1, 0, 0)), KINDA_SMALL_NUMBER, Error);
+
+		EXPECT_VECTOR_NEAR_ERR(Subject.Support(TVector3(0, 0, 1), T(1)), (TVector3(0, 0, 2)), KINDA_SMALL_NUMBER, Error);
+		EXPECT_VECTOR_NEAR_ERR(Subject.Support(TVector3(0, 0, -1), T(1)), (TVector3(0, 0, -2)), KINDA_SMALL_NUMBER, Error);
+		EXPECT_VECTOR_NEAR_ERR(Subject.Support(TVector3(0, 1, 0), T(1)), (TVector3(0, 2, 0)), KINDA_SMALL_NUMBER, Error);
+		EXPECT_VECTOR_NEAR_ERR(Subject.Support(TVector3(0, -1, 0), T(1)), (TVector3(0, -2, 0)), KINDA_SMALL_NUMBER, Error);
+		EXPECT_VECTOR_NEAR_ERR(Subject.Support(TVector3(1, 0, 0), T(1)), (TVector3(2, 0, 0)), KINDA_SMALL_NUMBER, Error);
+		EXPECT_VECTOR_NEAR_ERR(Subject.Support(TVector3(-1, 0, 0), T(1)), (TVector3(-2, 0, 0)), KINDA_SMALL_NUMBER, Error);
+	}
+
+
+	/* IMPLICIT OBJECT TESTS */
+
+
+	template<class T>
+	void ImplicitPlane()
+	{
+		typedef TVector<T, 3> TVector3;
+		FString Caller("ImplicitPlane()");
+
+		{// basic tests
+			TPlane<T, 3> Subject(TVector3(0), TVector3(0, 0, 1));
+
+			// check samples about the origin. 
+			EXPECT_VECTOR_NEAR_DEFAULT(Subject.Normal(TVector3(1, 1, 1)), (TVector3(0, 0, 1)));
+			EXPECT_VECTOR_NEAR_DEFAULT(Subject.Normal(TVector3(-1, -1, -1)), (TVector3(0, 0, 1)));
+
+			EXPECT_EQ(Subject.SignedDistance(TVector3(1, 1, 1)) , 1.f);
+			EXPECT_EQ(Subject.SignedDistance(TVector3(-1, -1, -1)) , -1.f);
+	
+			EXPECT_VECTOR_NEAR_DEFAULT(Subject.FindClosestPoint(TVector3(0, 0, 1)), (TVector3(0, 0, 0)));
+			EXPECT_VECTOR_NEAR_DEFAULT(Subject.FindClosestPoint(TVector3(1, 1, 2)), (TVector3(1, 1, 0)));
+			EXPECT_VECTOR_NEAR_DEFAULT(Subject.FindClosestPoint(TVector3(0, 0, -1)), (TVector3(0, 0, 0)));
+			EXPECT_VECTOR_NEAR_DEFAULT(Subject.FindClosestPoint(TVector3(1, 1, -2)), (TVector3(1, 1, 0)));
+		}
+		
+		{// closest point near origin
+			TPlane<T, 3> Subject(TVector3(0), TVector3(0, 0, 1));
+			TVector3 InputPoint = TVector3(1, 1, 1);
+			TestFindClosestIntersection(Subject, InputPoint, TVector<T, 3>(1, 1, 0), Caller);
+			EXPECT_VECTOR_NEAR_DEFAULT(Subject.FindClosestPoint(InputPoint), (TVector3(1, 1, 0)));
+		}
+
+		{// closest point single axis off origin (+)
+			TVector3 InputPoint = TVector3(0, 0, 2);
+			TPlane<T, 3> Subject = TPlane<T, 3>(TVector3(0, 0, 1), TVector3(0, 0, 1));
+			TestFindClosestIntersection(Subject, InputPoint, TVector<T, 3>(0, 0, 1), Caller);
+			EXPECT_VECTOR_NEAR(Subject.FindClosestPoint(InputPoint), FVector(0, 0, 1), 0.001);
+			EXPECT_VECTOR_NEAR(Subject.FindClosestPoint(TVector3(0, 1, 2)), FVector(0,1,1), 0.001);
+		}
+		
+		{// closest point off origin (+)
+			TVector3 InputPoint = TVector3(11,11,11);
+			TPlane<T, 3> Subject = TPlane<T, 3>(TVector3(10, 10, 10), TVector3(1, 1, 1).GetSafeNormal());
+			TestFindClosestIntersection(Subject, InputPoint, TVector<T, 3>(10, 10, 10), Caller);
+			TVector3 NearestPoint = Subject.FindClosestPoint(InputPoint); // wrong (9.26...)
+			EXPECT_VECTOR_NEAR(Subject.FindClosestPoint(InputPoint), FVector(10, 10, 10), 0.001);
+		}
+
+		{// closest point off origin (-)
+			TVector3 InputPoint = TVector3(9,9,9);
+			TPlane<T, 3>Subject = TPlane<T, 3>(TVector3(10, 10, 10), TVector3(1, 1, 1).GetSafeNormal());
+			TestFindClosestIntersection(Subject, InputPoint, TVector<T, 3>(10, 10, 10), Caller);
+			TVector3 NearestPoint = Subject.FindClosestPoint(InputPoint); // (10.73...)
+			EXPECT_VECTOR_NEAR(Subject.FindClosestPoint(InputPoint), FVector(10, 10, 10), 0.001);
+		}
+	}
+	template void ImplicitPlane<float>();
+
+
+	template<class T>
+	void ImplicitCube()
+	{
+		EXPECT_TRUE(false);
+
+		typedef TVector<T, 3> TVector3;
+		FString Caller("ImplicitCube()");
+
+		TBox<T, 3> Subject(TVector<T,3>(-1), TVector3(1));
+
+		UnitImplicitObjectNormalsInternal(Subject, Caller);
+		UnitImplicitObjectNormalsExternal(Subject, Caller);
+		UnitImplicitObjectIntersections(Subject, Caller);
+		
+		{// support phi - expects the corners for boxes
+			// Iterate through every face, edge, and corner direction, and ensure it snaps to the proper corner. 
+			for (int i0 = -1; i0 < 2; ++i0)
+			{
+				for (int i1 = -1; i1 < 2; ++i1)
+				{
+					for (int i2 = -1; i2 < 2; ++i2)
+					{
+						// If the direction is 0 or 1, it should snap to the upper corner. 
+						TVector<T, 3> Expected(1);
+						// If the direction is -1, it should snap to the lower corner. 
+						if (i0 == -1) Expected[0] = -1;
+						if (i1 == -1) Expected[1] = -1;
+						if (i2 == -1) Expected[2] = -1;
+
+						FString Error("Direction: ");
+						Error += FString::Printf(TEXT("(%d, %d, %d)"), i0, i1, i2);
+
+						EXPECT_VECTOR_NEAR_ERR(Subject.Support(TVector3(i0, i1, i2), T(0)), Expected, KINDA_SMALL_NUMBER, Error);
+					}
+				}
+			}
+
+#if RUN_KNOWN_BROKEN_TESTS
+			EXPECT_VECTOR_NEAR_DEFAULT(Subject.Support(TVector3(0, 0, 1), T(1)), (TVector3(2, 2, 2)));
+			EXPECT_VECTOR_NEAR_DEFAULT(Subject.Support(TVector3(0, 0, -1), T(1)), (TVector3(2, 2, -2)));
+			EXPECT_VECTOR_NEAR_DEFAULT(Subject.Support(TVector3(0, 1, 0), T(1)), (TVector3(2, 2, 2)));
+			EXPECT_VECTOR_NEAR_DEFAULT(Subject.Support(TVector3(0, -1, 0), T(1)), (TVector3(2, -2, 2)));
+			EXPECT_VECTOR_NEAR_DEFAULT(Subject.Support(TVector3(1, 0, 0), T(1)), (TVector3(2, 2, 2)));
+			EXPECT_VECTOR_NEAR_DEFAULT(Subject.Support(TVector3(-1, 0, 0), T(1)), (TVector3(-2, 2, 2)));
+#endif
+		}
+
+		{// support phi off origin
+			TBox<T, 3> Subject2(TVector<T, 3>(2), TVector3(4));
+
+			EXPECT_VECTOR_NEAR_DEFAULT(Subject2.Support(TVector3(0, 0, 1), T(0)), (TVector3(4, 4, 4)));
+			EXPECT_VECTOR_NEAR_DEFAULT(Subject2.Support(TVector3(0, 0, -1), T(0)), (TVector3(4, 4, 2)));
+			EXPECT_VECTOR_NEAR_DEFAULT(Subject2.Support(TVector3(0, 1, 0), T(0)), (TVector3(4, 4, 4)));
+			EXPECT_VECTOR_NEAR_DEFAULT(Subject2.Support(TVector3(0, -1, 0), T(0)), (TVector3(4, 2, 4)));
+			EXPECT_VECTOR_NEAR_DEFAULT(Subject2.Support(TVector3(0, 1, 0), T(0)), (TVector3(4, 4, 4)));
+			EXPECT_VECTOR_NEAR_DEFAULT(Subject2.Support(TVector3(1, 0, 0), T(0)), (TVector3(4, 4, 4)));
+			EXPECT_VECTOR_NEAR_DEFAULT(Subject2.Support(TVector3(-1, 0, 0), T(0)), (TVector3(2, 4, 4)));
+
+#if RUN_KNOWN_BROKEN_TESTS
+			EXPECT_VECTOR_NEAR_DEFAULT(Subject2.Support(TVector3(0, 0, 1), T(1)), (TVector3(5, 5, 5)));
+			EXPECT_VECTOR_NEAR_DEFAULT(Subject2.Support(TVector3(0, 0, -1), T(1)), (TVector3(5, 5, 1)));
+			EXPECT_VECTOR_NEAR_DEFAULT(Subject2.Support(TVector3(0, 1, 0), T(1)), (TVector3(5, 5, 5)));
+			EXPECT_VECTOR_NEAR_DEFAULT(Subject2.Support(TVector3(0, -1, 0), T(1)), (TVector3(5, 1, 5)));
+			EXPECT_VECTOR_NEAR_DEFAULT(Subject2.Support(TVector3(1, 0, 0), T(1)), (TVector3(5, 5, 5)));
+			EXPECT_VECTOR_NEAR_DEFAULT(Subject2.Support(TVector3(-1, 0, 0), T(1)), (TVector3(1, 5, 5)));
+#endif
+		}
+
+		// intersection
+		EXPECT_TRUE(Subject.Intersects(TAABB<T, 3>(TVector3(0.5), TVector3(1.5))));
+		EXPECT_FALSE(Subject.Intersects(TAABB<T, 3>(TVector3(2), TVector3(3))));
+
+		{// closest point near origin (+)
+			TVector<T, 3> InputPoint(0, 0, 2);
+			EXPECT_VECTOR_NEAR(Subject.FindClosestPoint(InputPoint), FVector(0,0,1), 0.001);
+			EXPECT_VECTOR_NEAR(Subject.FindClosestPoint(TVector3(3 / 2., 0, 0)), FVector(1,0,0), 0.001);
+		}
+
+		{// closest point near origin (-)
+			TVector<T, 3> InputPoint(0, 0, 1 / 2.);
+			EXPECT_VECTOR_NEAR(Subject.FindClosestPoint(InputPoint), FVector(0, 0, 1), 0.001);
+			EXPECT_VECTOR_NEAR(Subject.FindClosestPoint(TVector3(3 / 4., 0, 0)), FVector(1, 0, 0), 0.001);
+			EXPECT_FALSE(Subject.FindClosestPoint(TVector<T, 3>(0, 0, 0)).Equals(TVector<T, 3>(0)));
+			EXPECT_EQ(Subject.FindClosestPoint(TVector3(0, 0, 0)).Size(),1.0);
+		}
+
+		{// diagonal 3-corner case
+			TAABB<T, 3> Subject2(TVector3(-1), TVector3(1));
+			// outside
+			EXPECT_VECTOR_NEAR(Subject2.FindClosestPoint(TVector3(2, 2, 2)), FVector(1,1,1), 0.001);
+			EXPECT_VECTOR_NEAR(Subject2.FindClosestPoint(TVector3(-2, -2, -2)), FVector(-1, -1, -1), 0.001);
+			EXPECT_VECTOR_NEAR(Subject2.FindClosestPoint(TVector3(3 / 2., 3 / 2., 3 / 2.)), FVector(1, 1, 1), 0.001);
+			EXPECT_VECTOR_NEAR(Subject2.FindClosestPoint(TVector3(-3 / 2., 3 / 2., -3 / 2.)), FVector(-1, 1, -1), 0.001);
+			// inside
+			EXPECT_VECTOR_NEAR(Subject2.FindClosestPoint(TVector3(1 / 2., 1 / 2., 1 / 2.)), FVector(1, 1, 1), 0.001);
+			EXPECT_VECTOR_NEAR(Subject2.FindClosestPoint(TVector3(1 / 2., -1 / 2., 1 / 2.)), FVector(1, -1, 1), 0.001);
+		}
+
+		{// diagonal 2-corner case
+			TAABB<T, 3> Subject2(TVector3(-1), TVector3(1));
+			TVector3 test1 = Subject.FindClosestPoint(TVector3(2, 2, 0));
+			// outside
+			EXPECT_VECTOR_NEAR(Subject2.FindClosestPoint(TVector3(2, 2, 0)), FVector(1, 1, 0), 0.001);
+			EXPECT_VECTOR_NEAR(Subject2.FindClosestPoint(TVector3(0, 3 / 2., 3 / 2.)), FVector(0, 1, 1), 0.001);
+			// inside
+			EXPECT_VECTOR_NEAR(Subject2.FindClosestPoint(TVector3(1 / 2., 1 / 2., 0)), FVector(1, 1, 0), 0.001);
+			EXPECT_VECTOR_NEAR(Subject2.FindClosestPoint(TVector3(-1 / 2., 1 / 2., 0)), FVector(-1, 1, 0), 0.001);
+		}
+
+		{// closest point off origin (+)
+			TBox<T, 3> Subject2(TVector3(2), TVector3(4));
+			TVector3 InputPoint(5, 5, 5);
+			TestFindClosestIntersection(Subject2, InputPoint, TVector<T, 3>(4, 4, 4), Caller);
+
+			EXPECT_VECTOR_NEAR(Subject2.FindClosestPoint(InputPoint), FVector(4,4,4), 0.001);
+			TVector3 test2 = Subject2.FindClosestPoint(TVector3(3.5, 3.5, 3.5));
+			EXPECT_VECTOR_NEAR(Subject2.FindClosestPoint(TVector3(3.5,3.5,3.5)), FVector(4,4,4), 0.001);
+		}
+
+#if RUN_KNOWN_BROKEN_TESTS
+		{// different defining corners of the box
+			// Ensure fails in PhiWithNormal
+			TBox<T, 3> Test1(TVector<T, 3>(-1, -1, 0), TVector3(1, 1, -1));
+			EXPECT_VECTOR_NEAR(Test1.Normal(TVector3(0, 0, -2 / 3.)), (TVector3(0, 0, -1)), KINDA_SMALL_NUMBER);
+
+			// Ensure fails in PhiWithNormal
+			TBox<T, 3> Test2(TVector<T, 3>(1, 1, -1), TVector3(-1, -1, 0));
+			EXPECT_VECTOR_NEAR(Test2.Normal(TVector3(0, 0, -2 / 3.)), (TVector3(0, 0, -1)), KINDA_SMALL_NUMBER);
+
+			// Ensure fails in PhiWithNormal
+			TBox<T, 3> Test3(TVector<T, 3>(1, 1, 0), TVector3(-1, -1, -1));
+			EXPECT_VECTOR_NEAR(Test3.Normal(TVector3(0, 0, -2 / 3.)), (TVector3(0, 0, -1)), KINDA_SMALL_NUMBER);
+
+			// Works fine!
+			TBox<T, 3> Test4(TVector<T, 3>(-1, -1, -1), TVector3(1, 1, 0));
+			EXPECT_VECTOR_NEAR(Test4.Normal(TVector3(0, 0, -2 / 3.)), (TVector3(0, 0, -1)), KINDA_SMALL_NUMBER);
+		}
+#endif
+	}
+	template void ImplicitCube<float>();
+	
+
+	template<class T>
+	void ImplicitSphere()
+	{
+		typedef TVector<T, 3> TVector3;
+		FString Caller("ImplicitSphere()");
+
+		TSphere<T, 3> Subject(TVector<T,3>(0), 1);
+		UnitImplicitObjectNormalsInternal(Subject, Caller);
+		UnitImplicitObjectNormalsExternal(Subject, Caller);
+		UnitImplicitObjectIntersections(Subject, Caller);
+		UnitImplicitObjectSupportPhis(Subject, Caller);
+
+		// intersection
+		EXPECT_TRUE(Subject.Intersects(TSphere<T, 3>(TVector<T,3>(0.f), 2.f)));
+		EXPECT_TRUE(Subject.Intersects(TSphere<T, 3>(TVector3(.5f), 1.f)));
+		EXPECT_FALSE(Subject.Intersects(TSphere<T, 3>(TVector<T,3>(2.f), 1.f)));
+
+		{// closest point near origin (+)
+			TVector<T, 3> InputPoint(0, 0, 2.);
+			EXPECT_VECTOR_NEAR(Subject.FindClosestPoint(InputPoint), FVector(0, 0, 1), 0.001);
+			EXPECT_VECTOR_NEAR(Subject.FindClosestPoint(TVector3(3 / 2., 0, 0)), FVector(1, 0, 0), 0.001);
+		}
+
+		{// closest point near origin (-)
+			TVector<T, 3> InputPoint(0, 0, 1 / 2.);
+			EXPECT_VECTOR_NEAR_DEFAULT(Subject.FindClosestPoint(TVector3(0, 0, 0)), TVector3(0));
+			EXPECT_VECTOR_NEAR(Subject.FindClosestPoint(InputPoint), FVector(0, 0, 1), 0.001);
+			EXPECT_VECTOR_NEAR(Subject.FindClosestPoint(TVector3(3 / 4., 0, 0)), FVector(1, 0, 0), 0.001);
+		}
+
+		{// closest point off origin (+)
+			TSphere<T, 3> Subject2(TVector3(2), 2);
+			TVector3 InputPoint(2, 2, 5);
+			TestFindClosestIntersection(Subject2, InputPoint, TVector<T, 3>(2, 2, 4), Caller);
+			EXPECT_VECTOR_NEAR(Subject2.FindClosestPoint(InputPoint), FVector(2, 2, 4), 0.001);
+			EXPECT_VECTOR_NEAR(Subject2.FindClosestPoint(TVector3(2, 2, 3.5)), FVector(2, 2, 4), 0.001);
+		}
+	}
+	template void ImplicitSphere<float>();
+
+
+	/* Cylinder Helpers */
+
+	// Expects a unit cylinder. 
+	template<class T>
+	void CheckCylinderEdgeBehavior(FImplicitObject &Subject, FString Caller)
+	{
+		typedef TVector<T, 3> TVector3;
+		FString Error = FString("Called by ") + Caller + FString(".");
+
+		// inside normal
+		// defaults to side of cylinder when equally close to side and endcap
+		EXPECT_VECTOR_NEAR_ERR(Subject.Normal(TVector3(0, 1 / 2., 1 / 2.)), TVector3(0, 1, 0), KINDA_SMALL_NUMBER, Error);
+		EXPECT_VECTOR_NEAR_ERR(Subject.Normal(TVector3(0, 1 / 3., 1 / 2.)), TVector3(0, 0, 1), KINDA_SMALL_NUMBER, Error);
+		EXPECT_VECTOR_NEAR_ERR(Subject.Normal(TVector3(0, 1 / 2., -1 / 2.)), TVector3(0, 1, 0), KINDA_SMALL_NUMBER, Error);
+		EXPECT_VECTOR_NEAR_ERR(Subject.Normal(TVector3(0, 1 / 3., -1 / 2.)), TVector3(0, 0, -1), KINDA_SMALL_NUMBER, Error);
+
+		// outside normal 		
+		// defaults to endcap of cylinder above intersection of side and endcap
+		EXPECT_VECTOR_NEAR_ERR(Subject.Normal(TVector3(0, 1., 3 / 2.)), TVector3(0, 0, 1), KINDA_SMALL_NUMBER, Error);
+		EXPECT_VECTOR_NEAR_ERR(Subject.Normal(TVector3(0, 1., -3 / 2.)), TVector3(0, 0, -1), KINDA_SMALL_NUMBER, Error);
+		// defaults to side of cylinder next to intersection of side and endcap
+		EXPECT_VECTOR_NEAR_ERR(Subject.Normal(TVector3(0, 3 / 2., 1.)), TVector3(0, 1, 0), KINDA_SMALL_NUMBER, Error);
+		EXPECT_VECTOR_NEAR_ERR(Subject.Normal(TVector3(0, 3 / 2., -1.)), TVector3(0, 1, 0), KINDA_SMALL_NUMBER, Error);
+
+		//inside phi
+		EXPECT_EQ(Subject.SignedDistance(TVector<T, 3>(0, 1, 3 / 2.)), 1 / 2.) << *Error;
+		EXPECT_EQ(Subject.SignedDistance(TVector<T, 3>(0, 1, -3 / 2.)), 1 / 2.) << *Error;
+		EXPECT_EQ(Subject.SignedDistance(TVector<T, 3>(0, -1, 3 / 2.)), 1 / 2.) << *Error;
+		EXPECT_EQ(Subject.SignedDistance(TVector<T, 3>(0, -1, -3 / 2.)), 1 / 2.) << *Error;
+	}
+
+
+	// Expects a cylinder with endcap points (1,1,1) and (-1,-1,-1), radius 1.
+	template<class T>
+	void TiltedUnitImplicitCylinder(FImplicitObject &Subject, FString Caller)
+	{
+		typedef TVector<T, 3> TVector3;
+		FString Error = FString("Called by ") + Caller + FString(".");
+
+		// inside normals
+		EXPECT_VECTOR_NEAR_ERR(Subject.Normal(TVector3(1 / 2., 1 / 2., 1 / 2.)), TVector3(1, 1, 1).GetSafeNormal(), KINDA_SMALL_NUMBER, Error);
+		EXPECT_VECTOR_NEAR_ERR(Subject.Normal(TVector3(-1 / 2., -1 / 2., -1 / 2.)), TVector3(-1, -1, -1).GetSafeNormal(), KINDA_SMALL_NUMBER, Error);
+
+		EXPECT_VECTOR_NEAR_ERR(Subject.Normal(TVector3(0., 1 / 2., -1 / 2.)), TVector3(0, 1, -1).GetSafeNormal(), KINDA_SMALL_NUMBER, Error);
+		EXPECT_VECTOR_NEAR_ERR(Subject.Normal(TVector3(0., -1 / 2., 1 / 2.)), TVector3(0, -1, 1).GetSafeNormal(), KINDA_SMALL_NUMBER, Error);
+		EXPECT_VECTOR_NEAR_ERR(Subject.Normal(TVector3(1 / 2., 0., -1 / 2.)), TVector3(1, 0, -1).GetSafeNormal(), KINDA_SMALL_NUMBER, Error);
+		EXPECT_VECTOR_NEAR_ERR(Subject.Normal(TVector3(-1 / 2., 0., 1 / 2.)), TVector3(-1, 0, 1).GetSafeNormal(), KINDA_SMALL_NUMBER, Error);
+
+		//outside normals
+		EXPECT_VECTOR_NEAR_ERR(Subject.Normal(TVector3(3 / 2., 3 / 2., 3 / 2.)), TVector3(1, 1, 1).GetSafeNormal(), KINDA_SMALL_NUMBER, Error);
+		EXPECT_VECTOR_NEAR_ERR(Subject.Normal(TVector3(-3 / 2., -3 / 2., -3 / 2.)), TVector3(-1, -1, -1).GetSafeNormal(), KINDA_SMALL_NUMBER, Error);
+
+		EXPECT_VECTOR_NEAR_ERR(Subject.Normal(TVector3(0., 3 / 2., -3 / 2.)), TVector3(0, 1, -1).GetSafeNormal(), KINDA_SMALL_NUMBER, Error);
+		EXPECT_VECTOR_NEAR_ERR(Subject.Normal(TVector3(0., -3 / 2., 3 / 2.)), TVector3(0, -1, 1).GetSafeNormal(), KINDA_SMALL_NUMBER, Error);
+		EXPECT_VECTOR_NEAR_ERR(Subject.Normal(TVector3(3 / 2., 0., -3 / 2.)), TVector3(1, 0, -1).GetSafeNormal(), KINDA_SMALL_NUMBER, Error);
+		EXPECT_VECTOR_NEAR_ERR(Subject.Normal(TVector3(-3 / 2., 0., 3 / 2.)), TVector3(-1, 0, 1).GetSafeNormal(), KINDA_SMALL_NUMBER, Error);
+
+		// inside phi
+		EXPECT_EQ(Subject.SignedDistance(TVector3(1 / 2., 1 / 2., 1 / 2.)), -TVector3(1 / 2.).Size()) << *Error;
+		EXPECT_EQ(Subject.SignedDistance(TVector3(-1 / 2., -1 / 2., -1 / 2.)), -TVector3(1 / 2.).Size()) << *Error;
+		EXPECT_NEAR(Subject.SignedDistance(TVector3(0., sqrt(2) / 4., -sqrt(2) / 4.)), -1 / 2., KINDA_SMALL_NUMBER) << *Error;
+		EXPECT_NEAR(Subject.SignedDistance(TVector3(0., -sqrt(2) / 4., sqrt(2) / 4.)), -1 / 2., KINDA_SMALL_NUMBER) << *Error;
+		EXPECT_NEAR(Subject.SignedDistance(TVector3(sqrt(2) / 4., 0., -sqrt(2) / 4.)), -1 / 2., KINDA_SMALL_NUMBER) << *Error;
+		EXPECT_NEAR(Subject.SignedDistance(TVector3(-sqrt(2) / 4., 0., sqrt(2) / 4.)), -1 / 2., KINDA_SMALL_NUMBER) << *Error;
+
+		// outside phi
+		EXPECT_EQ(Subject.SignedDistance(TVector3(3 / 2., 3 / 2., 3 / 2.)), TVector3(1 / 2.).Size()) << *Error;
+		EXPECT_EQ(Subject.SignedDistance(TVector3(-3 / 2., -3 / 2., -3 / 2.)), TVector3(1 / 2.).Size()) << *Error;
+		EXPECT_NEAR(Subject.SignedDistance(TVector3(0., 3 * sqrt(2) / 4., -3 * sqrt(2) / 4.)), 1 / 2., KINDA_SMALL_NUMBER) << *Error;
+		EXPECT_NEAR(Subject.SignedDistance(TVector3(0., -3 * sqrt(2) / 4., 3 * sqrt(2) / 4.)), 1 / 2., KINDA_SMALL_NUMBER) << *Error;
+		EXPECT_NEAR(Subject.SignedDistance(TVector3(3 * sqrt(2) / 4., 0., -3 * sqrt(2) / 4.)), 1 / 2., KINDA_SMALL_NUMBER) << *Error;
+		EXPECT_NEAR(Subject.SignedDistance(TVector3(-3 * sqrt(2) / 4., 0., 3 * sqrt(2) / 4.)), 1 / 2., KINDA_SMALL_NUMBER) << *Error;
+	}
+
+	/* End Cylinder Helpers */
+
+	template<class T>
+	void ImplicitCylinder()
+	{
+		typedef TVector<T, 3> TVector3;
+		FString Caller("ImplicitCylinder()");
+
+		// unit cylinder tests
+		TCylinder<T> Subject(TVector<T, 3>(0, 0, 1), TVector<T, 3>(0, 0, -1), 1);
+		UnitImplicitObjectNormalsInternal(Subject, Caller);
+		UnitImplicitObjectNormalsExternal(Subject, Caller);
+		UnitImplicitObjectIntersections(Subject, Caller);
+		CheckCylinderEdgeBehavior(Subject, Caller);
+
+		// tilted tests
+		TCylinder<T> SubjectTilted(TVector3(1), TVector3(-1), 1);
+		TiltedUnitImplicitCylinder(SubjectTilted, Caller);
+
+#if RUN_KNOWN_BROKEN_TESTS
+		{// nearly flat cylinder tests (BROKEN)
+			TCylinder<T> SubjectFlat(TVector3(0, 0, KINDA_SMALL_NUMBER), TVector3(0, 0, -KINDA_SMALL_NUMBER), 1);
+			EXPECT_VECTOR_NEAR_DEFAULT(Subject.Normal(TVector3(0, 0, 1 / 2.)), TVector3(0, 0, 1));
+			EXPECT_VECTOR_NEAR_DEFAULT(Subject.Normal(TVector3(0, 0, -1 / 2.)), TVector3(0, 0, -1));
+			EXPECT_EQ(Subject.SignedDistance(TVector3(0, 0, 1 / 2.)), 1 / 2.);
+			EXPECT_EQ(Subject.SignedDistance(TVector3(0, 0, -1 / 2.)), 1 / 2.);
+			Pair<TVector3, bool> Result = SubjectFlat.FindClosestIntersection(TVector3(0, 1, 1), TVector3(0, -1, -1), KINDA_SMALL_NUMBER);
+			EXPECT_FALSE(Result.Second);
+		}
+#endif
+
+		{// closest point off origin (+)
+			TCylinder<T> Subject2(TVector<T, 3>(2,2,4), TVector<T, 3>(2,2,0), 2);
+			TVector<T, 3> InputPoint(2, 2, 5);
+			TestFindClosestIntersection(Subject2, InputPoint, TVector<T, 3>(2, 2, 4), Caller);
+		}
+
+		{// closest point off origin (-)
+			TCylinder<T> Subject2(TVector<T, 3>(2, 2, 4), TVector<T, 3>(2, 2, 0), 2);
+			TVector<T, 3> InputPoint(2, 3, 2);
+			TestFindClosestIntersection(Subject2, InputPoint, TVector<T, 3>(2, 4, 2), Caller);
+		}
+
+		{// near edge intersection
+			TCylinder<T> Cylinder(TVector<T, 3>(1, 1, -14), TVector<T, 3>(1, 1, 16), 15);
+			Pair<TVector3, bool> Result = Cylinder.FindClosestIntersection(TVector<T, 3>(16, 16, 1), TVector3(16, -16, 1), 0);
+			EXPECT_TRUE(Result.Second);
+			EXPECT_VECTOR_NEAR(Result.First, TVector3(16, 1, 1), KINDA_SMALL_NUMBER);
+		}
+	}
+	template void ImplicitCylinder<float>();
+
+
+	template<class T>
+	void ImplicitTaperedCylinder()
+	{
+		typedef TVector<T, 3> TVector3;
+		FString Caller("ImplicitTaperedCylinder()");
+
+		// unit tapered cylinder tests
+		TTaperedCylinder<T> Subject(TVector<T, 3>(0, 0, 1), TVector<T, 3>(0, 0, -1), 1, 1);
+		UnitImplicitObjectNormalsInternal(Subject, Caller);
+		UnitImplicitObjectNormalsExternal(Subject, Caller);
+		UnitImplicitObjectIntersections(Subject, Caller);
+		CheckCylinderEdgeBehavior(Subject, Caller);
+
+		// tilted tapered cylinder tests
+		TTaperedCylinder<T> SubjectTilted(TVector<T,3>(1), TVector<T,3>(-1), 1, 1);
+		TiltedUnitImplicitCylinder(SubjectTilted, Caller);
+
+		TTaperedCylinder<T> SubjectCone(TVector<T, 3>(0, 0, 1), TVector<T, 3>(0, 0, 0), 0, 1);
+
+		// inside normals 
+		EXPECT_VECTOR_NEAR_DEFAULT(SubjectCone.Normal(TVector3(0, 0, 0)), TVector3(0, 0, -1));
+		EXPECT_VECTOR_NEAR_DEFAULT(SubjectCone.Normal(TVector3(0, 0, 1)), TVector3(0, 0, 1));
+		
+		// Note: tapered cylinders always return normals parallel to the endcap planes when calculating for points near/on the body,
+		// very much like a normal cylinder. The slant is ignored. 
+		EXPECT_VECTOR_NEAR_DEFAULT(SubjectCone.Normal(TVector3(0, 1 / 3., 1 / 3.)),  TVector3(0, 1, 0));
+		EXPECT_VECTOR_NEAR_DEFAULT(SubjectCone.Normal(TVector3(1 / 3., 0, 1 / 3.)),  TVector3(1, 0, 0)); 
+		EXPECT_VECTOR_NEAR_DEFAULT(SubjectCone.Normal(TVector3(0, -1 / 3., 1 / 3.)), TVector3(0, -1, 0)); 
+		EXPECT_VECTOR_NEAR_DEFAULT(SubjectCone.Normal(TVector3(-1 / 3., 0, 1 / 3.)), TVector3(-1, 0, 0)); 
+		EXPECT_VECTOR_NEAR(SubjectCone.Normal(TVector3(1 / 3., 1 / 3., 1 / 2.)), TVector3(0.707, 0.707, 0), 0.001); 
+
+		// outside normals
+		EXPECT_VECTOR_NEAR_DEFAULT(SubjectCone.Normal(TVector3(0, 0, -1 / 2.)), TVector3(0, 0, -1));
+		EXPECT_VECTOR_NEAR_DEFAULT(SubjectCone.Normal(TVector3(0, 0,  3 / 2.)), TVector3(0, 0, 1));
+		EXPECT_VECTOR_NEAR_DEFAULT(SubjectCone.Normal(TVector3( 0,  1, 1 / 2.)), TVector3(0, 1, 0));
+		EXPECT_VECTOR_NEAR_DEFAULT(SubjectCone.Normal(TVector3( 1,  0, 1 / 2.)), TVector3(1, 0, 0));
+		EXPECT_VECTOR_NEAR_DEFAULT(SubjectCone.Normal(TVector3( 0, -1, 1 / 2.)), TVector3(0, -1, 0));
+		EXPECT_VECTOR_NEAR_DEFAULT(SubjectCone.Normal(TVector3(-1,  0, 1 / 2.)), TVector3(-1, 0, 0));
+
+		{// closest point off origin (+)
+			TTaperedCylinder<T> Subject2(TVector<T, 3>(2, 2, 4), TVector<T, 3>(2, 2, 0), 2, 2);
+			TVector<T, 3> InputPoint(2, 2, 5);
+			TestFindClosestIntersection(Subject2, InputPoint, TVector<T, 3>(2, 2, 4), Caller);
+		}
+
+		{// closest point off origin (-)
+			TTaperedCylinder<T> Subject2(TVector<T, 3>(2, 2, 4), TVector<T, 3>(2, 2, 0), 2, 2);
+			TVector<T, 3> InputPoint(2, 3, 2);
+			TestFindClosestIntersection(Subject2, InputPoint, TVector<T, 3>(2, 4, 2), Caller);
+		}
+	}
+	template void ImplicitTaperedCylinder<float>();
+
+
+	template<class T>
+	void ImplicitCapsule()
+	{
+		typedef TVector<T, 3> TVector3;
+		FString Caller("ImplicitCapsule()");
+
+		// Effectively a sphere - flat cylinder with two radius 1 spheres overlapping at origin.
+		TCapsule<T> SubjectUnit(TVector<T, 3>(0, 0, 0), TVector<T, 3>(0, 0, 0), 1);
+
+		UnitImplicitObjectNormalsInternal(SubjectUnit, Caller);
+		UnitImplicitObjectNormalsExternal(SubjectUnit, Caller);
+		UnitImplicitObjectSupportPhis(SubjectUnit, Caller);
+
+#if RUN_KNOWN_BROKEN_TESTS
+		// FindClosestIntersection broken with cylinder size 0
+		UnitImplicitObjectIntersections(SubjectUnit, Caller); 
+#endif
+
+		TCapsule<T> Subject(TVector3(0, 0, 1), TVector3(0, 0, -1), 1);
+
+		{// closest point near origin (+)
+			TVector3 InputPoint(0, 0, 3);
+			TestFindClosestIntersection(Subject, InputPoint, TVector<T, 3>(0, 0, 2), Caller);
+		}
+		
+		{// closest point near origin (-)
+			TVector3 InputPoint(0, 0, 3 / 2.);
+			// Equally close to inner cylinder and top sphere - defaults to sphere. 
+			TestFindClosestIntersection(Subject, InputPoint, TVector<T, 3>(0, 0, 2), Caller);
+		}
+
+		{// closest point off origin (+)
+			TCapsule<T> Subject2(TVector3(5, 4, 4), TVector3(3, 4, 4), 1);
+			TVector3 InputPoint(4, 4, 6);
+			TestFindClosestIntersection(Subject2, InputPoint, TVector<T, 3>(4, 4, 5), Caller);
+		}
+
+		{// closest point off origin (-)
+			TCapsule<T> Subject2(TVector3(5, 4, 4), TVector3(3, 4, 4), 1);
+			TVector3 InputPoint(4, 4, 4 + 1 / 2.);
+			TestFindClosestIntersection(Subject2, InputPoint, TVector<T, 3>(4, 4, 5), Caller);
+		}
+	}
+	template void ImplicitCapsule<float>();
+
+	
+	template <typename T>
+	void ImplicitScaled()
+	{
+		typedef TVector<T, 3> TVector3;
+		TUniquePtr<TBox<T, 3>> UnitCube = MakeUnique<TBox<T, 3>>(TVector<T, 3>(-1), TVector<T, 3>(1));
+		TImplicitObjectScaled<T, 3> UnitUnscaled(MakeSerializable(UnitCube), TVector<T, 3>(1));
+		UnitImplicitObjectNormalsInternal(UnitUnscaled, FString("ImplicitTransformed()"));
+		UnitImplicitObjectNormalsExternal(UnitUnscaled, FString("ImplicitTransformed()"));
+		UnitImplicitObjectIntersections(UnitUnscaled, FString("ImplicitTransformed()"));
+
+		TUniquePtr<TSphere<T, 3>> Sphere = MakeUnique<TSphere<T, 3>>(TVector<T, 3>(3, 0, 0), 5);
+		TImplicitObjectScaled<T, 3> Unscaled(MakeSerializable(Sphere), TVector<T, 3>(1));
+		TImplicitObjectScaled<T, 3> UniformScale(MakeSerializable(Sphere), TVector<T, 3>(2));
+		TImplicitObjectScaled<T, 3> NonUniformScale(MakeSerializable(Sphere), TVector<T, 3>(2, 1, 1));
+
+		{//phi
+			const TVector<T, 3> NearEdge(7.5, 0, 0);
+			TVector<T, 3> UnscaledNormal;
+			const T UnscaledPhi = Unscaled.PhiWithNormal(NearEdge, UnscaledNormal);
+			EXPECT_FLOAT_EQ(UnscaledPhi, -0.5);
+			EXPECT_VECTOR_NEAR(UnscaledNormal, TVector3(1, 0, 0), 0);
+
+			TVector<T, 3> ScaledNormal;
+			T ScaledPhi = UniformScale.PhiWithNormal(NearEdge, ScaledNormal);
+			EXPECT_FLOAT_EQ(ScaledPhi, -(16 - 7.5));
+			EXPECT_VECTOR_NEAR(ScaledNormal, TVector3(1, 0, 0), 0);
+
+			const TVector<T, 3> NearTop(6, 0, 4.5);
+			ScaledPhi = UniformScale.PhiWithNormal(NearTop, ScaledNormal);
+			EXPECT_FLOAT_EQ(ScaledPhi, -(10 - 4.5));
+			EXPECT_VECTOR_NEAR(ScaledNormal, TVector3(0, 0, 1), 0);
+
+			ScaledPhi = NonUniformScale.PhiWithNormal(NearTop, ScaledNormal);
+			EXPECT_FLOAT_EQ(ScaledPhi, -0.5);
+			EXPECT_VECTOR_NEAR(ScaledNormal, TVector3(0, 0, 1), 0);
+		}
+		
+		{//support
+			const TVector<T, 3> DirX(1, 0, 0);
+			TVector<T, 3> SupportPt = Unscaled.Support(DirX, 1);
+			EXPECT_VECTOR_NEAR(SupportPt, TVector3(9, 0, 0), 0);
+
+			SupportPt = UniformScale.Support(DirX, 1);
+			EXPECT_VECTOR_NEAR(SupportPt, TVector3(17, 0, 0), 0);
+
+			const TVector<T, 3> DirZ(0, 0, -1);
+			SupportPt = UniformScale.Support(DirZ, 1);
+			EXPECT_VECTOR_NEAR(SupportPt, TVector3(6, 0, -11), 0);
+
+			SupportPt = NonUniformScale.Support(DirX, 1);
+			EXPECT_VECTOR_NEAR(SupportPt, TVector3(17, 0, 0), 0);
+
+			SupportPt = NonUniformScale.Support(DirZ, 1);
+			EXPECT_VECTOR_NEAR(SupportPt, TVector3(6, 0, -6), 0);
+		}
+
+		{// closest intersection
+			Pair<TVector3, bool> Result;
+			Result = Unscaled.FindClosestIntersection(TVector3(7.5, 0, 0), TVector3(8.5, 0, 0), KINDA_SMALL_NUMBER);
+			EXPECT_VECTOR_NEAR(Result.First, TVector3(8, 0, 0), 0.001);
+
+			Result = UniformScale.FindClosestIntersection(TVector3(15.5, 0, 0), TVector3(16.5, 0, 0), KINDA_SMALL_NUMBER);
+			EXPECT_VECTOR_NEAR(Result.First, TVector3(16, 0, 0), 0.001);
+
+			Result = NonUniformScale.FindClosestIntersection(TVector3(6, 0, 4.5), TVector3(6, 0, 5.5), KINDA_SMALL_NUMBER);
+			EXPECT_VECTOR_NEAR(Result.First, TVector3(6, 0, 5), 0.001);
+		}
+	}
+	template void ImplicitScaled<float>();
+
+
+	template <class T>
+	void ImplicitTransformed()
+	{
+		typedef TVector<T, 3> TVector3;
+		
+		TUniquePtr<TBox<T, 3>> UnitCube = MakeUnique<TBox<T, 3>>(TVector<T, 3>(-1), TVector<T, 3>(1));
+		TImplicitObjectTransformed<T, 3> UnitUnrotated(MakeSerializable(UnitCube), TRigidTransform<T, 3>(TVector<T, 3>(0), TRotation<T, 3>(TVector<T, 3>(0), 0)));
+		UnitImplicitObjectNormalsInternal(UnitUnrotated, FString("ImplicitTransformed()"));
+		UnitImplicitObjectNormalsExternal(UnitUnrotated, FString("ImplicitTransformed()"));
+		UnitImplicitObjectIntersections(UnitUnrotated, FString("ImplicitTransformed()"));
+		
+		// Rotate 45 degrees around z axis @ origin.
+		TImplicitObjectTransformed<T, 3> UnitRotated(MakeSerializable(UnitCube), TRigidTransform<T, 3>(TVector<T, 3>(0), TRotation<T, 3>(TVector<T, 3>(0, 0, sin(.3927)), cos(.3927))));
+		
+		{// unit rotated normals
+			TVector<T, 3> Normal;
+			T TestPhi = UnitRotated.PhiWithNormal(TVector3(1 / 2., 1 / 2., 0), Normal);
+			EXPECT_VECTOR_NEAR_DEFAULT(Normal, TVector3(sqrt(2) / 2., sqrt(2) / 2., 0));
+			TestPhi = UnitRotated.PhiWithNormal(TVector3(-1 / 2., 1 / 2., 0), Normal);
+			EXPECT_VECTOR_NEAR_DEFAULT(Normal, TVector3(-sqrt(2) / 2., sqrt(2) / 2., 0));
+			TestPhi = UnitRotated.PhiWithNormal(TVector3(1 / 2., -1 / 2., 0), Normal);
+			EXPECT_VECTOR_NEAR_DEFAULT(Normal, TVector3(sqrt(2) / 2., -sqrt(2) / 2., 0));
+			TestPhi = UnitRotated.PhiWithNormal(TVector3(-1 / 2., -1 / 2., 0), Normal);
+			EXPECT_VECTOR_NEAR_DEFAULT(Normal, TVector3(-sqrt(2) / 2., -sqrt(2) / 2., 0));
+		}
+
+		TUniquePtr<TBox<T, 3>> Cube = MakeUnique<TBox<T, 3>>(TVector<T, 3>(-2, -5, -5), TVector<T,3>(8, 5, 5));
+		TImplicitObjectTransformed<T, 3> Untransformed(MakeSerializable(Cube), TRigidTransform<T, 3>(TVector<T, 3>(0), TRotation<T, 3>(TVector<T, 3>(0), 0)));
+		TImplicitObjectTransformed<T, 3> Translated(MakeSerializable(Cube), TRigidTransform<T, 3>(TVector<T, 3>(4, 0, 0), TRotation<T, 3>(TVector<T, 3>(0), 0)));
+		
+		// Rotate 90 degrees around z axis @ origin. 
+		float rad_45 = FMath::DegreesToRadians(45);
+		TImplicitObjectTransformed<T, 3> Rotated(MakeSerializable(Cube), TRigidTransform<T, 3>(TVector<T, 3>(0), TRotation<T, 3>(TVector<T, 3>(0, 0, sin(rad_45)), cos(rad_45))));
+		TImplicitObjectTransformed<T, 3> Transformed(MakeSerializable(Cube), TRigidTransform<T, 3>(TVector<T, 3>(4, 0, 0), TRotation<T, 3>(TVector<T, 3>(0, 0, sin(rad_45)), cos(rad_45))));
+
+		{// phi
+			const TVector<T, 3> NearEdge(7.5, 0, 0);
+			TVector<T, 3> UntransformedNormal;
+			const T UntransformedPhi = Untransformed.PhiWithNormal(NearEdge, UntransformedNormal);
+			EXPECT_FLOAT_EQ(UntransformedPhi, -0.5);
+			EXPECT_VECTOR_NEAR_DEFAULT(UntransformedNormal, TVector3(1, 0, 0));
+
+			TVector<T, 3> TransformedNormal;
+			T TranslatedPhi = Translated.PhiWithNormal(NearEdge, TransformedNormal);
+			EXPECT_FLOAT_EQ(TranslatedPhi, -(0.5 + 4));
+			EXPECT_VECTOR_NEAR_DEFAULT(TransformedNormal, TVector3(1, 0, 0));
+
+			const TVector<T, 3> NearEdgeRotated(0, 7.5, 0);
+			T RotatedPhi = Rotated.PhiWithNormal(NearEdgeRotated, TransformedNormal);
+			EXPECT_FLOAT_EQ(RotatedPhi, -0.5);
+			EXPECT_VECTOR_NEAR_DEFAULT(TransformedNormal, TVector3(0, 1, 0));
+
+			T TransformedPhi = Transformed.PhiWithNormal(NearEdge, TransformedNormal);
+			EXPECT_FLOAT_EQ(TransformedPhi, -(0.5 + 1));
+			EXPECT_VECTOR_NEAR_DEFAULT(TransformedNormal, TVector3(1, 0, 0));
+
+			const TVector<T, 3> NearTop(7, 0, 4.5);
+			TransformedPhi = Transformed.PhiWithNormal(NearTop, TransformedNormal);
+			EXPECT_FLOAT_EQ(TransformedPhi, -(0.5));
+			EXPECT_VECTOR_NEAR_DEFAULT(TransformedNormal, TVector3(0, 0, 1));
+		}
+		
+		{//support
+			const TVector<T, 3> DirX(1, 0, 0);
+			TVector<T, 3> SupportPt = Untransformed.Support(DirX, 1);
+			EXPECT_VECTOR_NEAR_DEFAULT(SupportPt, TVector3(9, 5, 5));
+
+			SupportPt = Translated.Support(DirX, 1);
+			EXPECT_VECTOR_NEAR_DEFAULT(SupportPt, TVector3(13, 5, 5));
+
+			const TVector<T, 3> DirZ(0, 0, -1);
+			SupportPt = Translated.Support(DirZ, 1);
+			EXPECT_VECTOR_NEAR_DEFAULT(SupportPt, TVector3(12, 5, -6));
+
+			SupportPt = Rotated.Support(DirZ, 1);
+			EXPECT_VECTOR_NEAR_DEFAULT(SupportPt, TVector3(-5, 8, -6)); // @todo why -5?
+
+			SupportPt = Transformed.Support(DirZ, 1);
+			EXPECT_VECTOR_NEAR_DEFAULT(SupportPt, TVector3(-1, 8, -6));
+		}
+
+		{// closest intersection
+			Pair<TVector3, bool> Result;
+			Result = Untransformed.FindClosestIntersection(TVector3(7.5, 0, 0), TVector3(8.5, 0, 0), KINDA_SMALL_NUMBER);
+			EXPECT_VECTOR_NEAR(Result.First, TVector3(8, 0, 0), 0.001);
+
+			Result = Translated.FindClosestIntersection(TVector3(11.5, 0, 0), TVector3(12.5, 0, 0), KINDA_SMALL_NUMBER);
+			EXPECT_VECTOR_NEAR(Result.First, TVector3(12, 0, 0), 0.001);
+
+			Result = Rotated.FindClosestIntersection(TVector3(0, 7.5, 0), TVector3(0, 8.5, 0), KINDA_SMALL_NUMBER);
+			EXPECT_VECTOR_NEAR(Result.First, TVector3(0, 8, 0), 0.001);
+
+			Result = Translated.FindClosestIntersection(TVector3(7, 0, 4.5), TVector3(7, 0, 5.5), KINDA_SMALL_NUMBER);
+			EXPECT_VECTOR_NEAR(Result.First, TVector3(7, 0, 5), 0.001);
+		}
+	}
+	template void ImplicitTransformed<float>();
+
+
+	template<class T>
+	void ImplicitIntersection()
+	{
+		typedef TVector<T, 3> TVector3;
+		FString Caller("ImplicitIntersection()");
+
+		// Two cylinders intersected to make a unit cylinder.
+		TArray<TUniquePtr<FImplicitObject>> Objects;
+		Objects.Add(MakeUnique<TCylinder<T>>(TVector<T, 3>(0, 0, 2), TVector3(0, 0, -1), 1));
+		Objects.Add(MakeUnique<TCylinder<T>>(TVector<T, 3>(0, 0, 1), TVector3(0, 0, -2), 1));
+
+		TImplicitObjectIntersection<T, 3> MIntersectedObjects(std::move(Objects));
+
+		UnitImplicitObjectNormalsInternal(MIntersectedObjects, Caller);
+		UnitImplicitObjectNormalsExternal(MIntersectedObjects, Caller);
+		UnitImplicitObjectIntersections(MIntersectedObjects, Caller);
+
+		Pair<TVector3, bool> Result;
+		{// closest intersection near origin
+			Result = MIntersectedObjects.FindClosestIntersection(TVector3(0, 0, 1 / 2.), TVector3(0, 0, 3 / 2.), KINDA_SMALL_NUMBER);
+			EXPECT_VECTOR_NEAR(Result.First, TVector3(0, 0, 1), 0.001);
+
+			Result = MIntersectedObjects.FindClosestIntersection(TVector3(0, 0, -3 / 2.), TVector3(0, 0, -1 / 2.), KINDA_SMALL_NUMBER);
+			EXPECT_VECTOR_NEAR(Result.First, TVector3(0, 0, -1), 0.001);
+
+			Result = MIntersectedObjects.FindClosestIntersection(TVector3(0, 1 / 2., 0), TVector3(0, 3 / 2., 0), KINDA_SMALL_NUMBER);
+			EXPECT_VECTOR_NEAR(Result.First, TVector3(0, 1, 0), 0.001);
+
+			Result = MIntersectedObjects.FindClosestIntersection(TVector3(0, 3 / 2., 0), TVector3(0, 1 / 2., 0), KINDA_SMALL_NUMBER);
+			EXPECT_VECTOR_NEAR(Result.First, TVector3(0, 1, 0), 0.001);
+
+			// Verify that there's no intersection with non-overlapping parts of the two cylinders. 
+			Result = MIntersectedObjects.FindClosestIntersection(TVector3(0, 0, 5 / 2.), TVector3(0, 0, 7 / 2.), KINDA_SMALL_NUMBER);
+			EXPECT_FALSE(Result.Second);
+
+			Result = MIntersectedObjects.FindClosestIntersection(TVector3(0, 0, -7 / 2.), TVector3(0, 0, -5 / 2.), KINDA_SMALL_NUMBER);
+			EXPECT_FALSE(Result.Second);
+		}
+
+		TArray<TUniquePtr<FImplicitObject>> Objects2;
+		Objects2.Add(MakeUnique<TCylinder<T>>(TVector<T, 3>(4, 4, 6), TVector3(4, 4, 3), 1));
+		Objects2.Add(MakeUnique<TCylinder<T>>(TVector<T, 3>(4, 4, 5), TVector3(4, 4, 2), 1));
+
+		TImplicitObjectIntersection<T, 3> MIntersectedObjects2(std::move(Objects2));
+		
+		{// closest intersection off origin
+			Result = MIntersectedObjects2.FindClosestIntersection(TVector3(4, 4, 4 + 1 / 2.), TVector3(4, 4, 4 + 3 / 2.), KINDA_SMALL_NUMBER);
+			EXPECT_VECTOR_NEAR(Result.First, TVector3(4, 4, 5), 0.001);
+
+			Result = MIntersectedObjects2.FindClosestIntersection(TVector3(4, 4, 4 + -3 / 2.), TVector3(4, 4, 4 + -1 / 2.), KINDA_SMALL_NUMBER);
+			EXPECT_VECTOR_NEAR(Result.First, TVector3(4, 4, 3), 0.001);
+
+			Result = MIntersectedObjects2.FindClosestIntersection(TVector3(4, 4 + 1 / 2., 4), TVector3(4, 4 + 3 / 2., 4), KINDA_SMALL_NUMBER);
+			EXPECT_VECTOR_NEAR(Result.First, TVector3(4, 5, 4), 0.001);
+
+			Result = MIntersectedObjects2.FindClosestIntersection(TVector3(4, 4 + 3 / 2., 4), TVector3(4, 4 + 1 / 2., 4), KINDA_SMALL_NUMBER);
+			EXPECT_VECTOR_NEAR(Result.First, TVector3(4, 5, 4), 0.001);
+
+			// Verify that there's no intersection with non-overlapping parts of the two cylinders. 
+			Result = MIntersectedObjects2.FindClosestIntersection(TVector3(4, 4, 4 + 5 / 2.), TVector3(4, 4, 4 + 7 / 2.), KINDA_SMALL_NUMBER);
+			EXPECT_FALSE(Result.Second);
+
+			Result = MIntersectedObjects2.FindClosestIntersection(TVector3(4, 4, 4 + -7 / 2.), TVector3(4, 4, 4 + -5 / 2.), KINDA_SMALL_NUMBER);
+			EXPECT_FALSE(Result.Second);
+		}
+	}
+	template void ImplicitIntersection<float>();
+
+
+	template<class T>
+	void ImplicitUnion()
+	{
+		typedef TVector<T, 3> TVector3;
+		FString Caller("ImplicitUnion()");
+		TUniquePtr<TImplicitObjectUnion<T, 3>> MUnionedObjects;
+
+		{// unit cylinder - sanity check
+			TArray<TUniquePtr<FImplicitObject>> Objects;
+			Objects.Add(MakeUnique<TCylinder<T>>(TVector<T, 3>(0, 0, 1), TVector3(0), 1));
+			Objects.Add(MakeUnique<TCylinder<T>>(TVector<T, 3>(0, 0, -1), TVector3(0), 1));
+			MUnionedObjects.Reset(new Chaos::TImplicitObjectUnion<float, 3>(std::move(Objects)));
+
+			// Can't use the default internal unit tests because they expect different behavior internally where the two cylinders are joined together. 
+			EXPECT_VECTOR_NEAR(MUnionedObjects->Normal(TVector3(0, 0, 2 / 3.)), (TVector3(0, 0, 1)), KINDA_SMALL_NUMBER);
+			EXPECT_VECTOR_NEAR(MUnionedObjects->Normal(TVector3(0, 0, -2 / 3.)), (TVector3(0, 0, -1)), KINDA_SMALL_NUMBER);
+			EXPECT_VECTOR_NEAR(MUnionedObjects->Normal(TVector3(0, 2 / 3., 0)), (TVector3(0, 0, 0)), KINDA_SMALL_NUMBER);
+			EXPECT_VECTOR_NEAR(MUnionedObjects->Normal(TVector3(0, -2 / 3., 0)), (TVector3(0, 0, 0)), KINDA_SMALL_NUMBER);
+			EXPECT_VECTOR_NEAR(MUnionedObjects->Normal(TVector3(2 / 3., 0, 0)), (TVector3(0, 0, 0)), KINDA_SMALL_NUMBER);
+			EXPECT_VECTOR_NEAR(MUnionedObjects->Normal(TVector3(-2 / 3., 0, 0)), (TVector3(0, 0, 0)), KINDA_SMALL_NUMBER);
+
+			UnitImplicitObjectNormalsExternal(*MUnionedObjects, Caller);
+
+			EXPECT_NEAR(MUnionedObjects->SignedDistance(TVector3(0, 0, 5 / 4.)), 1 / 4., KINDA_SMALL_NUMBER);
+			EXPECT_NEAR(MUnionedObjects->SignedDistance(TVector3(0, 0, 3 / 4.)), -1 / 4., KINDA_SMALL_NUMBER);
+			EXPECT_NEAR(MUnionedObjects->SignedDistance(TVector3(0, 5 / 4., 0)), 1 / 4., KINDA_SMALL_NUMBER);
+			// Internal distance 0 because it's where the spheres overlap.
+			EXPECT_NEAR(MUnionedObjects->SignedDistance(TVector3(0, 3 / 4., 0)), 0., KINDA_SMALL_NUMBER);
+
+			TestFindClosestIntersection(*MUnionedObjects, TVector3(0, 0, 5 / 4.), TVector3(0, 0, 1), Caller);
+			TestFindClosestIntersection(*MUnionedObjects, TVector3(0, 0, -5 / 4.), TVector3(0, 0, -1), Caller);
+		}
+
+		TArray<TUniquePtr<FImplicitObject>> Objects;
+		Objects.Add(MakeUnique<TCylinder<T>>(TVector<T, 3>(0, 0, -2), TVector3(0, 0, 2), 1));
+		Objects.Add(MakeUnique<TCylinder<T>>(TVector<T, 3>(0, -2, 0), TVector3(0, 2, 0), 1));
+		MUnionedObjects.Reset(new Chaos::TImplicitObjectUnion<float, 3>(std::move(Objects)));
+
+		{// closest point near origin (+)
+			EXPECT_NEAR(MUnionedObjects->SignedDistance(TVector3(0, 0, 9 / 4.)), 1 / 4., KINDA_SMALL_NUMBER);
+			TestFindClosestIntersection(*MUnionedObjects, TVector3(0, 0, 9 / 4.), TVector3(0, 0, 2), Caller);
+			TestFindClosestIntersection(*MUnionedObjects, TVector3(0, 0, -9 / 4.), TVector3(0, 0, -2), Caller);
+			TestFindClosestIntersection(*MUnionedObjects, TVector3(0, 9 / 4., 0), TVector3(0, 2, 0), Caller);
+			TestFindClosestIntersection(*MUnionedObjects, TVector3(0, -9 / 4., 0), TVector3(0, -2, 0), Caller);
+			TestFindClosestIntersection(*MUnionedObjects, TVector3(3 / 2., 0, 0), TVector3(1, 0, 0), Caller);
+			TestFindClosestIntersection(*MUnionedObjects, TVector3(-3 / 2., 0, 0), TVector3(-1, 0, 0), Caller);
+		}
+
+		{// closest point near origin (-)
+			EXPECT_NEAR(MUnionedObjects->SignedDistance(TVector3(0, 0, 7 / 4.)), -1 / 4., KINDA_SMALL_NUMBER);
+			TestFindClosestIntersection(*MUnionedObjects, TVector3(0, 0, 7 / 4.), TVector3(0, 0, 2), Caller);
+			TestFindClosestIntersection(*MUnionedObjects, TVector3(0, 0, -7 / 4.), TVector3(0, 0, -2), Caller);
+			TestFindClosestIntersection(*MUnionedObjects, TVector3(0, 7 / 4., 0), TVector3(0, 2, 0), Caller);
+			TestFindClosestIntersection(*MUnionedObjects, TVector3(0, -7 / 4., 0), TVector3(0, -2, 0), Caller);
+			TestFindClosestIntersection(*MUnionedObjects, TVector3(1 / 2., 0, 0), TVector3(1, 0, 0), Caller);
+			TestFindClosestIntersection(*MUnionedObjects, TVector3(-1 / 2., 0, 0), TVector3(-1, 0, 0), Caller);
+		}
+		
+		TArray<TUniquePtr<FImplicitObject>> Objects2;
+		Objects2.Add(MakeUnique<TCylinder<T>>(TVector<T, 3>(4, 4, 2), TVector3(4, 4, 6), 1));
+		Objects2.Add(MakeUnique<TCylinder<T>>(TVector<T, 3>(4, 2, 4), TVector3(4, 6, 4), 1));
+		MUnionedObjects.Reset(new Chaos::TImplicitObjectUnion<float, 3>(std::move(Objects2)));
+
+		{// closest point off origin (+)
+			EXPECT_NEAR(MUnionedObjects->SignedDistance(TVector3(4, 4, 4 + 9 / 4.)), 1 / 4., KINDA_SMALL_NUMBER);
+			TestFindClosestIntersection(*MUnionedObjects, TVector3(4, 4, 4 + 9 / 4.), TVector3(4, 4, 6), Caller);
+			TestFindClosestIntersection(*MUnionedObjects, TVector3(4, 4, 4 + -9 / 4.), TVector3(4, 4, 2), Caller);
+			TestFindClosestIntersection(*MUnionedObjects, TVector3(4, 4 + 9 / 4., 4), TVector3(4, 6, 4), Caller);
+			TestFindClosestIntersection(*MUnionedObjects, TVector3(4, 4 + -9 / 4., 4), TVector3(4, 2, 4), Caller);
+			TestFindClosestIntersection(*MUnionedObjects, TVector3(4 + 3 / 2., 4, 4), TVector3(5, 4, 4), Caller);
+			TestFindClosestIntersection(*MUnionedObjects, TVector3(4 + -3 / 2., 4, 4), TVector3(3, 4, 4), Caller);
+		}
+
+		{// closest point off origin (-)
+			EXPECT_NEAR(MUnionedObjects->SignedDistance(TVector3(4, 4, 4 + 7 / 4.)), -1 / 4., KINDA_SMALL_NUMBER);
+			TestFindClosestIntersection(*MUnionedObjects, TVector3(4, 4, 4 + 7 / 4.), TVector3(4, 4, 6), Caller);
+			TestFindClosestIntersection(*MUnionedObjects, TVector3(4, 4, 4 + -7 / 4.), TVector3(4, 4, 2), Caller);
+			TestFindClosestIntersection(*MUnionedObjects, TVector3(4, 4 + 7 / 4., 4), TVector3(4, 6, 4), Caller);
+			TestFindClosestIntersection(*MUnionedObjects, TVector3(4, 4 + -7 / 4., 4), TVector3(4, 2, 4), Caller);
+			TestFindClosestIntersection(*MUnionedObjects, TVector3(4 + 1 / 2., 4, 4), TVector3(5, 4, 4), Caller);
+			TestFindClosestIntersection(*MUnionedObjects, TVector3(4 + -1 / 2., 4, 4), TVector3(3, 4, 4), Caller);
+		}
+
+		/* Nested Unions */
+		
+		{// Union of unions (capsule)
+			TArray<TUniquePtr<FImplicitObject>> Unions;
+			Unions.Add(MakeUnique<TCapsule<T>>(TVector<T, 3>(0, 0, 0), TVector<T, 3>(0, 0, -2), 1));
+			Unions.Add(MakeUnique<TCapsule<T>>(TVector<T, 3>(0, 0, 0), TVector<T, 3>(0, 0, 2), 1));
+			MUnionedObjects.Reset(new Chaos::TImplicitObjectUnion<float, 3>(std::move(Unions)));
+
+			EXPECT_VECTOR_NEAR(MUnionedObjects->Normal(TVector3(0, 0, 7 / 3.)), (TVector3(0, 0, 1)), KINDA_SMALL_NUMBER);
+			EXPECT_VECTOR_NEAR(MUnionedObjects->Normal(TVector3(0, 0, -7 / 3.)), (TVector3(0, 0, -1)), KINDA_SMALL_NUMBER);
+			EXPECT_VECTOR_NEAR(MUnionedObjects->Normal(TVector3(0, 1 / 2., 0)), (TVector3(0, 1, 0)), KINDA_SMALL_NUMBER);
+			EXPECT_VECTOR_NEAR(MUnionedObjects->Normal(TVector3(0, -1 / 2., 0)), (TVector3(0, -1, 0)), KINDA_SMALL_NUMBER);
+			EXPECT_VECTOR_NEAR(MUnionedObjects->Normal(TVector3(1 / 2., 0, 0)), (TVector3(1, 0, 0)), KINDA_SMALL_NUMBER);
+			EXPECT_VECTOR_NEAR(MUnionedObjects->Normal(TVector3(-1 / 2., 0, 0)), (TVector3(-1, 0, 0)), KINDA_SMALL_NUMBER);
+			EXPECT_NEAR(MUnionedObjects->SignedDistance(TVector3(0, 0, 13 / 4.)), 1 / 4., KINDA_SMALL_NUMBER);
+			EXPECT_NEAR(MUnionedObjects->SignedDistance(TVector3(0, 0, 11 / 4.)), -1 / 4., KINDA_SMALL_NUMBER);
+			EXPECT_NEAR(MUnionedObjects->SignedDistance(TVector3(0, 1 / 2., 0)), -1 / 2., KINDA_SMALL_NUMBER);
+			EXPECT_NEAR(MUnionedObjects->SignedDistance(TVector3(0, 3 / 2., 0)), 1 / 2., KINDA_SMALL_NUMBER);
+		}
+
+		{// Union of a union containing all the unit geometries overlapping - should still pass all the normal unit tests. 
+			TArray<TUniquePtr<FImplicitObject>> Objects1;
+			Objects1.Add(MakeUnique<TCylinder<T>>(TVector<T, 3>(0, 0, 1), TVector<T, 3>(0, 0, -1), 1));
+			Objects1.Add(MakeUnique<TSphere<T, 3>>(TVector<T, 3>(0, 0, 0), 1));
+			Objects1.Add(MakeUnique<TBox<T, 3>>(TVector<T, 3>(-1, -1, -1), TVector<T, 3>(1, 1, 1)));
+			Objects1.Add(MakeUnique<TTaperedCylinder<T>>(TVector<T, 3>(0, 0, 1), TVector<T, 3>(0, 0, -1), 1, 1));
+
+			TArray<TUniquePtr<FImplicitObject>> Unions;
+			Unions.Emplace(new TImplicitObjectUnion<T, 3>(MoveTemp(Objects1)));
+			TUniquePtr<TImplicitObjectUnion<T, 3>> UnionedUnions;
+			UnionedUnions.Reset(new Chaos::TImplicitObjectUnion<float, 3>(std::move(Unions)));
+
+			UnitImplicitObjectNormalsExternal(*UnionedUnions, FString("ImplicitUnion() - nested union unit cylinder 1"));
+			UnitImplicitObjectNormalsInternal(*UnionedUnions, FString("ImplicitUnion() - nested union unit cylinder 1"));
+			UnitImplicitObjectIntersections(*UnionedUnions, FString("ImplicitUnion() - nested union unit cylinder 1"));
+		}
+
+		{// Union of two unions, each with two unit objects
+			TArray<TUniquePtr<FImplicitObject>> ObjectsA;
+			TArray<TUniquePtr<FImplicitObject>> ObjectsB;
+			ObjectsA.Add(MakeUnique<TCylinder<T>>(TVector<T, 3>(0, 0, 1), TVector<T, 3>(0, 0, -1), 1));
+			ObjectsA.Add(MakeUnique<TSphere<T, 3>>(TVector<T, 3>(0, 0, 0), 1));
+			ObjectsB.Add(MakeUnique<TBox<T, 3>>(TVector<T, 3>(-1, -1, -1), TVector<T, 3>(1, 1, 1)));
+			ObjectsB.Add(MakeUnique<TTaperedCylinder<T>>(TVector<T, 3>(0, 0, 1), TVector<T, 3>(0, 0, -1), 1, 1));
+
+			TArray<TUniquePtr<FImplicitObject>> Unions;
+			Unions.Emplace(new TImplicitObjectUnion<T, 3>(MoveTemp(ObjectsA)));
+			Unions.Emplace(new TImplicitObjectUnion<T, 3>(MoveTemp(ObjectsB)));
+			TUniquePtr<TImplicitObjectUnion<T, 3>> UnionedUnions;
+			UnionedUnions.Reset(new Chaos::TImplicitObjectUnion<float, 3>(std::move(Unions)));
+
+			UnitImplicitObjectNormalsExternal(*UnionedUnions, FString("ImplicitUnion() - nested union unit sphere 1"));
+			UnitImplicitObjectNormalsInternal(*UnionedUnions, FString("ImplicitUnion() - nested union unit sphere 1"));
+			UnitImplicitObjectIntersections(*UnionedUnions, FString("ImplicitUnion() - nested union unit sphere 1"));
+		}
+
+		{// Mimic a unit cylinder, but made up of multiple unions. 
+			TArray<TUniquePtr<FImplicitObject>> ObjectsA;
+			TArray<TUniquePtr<FImplicitObject>> ObjectsB;
+			ObjectsA.Add(MakeUnique<TCylinder<T>>(TVector<T, 3>(0, 0, 0), TVector<T, 3>(0, 0, -1), 1));
+			ObjectsB.Add(MakeUnique<TCylinder<T>>(TVector<T, 3>(0, 0, 0), TVector<T, 3>(0, 0, 1), 1));
+			TArray<TUniquePtr<FImplicitObject>> Unions;
+			Unions.Emplace(new TImplicitObjectUnion<T, 3>(MoveTemp(ObjectsA)));
+			Unions.Emplace(new TImplicitObjectUnion<T, 3>(MoveTemp(ObjectsB)));
+			TUniquePtr<TImplicitObjectUnion<T, 3>> UnionedUnions;
+			UnionedUnions.Reset(new Chaos::TImplicitObjectUnion<float, 3>(std::move(Unions)));
+
+			UnitImplicitObjectNormalsExternal(*UnionedUnions, FString("ImplicitUnion() - nested union unit cylinder 2"));
+
+			EXPECT_VECTOR_NEAR(UnionedUnions->Normal(TVector3(0, 0, 2 / 3.)), (TVector3(0, 0, 1)), KINDA_SMALL_NUMBER);
+			EXPECT_VECTOR_NEAR(UnionedUnions->Normal(TVector3(0, 0, -2 / 3.)), (TVector3(0, 0, -1)), KINDA_SMALL_NUMBER);
+			// Normal is averaged to 0 at the joined faces. 
+			EXPECT_VECTOR_NEAR(UnionedUnions->Normal(TVector3(0, 0, 0)), (TVector3(0, 0, 0)), KINDA_SMALL_NUMBER);
+
+			EXPECT_NEAR(UnionedUnions->SignedDistance(TVector3(0, 5 / 4., 0)), 1 / 4., KINDA_SMALL_NUMBER);
+			EXPECT_NEAR(UnionedUnions->SignedDistance(TVector3(0, -5 / 4., 0)), 1 / 4., KINDA_SMALL_NUMBER);
+			EXPECT_NEAR(UnionedUnions->SignedDistance(TVector3(5 / 4., 0, 0)), 1 / 4., KINDA_SMALL_NUMBER);
+			EXPECT_NEAR(UnionedUnions->SignedDistance(TVector3(-5 / 4., 0, 0)), 1 / 4., KINDA_SMALL_NUMBER);
+
+			// Distance is 0 at the joined faces.
+			EXPECT_NEAR(UnionedUnions->SignedDistance(TVector3(0, 0, 0)), 0., KINDA_SMALL_NUMBER);
+		}
+
+	}
+	template void ImplicitUnion<float>();
+
+
+	template<class T>
+	void ImplicitLevelset()
+	{
+		typedef TVector<T, 3> TVector3;
+		Chaos::TPBDRigidParticles<T, 3> Particles;
+		TArray<TVector<int32, 3>> CollisionMeshElements;
+		int32 BoxId = AppendParticleBox<T>(Particles, TVector<T, 3>(1), &CollisionMeshElements);
+		TLevelSet<T, 3> Levelset = ConstructLevelset(*Particles.CollisionParticles(BoxId), CollisionMeshElements);
+
+		TVector<T, 3> Normal;
+		T Phi = Levelset.PhiWithNormal(TVector3(0, 0, 2), Normal);
+		EXPECT_GT(Phi, 0);
+		EXPECT_NEAR((Phi - 1.5), 0, KINDA_SMALL_NUMBER);
+		EXPECT_VECTOR_NEAR(Normal, TVector3(0, 0, 1), 0.001);
+
+		Phi = Levelset.PhiWithNormal(TVector3(0, 2, 0), Normal);
+		EXPECT_GT(Phi, 0);
+		EXPECT_NEAR((Phi - 1.5), 0, KINDA_SMALL_NUMBER);
+		EXPECT_VECTOR_NEAR(Normal, TVector3(0, 1, 0), 0.001);
+		
+		Phi = Levelset.PhiWithNormal(TVector3(2, 0, 0), Normal);
+		EXPECT_GT(Phi, 0);
+		EXPECT_NEAR((Phi - 1.5), 0, KINDA_SMALL_NUMBER);
+		EXPECT_VECTOR_NEAR(Normal, TVector3(1, 0, 0), 0.001);
+
+		Phi = Levelset.PhiWithNormal(TVector3(0, 0, -2), Normal);
+		EXPECT_GT(Phi, 0);
+		EXPECT_NEAR((Phi - 1.5), 0, KINDA_SMALL_NUMBER);
+		EXPECT_VECTOR_NEAR(Normal, TVector3(0, 0, -1), 0.001);
+
+		Phi = Levelset.PhiWithNormal(TVector3(0, -2, 0), Normal);
+		EXPECT_GT(Phi, 0);
+		EXPECT_NEAR((Phi - 1.5), 0, KINDA_SMALL_NUMBER);
+		EXPECT_VECTOR_NEAR(Normal, TVector3(0, -1, 0), 0.001);
+		
+		Phi = Levelset.PhiWithNormal(TVector3(-2, 0, 0), Normal);
+		EXPECT_GT(Phi, 0);
+		EXPECT_NEAR((Phi - 1.5), 0, KINDA_SMALL_NUMBER);
+		EXPECT_VECTOR_NEAR(Normal, TVector3(-1, 0, 0), 0.001); /**/
+
+		Phi = Levelset.PhiWithNormal(TVector3(0, 0, 0.25f), Normal);
+		EXPECT_LT(Phi, 0);
+		EXPECT_NEAR((Phi + 0.25), 0, KINDA_SMALL_NUMBER);
+		EXPECT_VECTOR_NEAR(Normal, TVector3(0, 0, 1), 0.001);
+
+		Phi = Levelset.PhiWithNormal(TVector3(0, 0.25f, 0), Normal);
+		EXPECT_LT(Phi, 0);
+		EXPECT_NEAR((Phi + 0.25), 0, KINDA_SMALL_NUMBER);
+		EXPECT_VECTOR_NEAR(Normal, TVector3(0, 1, 0), 0.001);
+
+		Phi = Levelset.PhiWithNormal(TVector3(0.25f, 0, 0), Normal);
+		EXPECT_LT(Phi, 0);
+		EXPECT_NEAR((Phi + 0.25), 0, KINDA_SMALL_NUMBER);
+		EXPECT_VECTOR_NEAR(Normal, TVector3(1, 0, 0), 0.001);
+
+		Phi = Levelset.PhiWithNormal(TVector3(0, 0, -0.25f), Normal);
+		EXPECT_LT(Phi, 0);
+		EXPECT_NEAR((Phi + 0.25), 0, KINDA_SMALL_NUMBER);
+		EXPECT_VECTOR_NEAR(Normal, TVector3(0, 0, -1), 0.001);
+
+		Phi = Levelset.PhiWithNormal(TVector3(0, -0.25f, 0), Normal);
+		EXPECT_LT(Phi, 0);
+		EXPECT_NEAR((Phi + 0.25), 0, KINDA_SMALL_NUMBER);
+		EXPECT_VECTOR_NEAR(Normal, TVector3(0, -1, 0), 0.001);
+
+		Phi = Levelset.PhiWithNormal(TVector3(-0.25f, 0, 0), Normal);
+		EXPECT_LT(Phi, 0);
+		EXPECT_NEAR((Phi + 0.25), 0, KINDA_SMALL_NUMBER);
+		EXPECT_VECTOR_NEAR(Normal, TVector3(-1, 0, 0), 0.001);
+	}
+
+	template<class T>
+	void RasterizationImplicit()
+	{
+		TUniquePtr<TBox<T, 3>> Box(new TBox<T,3>(TVector<T, 3>(-0.5, -0.5, -0.5), TVector<T, 3>(0.5, 0.5, 0.5)));
+		TArray<TUniquePtr<FImplicitObject>> Objects;
+		Objects.Add(TUniquePtr<FImplicitObject>(new TImplicitObjectTransformed<T, 3>(MakeSerializable(Box), TRigidTransform<T, 3>(TVector<T, 3>(0.5, 0, 0), TRotation<T, 3>::FromVector(TVector<T, 3>(0))))));
+		Objects.Add(TUniquePtr<FImplicitObject>(new TImplicitObjectTransformed<T, 3>(MakeSerializable(Box), TRigidTransform<T, 3>(TVector<T, 3>(-0.5, 0, 0), TRotation<T, 3>::FromVector(TVector<T, 3>(0))))));
+		FImplicitObjectUnion Union(MoveTemp(Objects));
+		FErrorReporter ErrorReporter;
+		// This one should be exactly right as we don't actually do an fast marching interior to the region
+		{
+			TUniformGrid<T, 3> Grid(TVector<T, 3>(-2.0, -1.5, -1.5), TVector<T, 3>(2.0, 1.5, 1.5), TVector<int32, 3>(4, 3, 3));
+			TLevelSet<T, 3> LevelSet(ErrorReporter, Grid, Union);
+			EXPECT_TRUE(LevelSet.IsConvex());
+			EXPECT_LT(LevelSet.SignedDistance(TVector<T, 3>(0)) + T(0.5), KINDA_SMALL_NUMBER);
+		}
+		// We should get closer answers every time we refine the resolution
+		{
+			ErrorReporter.HandleLatestError();
+			TUniformGrid<T, 3> Grid(TVector<T, 3>(-1.5, -1.0, -1.0), TVector<T, 3>(1.5, 1.0, 1.0), TVector<int32, 3>(6, 4, 4));
+			TLevelSet<T, 3> LevelSet(ErrorReporter, Grid, Union);
+			EXPECT_TRUE(LevelSet.IsConvex());
+			EXPECT_LT(LevelSet.SignedDistance(TVector<T, 3>(0)) + T(0.25), KINDA_SMALL_NUMBER);
+		}
+		{
+			ErrorReporter.HandleLatestError();
+			TUniformGrid<T, 3> Grid(TVector<T, 3>(-1.25, -0.75, -0.75), TVector<T, 3>(1.25, 0.75, 0.75), TVector<int32, 3>(10, 6, 6));
+			TLevelSet<T, 3> LevelSet(ErrorReporter, Grid, Union);
+			EXPECT_TRUE(LevelSet.IsConvex());
+			EXPECT_LT(LevelSet.SignedDistance(TVector<T, 3>(0)) + T(0.3), KINDA_SMALL_NUMBER);
+		}
+		{
+			ErrorReporter.HandleLatestError();
+			TUniformGrid<T, 3> Grid(TVector<T, 3>(-1.1, -0.6, -0.6), TVector<T, 3>(1.1, 0.6, 0.6), TVector<int32, 3>(22, 12, 12));
+			TLevelSet<T, 3> LevelSet(ErrorReporter, Grid, Union);
+			EXPECT_TRUE(LevelSet.IsConvex());
+			EXPECT_LT(LevelSet.SignedDistance(TVector<T, 3>(0)) + T(0.4), KINDA_SMALL_NUMBER);
+		}
+		{
+			ErrorReporter.HandleLatestError();
+			TUniformGrid<T, 3> Grid(TVector<T, 3>(-1.05, -0.55, -0.55), TVector<T, 3>(1.05, 0.55, 0.55), TVector<int32, 3>(42, 22, 22));
+			TLevelSet<T, 3> LevelSet(ErrorReporter, Grid, Union);
+			EXPECT_TRUE(LevelSet.IsConvex());
+			EXPECT_LT(LevelSet.SignedDistance(TVector<T, 3>(0)) + T(0.45), KINDA_SMALL_NUMBER);
+		}
+		{
+			ErrorReporter.HandleLatestError();
+			TUniformGrid<T, 3> Grid(TVector<T, 3>(-1.5, -1.0, -1.0), TVector<T, 3>(1.5, 1.0, 1.0), TVector<int32, 3>(20, 20, 20));
+			TLevelSet<T, 3> LevelSet(ErrorReporter, Grid, Union);
+
+			T Volume;
+			TVector<T, 3> COM;
+			PMatrix<T, 3, 3> Inertia;
+			TRotation<T, 3> RotationOfMass;
+
+			LevelSet.ComputeMassProperties(Volume, COM, Inertia, RotationOfMass);
+			EXPECT_GT(Volume, 1);
+			EXPECT_LT(Volume, 3);
+			EXPECT_LT(Inertia.M[0][0] * 1.5, Inertia.M[1][1]);
+			EXPECT_GT(Inertia.M[0][0] * 3, Inertia.M[1][1]);
+			EXPECT_EQ(Inertia.M[2][2], Inertia.M[1][1]);
+		}
+	}
+
+	template<class T>
+	void RasterizationImplicitWithHole()
+	{
+		TUniquePtr<TBox<T, 3>> Box(new TBox<T,3>(TVector<T, 3>(-0.5, -0.5, -0.5), TVector<T, 3>(0.5, 0.5, 0.5)));
+		TSerializablePtr<TBox<T, 3>> SerializableBox(Box);
+		TArray<TUniquePtr<FImplicitObject>> Objects;
+		Objects.Add(TUniquePtr<FImplicitObject>(new TImplicitObjectTransformed<T, 3>(SerializableBox, TRigidTransform<T, 3>(TVector<T, 3>(1, 1, 0), TRotation<T, 3>::FromVector(TVector<T, 3>(0))))));
+		Objects.Add(TUniquePtr<FImplicitObject>(new TImplicitObjectTransformed<T, 3>(SerializableBox, TRigidTransform<T, 3>(TVector<T, 3>(0, 1, 0), TRotation<T, 3>::FromVector(TVector<T, 3>(0))))));
+		Objects.Add(TUniquePtr<FImplicitObject>(new TImplicitObjectTransformed<T, 3>(SerializableBox, TRigidTransform<T, 3>(TVector<T, 3>(-1, 1, 0), TRotation<T, 3>::FromVector(TVector<T, 3>(0))))));
+		Objects.Add(TUniquePtr<FImplicitObject>(new TImplicitObjectTransformed<T, 3>(SerializableBox, TRigidTransform<T, 3>(TVector<T, 3>(1, 0, 0), TRotation<T, 3>::FromVector(TVector<T, 3>(0))))));
+		Objects.Add(TUniquePtr<FImplicitObject>(new TImplicitObjectTransformed<T, 3>(SerializableBox, TRigidTransform<T, 3>(TVector<T, 3>(-1, 0, 0), TRotation<T, 3>::FromVector(TVector<T, 3>(0))))));
+		Objects.Add(TUniquePtr<FImplicitObject>(new TImplicitObjectTransformed<T, 3>(SerializableBox, TRigidTransform<T, 3>(TVector<T, 3>(1, -1, 0), TRotation<T, 3>::FromVector(TVector<T, 3>(0))))));
+		Objects.Add(TUniquePtr<FImplicitObject>(new TImplicitObjectTransformed<T, 3>(SerializableBox, TRigidTransform<T, 3>(TVector<T, 3>(0, -1, 0), TRotation<T, 3>::FromVector(TVector<T, 3>(0))))));
+		Objects.Add(TUniquePtr<FImplicitObject>(new TImplicitObjectTransformed<T, 3>(SerializableBox, TRigidTransform<T, 3>(TVector<T, 3>(-1, -1, 0), TRotation<T, 3>::FromVector(TVector<T, 3>(0))))));
+		FImplicitObjectUnion Union(MoveTemp(Objects));
+		{
+			TUniformGrid<T, 3> Grid(TVector<T, 3>(-1.6, -1.6, -0.6), TVector<T, 3>(1.6, 1.6, 0.6), TVector<int32, 3>(32, 32, 12));
+			FErrorReporter ErrorReporter;
+			TLevelSet<T, 3> LevelSet(ErrorReporter, Grid, Union);
+			EXPECT_FALSE(LevelSet.IsConvex());
+			EXPECT_GT(LevelSet.SignedDistance(TVector<T, 3>(0)), -KINDA_SMALL_NUMBER);
+			EXPECT_LT(LevelSet.SignedDistance(TVector<T, 3>(1, 1, 0)), KINDA_SMALL_NUMBER);
+			EXPECT_LT(LevelSet.SignedDistance(TVector<T, 3>(0, 1, 0)), KINDA_SMALL_NUMBER);
+			EXPECT_LT(LevelSet.SignedDistance(TVector<T, 3>(-1, 1, 0)), KINDA_SMALL_NUMBER);
+			EXPECT_LT(LevelSet.SignedDistance(TVector<T, 3>(-1, 0, 0)), KINDA_SMALL_NUMBER);
+			EXPECT_LT(LevelSet.SignedDistance(TVector<T, 3>(1, 0, 0)), KINDA_SMALL_NUMBER);
+			EXPECT_LT(LevelSet.SignedDistance(TVector<T, 3>(1, -1, 0)), KINDA_SMALL_NUMBER);
+			EXPECT_LT(LevelSet.SignedDistance(TVector<T, 3>(0, -1, 0)), KINDA_SMALL_NUMBER);
+			EXPECT_LT(LevelSet.SignedDistance(TVector<T, 3>(-1, -1, 0)), KINDA_SMALL_NUMBER);
+		}
+	}
+
+	template<class T>
+	void ConvexHull()
+	{
+		{
+			TParticles<T, 3> Particles;
+			Particles.AddParticles(9);
+			Particles.X(0) = TVector<float, 3>(-1, -1, -1);
+			Particles.X(1) = TVector<float, 3>(-1, -1, 1);
+			Particles.X(2) = TVector<float, 3>(-1, 1, -1);
+			Particles.X(3) = TVector<float, 3>(-1, 1, 1);
+			Particles.X(4) = TVector<float, 3>(1, -1, -1);
+			Particles.X(5) = TVector<float, 3>(1, -1, 1);
+			Particles.X(6) = TVector<float, 3>(1, 1, -1);
+			Particles.X(7) = TVector<float, 3>(1, 1, 1);
+			Particles.X(8) = TVector<float, 3>(0, 0, 0);
+			const TTriangleMesh<T> TriMesh = TTriangleMesh<T>::GetConvexHullFromParticles(Particles);
+			EXPECT_EQ(TriMesh.GetSurfaceElements().Num(), 12);
+			for (const auto& Tri : TriMesh.GetSurfaceElements())
+			{
+				EXPECT_NE(Tri.X, 8);
+				EXPECT_NE(Tri.Y, 8);
+				EXPECT_NE(Tri.Z, 8);
+			}
+
+			FConvex Convex(Particles);
+			const TParticles<T, 3>& CulledParticles = Convex.GetSurfaceParticles();
+			EXPECT_EQ(CulledParticles.Size(), 8);
+
+			for (uint32 Idx = 0; Idx < CulledParticles.Size(); ++Idx)
+			{
+				EXPECT_NE(Particles.X(8), CulledParticles.X(Idx));	//interior particle gone
+				bool bFound = false;
+				for (uint32 InnerIdx = 0; InnerIdx < Particles.Size(); ++InnerIdx)	//remaining particles are from the original set
+				{
+					if (Particles.X(InnerIdx) == CulledParticles.X(Idx))
+					{
+						bFound = true;
+						break;
+					}
+				}
+				EXPECT_TRUE(bFound);
+			}
+
+		}
+
+		{
+			TParticles<T, 3> Particles;
+			Particles.AddParticles(6);
+			Particles.X(0) = TVector<float, 3>(-1, -1, -1);
+			Particles.X(1) = TVector<float, 3>(1, -1, -1);
+			Particles.X(2) = TVector<float, 3>(1, 1, -1);
+			Particles.X(3) = TVector<float, 3>(0, 0, 0.5);
+			Particles.X(4) = (Particles.X(3) - Particles.X(1)) * 0.5 + Particles.X(1) + TVector<float, 3>(0, 0, 0.1);
+			Particles.X(5) = Particles.X(4) + TVector<float, 3>(-0.1, 0, 0);
+			const TTriangleMesh<T> TriMesh = TTriangleMesh<T>::GetConvexHullFromParticles(Particles);
+			//EXPECT_EQ(TriMesh.GetSurfaceElements().Num(), 6);
+		}
+	}
+
+	template void ImplicitLevelset<float>();
+	template void RasterizationImplicit<float>();
+	template void RasterizationImplicitWithHole<float>();
+	template void ConvexHull<float>();
+
+	template<class T>
+	void ConvexHull2()
+	{
+		{
+			//degenerates
+			Chaos::TParticles<T, 3> Particles;
+			Particles.AddParticles(3);
+			Particles.X(0) = TVector<float, 3>(-1, -1, -1);
+			Particles.X(1) = TVector<float, 3>(1, -1, -1);
+			Particles.X(2) = TVector<float, 3>(1, 1, -1);
+			TArray<TVector<int32, 3>>Indices;
+			Chaos::FConvexBuilder::BuildConvexHull(Particles, Indices);
+			EXPECT_EQ(Indices.Num(), 0);
+			Particles.AddParticles(1);
+			Particles.X(3) = TVector<float, 3>(2, 3, -1);
+			Chaos::FConvexBuilder::BuildConvexHull(Particles, Indices);
+			EXPECT_EQ(Indices.Num(), 0);
+		}
+		{
+			Chaos::TParticles<T, 3> Particles;
+			Particles.AddParticles(9);
+			Particles.X(0) = TVector<float, 3>(-1, -1, -1);
+			Particles.X(1) = TVector<float, 3>(-1, -1, 1);
+			Particles.X(2) = TVector<float, 3>(-1, 1, -1);
+			Particles.X(3) = TVector<float, 3>(-1, 1, 1);
+			Particles.X(4) = TVector<float, 3>(1, -1, -1);
+			Particles.X(5) = TVector<float, 3>(1, -1, 1);
+			Particles.X(6) = TVector<float, 3>(1, 1, -1);
+			Particles.X(7) = TVector<float, 3>(1, 1, 1);
+			Particles.X(8) = TVector<float, 3>(0, 0, 0);
+			TArray<TVector<int32, 3>>Indices;
+			Chaos::FConvexBuilder::BuildConvexHull(Particles, Indices);
+			EXPECT_EQ(Indices.Num(), 12);
+			for (const auto& Tri : Indices)
+			{
+				EXPECT_NE(Tri.X, 8);
+				EXPECT_NE(Tri.Y, 8);
+				EXPECT_NE(Tri.Z, 8);
+			}
+		}
+		{
+			Chaos::TParticles<T, 3> Particles;
+			Particles.AddParticles(5);
+			Particles.X(0) = TVector<float, 3>(-1, -1, -1);
+			Particles.X(1) = TVector<float, 3>(1, -1, -1);
+			Particles.X(2) = TVector<float, 3>(1, 1, -1);
+			Particles.X(3) = TVector<float, 3>(0, 0, 0.5);
+			Particles.X(4) = (Particles.X(3) - Particles.X(1)) * 0.5 + Particles.X(1) + TVector<float, 3>(0, 0, 0.1);
+			TArray<TVector<int32, 3>> Indices;
+			Chaos::FConvexBuilder::BuildConvexHull(Particles, Indices);
+			EXPECT_EQ(Indices.Num(), 6);
+		}
+		{
+			Chaos::TParticles<T, 3> Particles;
+			Particles.AddParticles(6);
+			Particles.X(0) = TVector<float, 3>(-1, -1, -1);
+			Particles.X(1) = TVector<float, 3>(1, -1, -1);
+			Particles.X(2) = TVector<float, 3>(1, 1, -1);
+			Particles.X(3) = TVector<float, 3>(0, 0, 0.5);
+			Particles.X(4) = (Particles.X(3) - Particles.X(1)) * 0.5 + Particles.X(1) + TVector<float, 3>(0, 0, 0.1);
+			Particles.X(5) = Particles.X(4) + TVector<float, 3>(-0.1, 0, 0);
+			TArray<TVector<int32, 3>> Indices;
+			Chaos::FConvexBuilder::BuildConvexHull(Particles, Indices);
+			EXPECT_EQ(Indices.Num(), 8);
+		}
+		{
+			Chaos::TParticles<T, 3> Particles;
+			int32 NumParticles = 3600;
+			Particles.AddParticles(NumParticles);
+			Particles.X(0) = TVector<float, 3>(-1, -1, -1);
+			Particles.X(1) = TVector<float, 3>(-1, -1, 1);
+			Particles.X(2) = TVector<float, 3>(-1, 1, -1);
+			Particles.X(3) = TVector<float, 3>(-1, 1, 1);
+			Particles.X(4) = TVector<float, 3>(1, -1, -1);
+			Particles.X(5) = TVector<float, 3>(1, -1, 1);
+			Particles.X(6) = TVector<float, 3>(1, 1, -1);
+			Particles.X(7) = TVector<float, 3>(1, 1, 1);
+			FRandomStream Stream(42);
+			for (int i = 8; i < NumParticles; ++i)
+			{
+				Particles.X(i) = TVector<float, 3>(Stream.FRandRange(-1.f, 1.f), Stream.FRandRange(-1.f, 1.f), Stream.FRandRange(-1.f, 1.f));
+			}
+			TArray<TVector<int32, 3>> Indices;
+			Chaos::FConvexBuilder::BuildConvexHull(Particles, Indices);
+			//EXPECT_EQ(Indices.Num(), 12); todo(ocohen): handle coplaner verts
+			for (auto Tri : Indices)
+			{
+				for (int i = 0; i < 3; ++i)
+				{
+					TVector<T, 3> V = Particles.X(Tri[i]);
+					TVector<T, 3> VAbs = V.GetAbs();
+					T Max = VAbs.GetMax();
+					EXPECT_GE(Max, 1 - 1e-2);
+				}
+			}
+		}
+	}
+	template void ConvexHull2<float>();
+
+	template <typename T>
+	void Simplify()
+	{
+		Chaos::TParticles<T, 3> Particles;
+		Particles.AddParticles(18);
+		Particles.X(0) = TVector<float, 3>(0, 0, 12.0f);
+		Particles.X(1) = TVector<float, 3>(-0.707f, -0.707f, 10.0f);
+		Particles.X(2) = TVector<float, 3>(0, -1, 10.0f);
+		Particles.X(3) = TVector<float, 3>(0.707f, -0.707f, 10.0f);
+		Particles.X(4) = TVector<float, 3>(1, 0, 10.0f);
+		Particles.X(5) = TVector<float, 3>(0.707f, 0.707f, 10.0f);
+		Particles.X(6) = TVector<float, 3>(0.0f, 1.0f, 10.0f);
+		Particles.X(7) = TVector<float, 3>(-0.707f, 0.707f, 10.0f);
+		Particles.X(8) = TVector<float, 3>(-1.0f, 0.0f, 10.0f);
+		Particles.X(9) = TVector<float, 3>(-0.707f, -0.707f, 0.0f);
+		Particles.X(10) = TVector<float, 3>(0, -1, 0.0f);
+		Particles.X(11) = TVector<float, 3>(0.707f, -0.707f, 0.0f);
+		Particles.X(12) = TVector<float, 3>(1, 0, 0.0f);
+		Particles.X(13) = TVector<float, 3>(0.707f, 0.707f, 0.0f);
+		Particles.X(14) = TVector<float, 3>(0.0f, 1.0f, 0.0f);
+		Particles.X(15) = TVector<float, 3>(-0.707f, 0.707f, 0.0f);
+		Particles.X(16) = TVector<float, 3>(-1.0f, 0.0f, 0.0f);
+		Particles.X(17) = TVector<float, 3>(0, 0, -2.0f);
+
+		FConvex Convex(Particles);
+
+		// capture original details
+		uint32 OriginalNumberParticles = Convex.GetSurfaceParticles().Size();
+		int32 OriginalNumberFaces = Convex.GetFaces().Num();
+		TBox<T, 3> OriginalBoundingBox = Convex.BoundingBox();
+
+		const TParticles<T, 3>& CulledParticles = Convex.GetSurfaceParticles();
+		const TArray<TPlaneConcrete<T, 3>> Planes = Convex.GetFaces();
+
+		// set target number of particles in simplified convex
+		FConvexBuilder::PerformGeometryReduction = 1;
+		FConvexBuilder::ParticlesThreshold = 10;
+
+		// simplify
+		Convex.PerformanceWarningAndSimplifaction();
+
+		// capture new details
+		uint32 NewNumberParticles = Convex.GetSurfaceParticles().Size();
+		int32 NewNumberFaces = Convex.GetFaces().Num();
+		TBox<T, 3> NewBoundingBox = Convex.BoundingBox();
+
+		EXPECT_EQ(OriginalNumberParticles, 18);
+		EXPECT_EQ(NewNumberParticles, 10);
+		EXPECT_LT(NewNumberFaces, OriginalNumberFaces);
+
+		TVector<T, 3> DiffMin = OriginalBoundingBox.Min() - NewBoundingBox.Min();
+		TVector<T, 3> DiffMax = OriginalBoundingBox.Max() - NewBoundingBox.Max();
+
+		// bounding box won't be identical, so long as it's not too far out
+		for (int Idx=0; Idx<3; Idx++)
+		{
+			EXPECT_LT(FMath::Abs(DiffMin[Idx]), 0.15f);
+			EXPECT_LT(FMath::Abs(DiffMax[Idx]), 0.15f);
+		}
+
+		FConvexBuilder::PerformGeometryReduction = 0;
+	}
+	template void Simplify<float>();
+
+	template <typename T>
+	void ImplicitScaled()
+	{
+		T Thickness = 0.1;
+		TUniquePtr<TSphere<T, 3>> Sphere = MakeUnique<TSphere<T,3>>(TVector<T, 3>(3, 0, 0), 5);
+		TImplicitObjectScaled<TSphere<T, 3>> Unscaled(MakeSerializable(Sphere), TVector<T,3>(1));
+		TImplicitObjectScaled<TSphere<T, 3>> UnscaledThickened(MakeSerializable(Sphere), TVector<T, 3>(1), Thickness);
+		TImplicitObjectScaled<TSphere<T, 3>> UniformScale(MakeSerializable(Sphere), TVector<T, 3>(2));
+		TImplicitObjectScaled<TSphere<T, 3>> UniformScaleThickened(MakeSerializable(Sphere), TVector<T, 3>(2), Thickness);
+		TImplicitObjectScaled<TSphere<T, 3>> NonUniformScale(MakeSerializable(Sphere), TVector<T, 3>(2, 1, 1));
+		TImplicitObjectScaled<TSphere<T, 3>> NonUniformScaleThickened(MakeSerializable(Sphere), TVector<T, 3>(2, 1, 1), Thickness);
+
+		//phi
+		{
+			const TVector<T, 3> NearEdge(7.5, 0, 0);
+			TVector<T, 3> UnscaledNormal;
+			const T UnscaledPhi = Unscaled.PhiWithNormal(NearEdge, UnscaledNormal);
+			EXPECT_FLOAT_EQ(UnscaledPhi, -0.5);
+			EXPECT_FLOAT_EQ(UnscaledNormal[0], 1);
+			EXPECT_FLOAT_EQ(UnscaledNormal[1], 0);
+			EXPECT_FLOAT_EQ(UnscaledNormal[2], 0);
+
+			TVector<T, 3> UnscaledNormalThickened;
+			const T UnscaledThickenedPhi = UnscaledThickened.PhiWithNormal(NearEdge, UnscaledNormalThickened);
+			EXPECT_FLOAT_EQ(UnscaledThickenedPhi, -0.5 - Thickness);
+			EXPECT_FLOAT_EQ(UnscaledNormalThickened[0], 1);
+			EXPECT_FLOAT_EQ(UnscaledNormalThickened[1], 0);
+			EXPECT_FLOAT_EQ(UnscaledNormalThickened[2], 0);
+
+			TVector<T, 3> ScaledNormal;
+			T ScaledPhi = UniformScale.PhiWithNormal(NearEdge, ScaledNormal);
+			EXPECT_FLOAT_EQ(ScaledPhi, -(16 - 7.5));
+			EXPECT_FLOAT_EQ(ScaledNormal[0], 1);
+			EXPECT_FLOAT_EQ(ScaledNormal[1], 0);
+			EXPECT_FLOAT_EQ(ScaledNormal[2], 0);
+
+			TVector<T, 3> ScaledNormalThickened;
+			T ScaledPhiThickened = UniformScaleThickened.PhiWithNormal(NearEdge, ScaledNormalThickened);
+			EXPECT_FLOAT_EQ(ScaledPhiThickened, -(16 + Thickness * 2 - 7.5));
+			EXPECT_FLOAT_EQ(ScaledNormalThickened[0], 1);
+			EXPECT_FLOAT_EQ(ScaledNormalThickened[1], 0);
+			EXPECT_FLOAT_EQ(ScaledNormalThickened[2], 0);
+
+			const TVector<T, 3> NearTop(6, 0, 4.5);
+			ScaledPhi = UniformScale.PhiWithNormal(NearTop, ScaledNormal);
+			EXPECT_FLOAT_EQ(ScaledPhi, -(10-4.5));
+			EXPECT_FLOAT_EQ(ScaledNormal[0], 0);
+			EXPECT_FLOAT_EQ(ScaledNormal[1], 0);
+			EXPECT_FLOAT_EQ(ScaledNormal[2], 1);
+
+			ScaledPhiThickened = UniformScaleThickened.PhiWithNormal(NearTop, ScaledNormalThickened);
+			EXPECT_FLOAT_EQ(ScaledPhiThickened, -(10 + Thickness*2 - 4.5));
+			EXPECT_FLOAT_EQ(ScaledNormalThickened[0], 0);
+			EXPECT_FLOAT_EQ(ScaledNormalThickened[1], 0);
+			EXPECT_FLOAT_EQ(ScaledNormalThickened[2], 1);
+
+			ScaledPhi = NonUniformScale.PhiWithNormal(NearTop, ScaledNormal);
+			EXPECT_FLOAT_EQ(ScaledPhi, -0.5);
+			EXPECT_FLOAT_EQ(ScaledNormal[0], 0);
+			EXPECT_FLOAT_EQ(ScaledNormal[1], 0);
+			EXPECT_FLOAT_EQ(ScaledNormal[2], 1);
+
+			ScaledPhiThickened = NonUniformScaleThickened.PhiWithNormal(NearTop, ScaledNormalThickened);
+			EXPECT_FLOAT_EQ(ScaledPhiThickened, -0.5 - Thickness);
+			EXPECT_FLOAT_EQ(ScaledNormalThickened[0], 0);
+			EXPECT_FLOAT_EQ(ScaledNormalThickened[1], 0);
+			EXPECT_FLOAT_EQ(ScaledNormalThickened[2], 1);
+
+			ScaledPhiThickened = NonUniformScaleThickened.PhiWithNormal(NearEdge, ScaledNormalThickened);
+			EXPECT_FLOAT_EQ(ScaledPhiThickened, -(16 + Thickness * 2 - 7.5));
+			EXPECT_FLOAT_EQ(ScaledNormalThickened[0], 1);
+			EXPECT_FLOAT_EQ(ScaledNormalThickened[1], 0);
+			EXPECT_FLOAT_EQ(ScaledNormalThickened[2], 0);
+
+		}
+
+		//support
+		{
+			const TVector<T, 3> DirX(1, 0, 0);
+			TVector<T, 3> SupportPt = Unscaled.Support(DirX, 1);
+			EXPECT_FLOAT_EQ(SupportPt[0], 9);
+			EXPECT_FLOAT_EQ(SupportPt[1], 0);
+			EXPECT_FLOAT_EQ(SupportPt[2], 0);
+
+			SupportPt = UnscaledThickened.Support(DirX, 1);
+			EXPECT_FLOAT_EQ(SupportPt[0], 9+Thickness);
+			EXPECT_FLOAT_EQ(SupportPt[1], 0);
+			EXPECT_FLOAT_EQ(SupportPt[2], 0);
+
+			SupportPt = UniformScale.Support(DirX, 1);
+			EXPECT_FLOAT_EQ(SupportPt[0], 17);
+			EXPECT_FLOAT_EQ(SupportPt[1], 0);
+			EXPECT_FLOAT_EQ(SupportPt[2], 0);
+
+			SupportPt = UniformScaleThickened.Support(DirX, 1);
+			EXPECT_FLOAT_EQ(SupportPt[0], 17 + Thickness * 2);
+			EXPECT_FLOAT_EQ(SupportPt[1], 0);
+			EXPECT_FLOAT_EQ(SupportPt[2], 0);
+
+			const TVector<T, 3> DirZ(0, 0, -1);
+			SupportPt = UniformScale.Support(DirZ, 1);
+			EXPECT_FLOAT_EQ(SupportPt[0], 6);
+			EXPECT_FLOAT_EQ(SupportPt[1], 0);
+			EXPECT_FLOAT_EQ(SupportPt[2], -11);
+
+			SupportPt = UniformScaleThickened.Support(DirZ, 1);
+			EXPECT_FLOAT_EQ(SupportPt[0], 6);
+			EXPECT_FLOAT_EQ(SupportPt[1], 0);
+			EXPECT_FLOAT_EQ(SupportPt[2], -11 - Thickness * 2);
+
+			SupportPt = NonUniformScale.Support(DirX, 1);
+			EXPECT_FLOAT_EQ(SupportPt[0], 17);
+			EXPECT_FLOAT_EQ(SupportPt[1], 0);
+			EXPECT_FLOAT_EQ(SupportPt[2], 0);
+
+			SupportPt = NonUniformScaleThickened.Support(DirX, 1);
+			EXPECT_FLOAT_EQ(SupportPt[0], 17 + Thickness * 2);
+			EXPECT_FLOAT_EQ(SupportPt[1], 0);
+			EXPECT_FLOAT_EQ(SupportPt[2], 0);
+
+			SupportPt = NonUniformScale.Support(DirZ, 1);
+			EXPECT_FLOAT_EQ(SupportPt[0], 6);
+			EXPECT_FLOAT_EQ(SupportPt[1], 0);
+			EXPECT_FLOAT_EQ(SupportPt[2], -6);
+
+			SupportPt = NonUniformScaleThickened.Support(DirZ, 1);
+			EXPECT_FLOAT_EQ(SupportPt[0], 6);
+			EXPECT_FLOAT_EQ(SupportPt[1], 0);
+			EXPECT_FLOAT_EQ(SupportPt[2], -6 - Thickness);
+		}
+	}
+	template void ImplicitScaled<float>();
+
+}
diff --git a/Engine/Source/Programs/HeadlessChaos/Private/HeadlessChaosTestJoints.cpp b/Engine/Source/Programs/HeadlessChaos/Private/HeadlessChaosTestJoints.cpp
new file mode 100644
--- /dev/null
+++ b/Engine/Source/Programs/HeadlessChaos/Private/HeadlessChaosTestJoints.cpp
@@ -0,0 +1,600 @@
+// Copyright Epic Games, Inc. All Rights Reserved.
+
+#include "HeadlessChaos.h"
+#include "HeadlessChaosTestConstraints.h"
+
+#include "Modules/ModuleManager.h"
+#include "Chaos/ParticleHandle.h"
+#include "Chaos/PBDConstraintRule.h"
+#include "Chaos/PBDRigidsEvolution.h"
+#include "Chaos/PBDRigidsEvolutionGBF.h"
+#include "Chaos/PBDRigidParticles.h"
+#include "Chaos/PBDJointConstraints.h"
+#include "Chaos/PBDRigidDynamicSpringConstraints.h"
+#include "Chaos/PBDRigidSpringConstraints.h"
+#include "Chaos/Box.h"
+#include "Chaos/Plane.h"
+#include "Chaos/Sphere.h"
+#include "Chaos/ImplicitObject.h"
+#include "Chaos/ImplicitObjectTransformed.h"
+#include "Chaos/ImplicitObjectUnion.h"
+#include "Chaos/Levelset.h"
+#include "Chaos/UniformGrid.h"
+#include "Chaos/Utilities.h"
+
+namespace ChaosTest {
+
+	using namespace Chaos;
+
+
+	/**
+	 * Base class for joint constraint tests.
+	 * Initialize the particle and joint data in the test code and call Create()
+	 */
+	class FJointConstraintsTest : public FConstraintsTest
+	{
+	public:
+		FJointConstraintsTest(const int32 NumIterations, const FReal Gravity)
+			: FConstraintsTest(NumIterations, Gravity)
+			, JointsRule(Joints)
+		{
+			Evolution.AddConstraintRule(&JointsRule);
+		}
+
+		FPBDJointConstraintHandle* AddJoint(const TVector<TGeometryParticleHandle<FReal, 3>*, 2>& InConstrainedParticleIndices, const FVec3& InLocation)
+		{
+			return Joints.AddConstraint(InConstrainedParticleIndices, FRigidTransform3(InLocation, FRotation3::FromIdentity()));
+		}
+
+		virtual void Create()
+		{
+			for (int32 ParticleIndex = 0; ParticleIndex < ParticlePositions.Num(); ++ParticleIndex)
+			{
+				AddParticleBox(ParticlePositions[ParticleIndex], FRotation3::MakeFromEuler(FVec3(0.f, 0.f, 0.f)).GetNormalized(), ParticleSizes[ParticleIndex], ParticleMasses[ParticleIndex]);
+			}
+
+			for (int32 JointIndex = 0; JointIndex < JointPositions.Num(); ++JointIndex)
+			{
+				const TVector<TGeometryParticleHandle<FReal, 3>*, 2> ConstraintedParticleIds(GetParticle(JointParticleIndices[JointIndex][0]), GetParticle(JointParticleIndices[JointIndex][1]));
+				AddJoint(ConstraintedParticleIds, JointPositions[JointIndex]);
+			}
+		}
+
+		// Initial particles setup
+		TArray<FVec3> ParticlePositions;
+		TArray<FVec3> ParticleSizes;
+		TArray<FReal> ParticleMasses;
+
+		// Initial joints setup
+		TArray<FVec3> JointPositions;
+		TArray<TVector<int32, 2>> JointParticleIndices;
+
+		// Solver state
+		FPBDJointConstraints Joints;
+		TPBDConstraintIslandRule<FPBDJointConstraints> JointsRule;
+	};
+
+	/**
+	 * One kinematic, one dynamic particle connected by a ball-socket joint in the middle.
+	 */
+	void JointConstraint_Single()
+	{
+
+		const int32 NumIterations = 1;
+		const FReal Gravity = 980;
+
+		FJointConstraintsTest Test(NumIterations, Gravity);
+
+		Test.ParticlePositions =
+		{
+			{ (FReal)0, (FReal)0, (FReal)1000 },
+			{ (FReal)500, (FReal)0, (FReal)1000 },
+		};
+		Test.ParticleSizes =
+		{
+			{ (FReal)100, (FReal)100, (FReal)100 },
+			{ (FReal)100, (FReal)100, (FReal)100 },
+		};
+		Test.ParticleMasses =
+		{
+			(FReal)0,
+			(FReal)1,
+		};
+
+		Test.JointPositions =
+		{
+			{ (FReal)250, (FReal)0, (FReal)1000 },
+		};
+		Test.JointParticleIndices =
+		{
+			{ 0, 1 },
+		};
+
+		Test.Create();
+
+		const int32 Box1Id = 0;
+		const int32 Box2Id = 1;
+		const float ExpectedDistance = (Test.ParticlePositions[1] - Test.ParticlePositions[0]).Size();
+		const FVec3 Box2LocalSpaceJointPosition = Test.JointPositions[0] - Test.ParticlePositions[1];
+
+		const FReal Dt = 0.01f;
+		for (int32 i = 0; i < 100; ++i)
+		{
+			Test.Evolution.AdvanceOneTimeStep(Dt);
+			Test.Evolution.EndFrame(Dt);
+
+
+			// Joint position calculated from pose and local-space joint pos
+			const FVec3 Box2WorldSpaceJointPosition = Test.GetParticle(Box2Id)->R().RotateVector(Box2LocalSpaceJointPosition) + Test.GetParticle(Box2Id)->X();
+			EXPECT_LT((Box2WorldSpaceJointPosition - Test.JointPositions[0]).Size(), (FReal)0.1);
+
+			// Kinematic particle should not have moved
+			EXPECT_LT((Test.GetParticle(Box1Id)->X() - Test.ParticlePositions[0]).Size(), (FReal)0.1);
+		}
+	}
+
+	void JointConstraint_SingleMoveRoot()
+	{
+		const int32 NumIterations = 5;
+		const FReal Gravity = 0;
+		const FReal BoxSize = 1;
+		const FReal BoxMass = 1;
+		const FReal Dt = (FReal)1 / 20;
+		const FVec3 RootDelta(1 * BoxSize, 0, 0);
+
+		FJointConstraintsTest Test(NumIterations, Gravity);
+
+		Test.ParticlePositions =
+		{
+			{ (FReal)0, (FReal)0, (FReal)10 * BoxSize },
+			{ (FReal)0, (FReal)0, (FReal)5 * BoxSize },
+		};
+		Test.ParticleSizes =
+		{
+			{ (FReal)BoxSize, (FReal)BoxSize, (FReal)BoxSize },
+			{ (FReal)BoxSize, (FReal)BoxSize, (FReal)BoxSize },
+		};
+		Test.ParticleMasses =
+		{
+			(FReal)0,
+			(FReal)BoxMass,
+		};
+
+		Test.JointPositions =
+		{
+			Test.ParticlePositions[0],
+		};
+		Test.JointParticleIndices =
+		{
+			{ 0, 1 },
+		};
+
+		Test.Create();
+
+		const int32 Box1Id = 0;
+		const int32 Box2Id = 1;
+		const float ExpectedDistance = (Test.ParticlePositions[1] - Test.ParticlePositions[0]).Size();
+		const FVec3 Box2LocalSpaceJointPosition = Test.JointPositions[0] - Test.ParticlePositions[1];
+
+		// Everything should be in a stable state
+		for (int32 i = 0; i < 10; ++i)
+		{
+			Test.Evolution.AdvanceOneTimeStep(Dt);
+			Test.Evolution.EndFrame(Dt);
+
+			// Nothing should have moved
+			for (int32 ParticleIndex = 0; ParticleIndex < Test.ParticlePositions.Num(); ++ParticleIndex)
+			{
+				EXPECT_LT((Test.GetParticle(ParticleIndex)->X() - Test.ParticlePositions[ParticleIndex]).Size(), (FReal)0.1) << "Initial configuration instability on frame " << i;
+			}
+		}
+
+		// Move the kinematic body
+		const FVec3 RootPosition = Test.ParticlePositions[0] + RootDelta;
+		Test.GetParticle(Box1Id)->X() = RootPosition;
+
+		for (int32 i = 0; i < 1000; ++i)
+		{
+			Test.Evolution.AdvanceOneTimeStep(Dt);
+			Test.Evolution.EndFrame(Dt);
+
+			// Kinematic particle should have moved to animated position
+			EXPECT_LT((Test.GetParticle(Box1Id)->X() - RootPosition).Size(), (FReal)0.01 * BoxSize) << "Post-move instability on frame " << i;
+
+			// Particles should remain fixed distance apart (joint point is at Box1 location)
+			const FVec3 Delta = Test.GetParticle(Box2Id)->CastToRigidParticle()->P() - Test.GetParticle(Box1Id)->X();
+			const float Distance = Delta.Size();
+			EXPECT_NEAR(Distance, ExpectedDistance, (FReal)0.01 * BoxSize) << "Post-move instability on frame " << i;
+
+			// Joint position calculted from pose and local-space joint pos
+			const FVec3 Box2WorldSpaceJointPosition = Test.GetParticle(Box2Id)->R().RotateVector(Box2LocalSpaceJointPosition) + Test.GetParticle(Box2Id)->X();
+			EXPECT_LT((Box2WorldSpaceJointPosition - RootPosition).Size(), (FReal)0.01 * BoxSize) << "Post-move instability on frame " << i;
+		}
+	}
+
+
+	/**
+	 * Pendulum with animated root.
+	 */
+	void JointConstraint_SingleAnimated()
+	{
+		const int32 NumIterations = 5;
+		const FReal Gravity = 980;
+		const FReal BoxSize = 100;
+		const FReal BoxMass = 1000;
+		const FReal Dt = (FReal)1 / 20;
+		const FReal AnimPeriod = (FReal)2;
+		const FVec3 AnimDelta = FVec3(10 * BoxSize, 0, 0);
+
+		FJointConstraintsTest Test(NumIterations, Gravity);
+
+		Test.ParticlePositions =
+		{
+			{ (FReal)0, (FReal)0, (FReal)10 * BoxSize },
+			{ (FReal)0, (FReal)2 * BoxSize, (FReal)10 * BoxSize },
+		};
+		Test.ParticleSizes =
+		{
+			{ (FReal)BoxSize, (FReal)BoxSize, (FReal)BoxSize },
+			{ (FReal)BoxSize, (FReal)BoxSize, (FReal)BoxSize },
+		};
+		Test.ParticleMasses =
+		{
+			(FReal)0,
+			(FReal)BoxMass,
+		};
+
+		Test.JointPositions =
+		{
+			Test.ParticlePositions[0],
+		};
+		Test.JointParticleIndices =
+		{
+			{ 0, 1 },
+		};
+
+		Test.Create();
+
+		const int32 Box1Id = 0;
+		const int32 Box2Id = 1;
+		const float ExpectedDistance = (Test.ParticlePositions[1] - Test.ParticlePositions[0]).Size();
+		const FVec3 Box2LocalSpaceJointPosition = Test.JointPositions[0] - Test.ParticlePositions[1];
+
+		for (int32 i = 0; i < 1000; ++i)
+		{
+			const FReal Time = i * Dt;
+			const FVec3 RootOffset = FMath::Sin((FReal)2 * PI * Time / AnimPeriod) * AnimDelta;
+			const FVec3 RootPosition = Test.ParticlePositions[0] + RootOffset;
+
+			Test.GetParticle(Box1Id)->X() = RootPosition;
+
+			Test.Evolution.AdvanceOneTimeStep(Dt);
+			Test.Evolution.EndFrame(Dt);
+
+			// Kinematic particle should have moved to animated position
+			EXPECT_LT((Test.GetParticle(Box1Id)->X() - RootPosition).Size(), (FReal)1) << "Failed on frame " << i;
+
+			// Particles should remain fixed distance apart (joint point is at Box1 location)
+			const FVec3 Delta = Test.GetParticle(Box2Id)->CastToRigidParticle()->P() - Test.GetParticle(Box1Id)->X();
+			const float Distance = Delta.Size();
+			EXPECT_NEAR(Distance, ExpectedDistance, (FReal)1) << "Failed on frame " << i;
+
+			// Joint position calculated from pose and local-space joint pos
+			const FVec3 Box2WorldSpaceJointPosition = Test.GetParticle(Box2Id)->R().RotateVector(Box2LocalSpaceJointPosition) + Test.GetParticle(Box2Id)->X();
+			EXPECT_LT((Box2WorldSpaceJointPosition - RootPosition).Size(), (FReal)1) << "Failed on frame " << i;
+		}
+	}
+
+	/**
+	 * Pendulum with animated root.
+	 */
+	void JointConstraint_ShortChainAnimated()
+	{
+		const int32 NumIterations = 10;
+		const FReal Gravity = 980;
+		const FReal BoxSize = 100;
+		const FReal BoxMass = 1000;
+		const FReal Dt = (FReal)1 / 100;
+		const FReal AnimPeriod = (FReal)1;
+		const FVec3 AnimDelta = FVec3(5 * BoxSize, 0, 0);
+		const FReal AcceptableDistanceError = 5;
+
+		FJointConstraintsTest Test(NumIterations, Gravity);
+
+		Test.ParticlePositions =
+		{
+			{ (FReal)0, (FReal)0, (FReal)20 * BoxSize },
+			{ (FReal)0, (FReal)2 * BoxSize, (FReal)20 * BoxSize },
+			{ (FReal)0, (FReal)4 * BoxSize, (FReal)20 * BoxSize },
+		};
+		Test.ParticleSizes =
+		{
+			{ (FReal)BoxSize, (FReal)BoxSize, (FReal)BoxSize },
+			{ (FReal)BoxSize, (FReal)BoxSize, (FReal)BoxSize },
+			{ (FReal)BoxSize, (FReal)BoxSize, (FReal)BoxSize },
+		};
+		Test.ParticleMasses =
+		{
+			(FReal)0,
+			(FReal)BoxMass,
+			(FReal)BoxMass,
+		};
+
+		Test.JointPositions =
+		{
+			Test.ParticlePositions[0],
+			Test.ParticlePositions[1],
+		};
+		Test.JointParticleIndices =
+		{
+			{ 0, 1 },
+			{ 1, 2 },
+		};
+
+		Test.Create();
+
+		const FVec3 Box2LocalSpaceJointPosition = Test.JointPositions[0] - Test.ParticlePositions[1];
+
+		FReal MaxDistanceError = 0.0f;
+		int32 MaxDistanceErrorFrameIndex = INDEX_NONE;
+		for (int32 FrameIndex = 0; FrameIndex < 1000; ++FrameIndex)
+		{
+			const FReal Time = FrameIndex * Dt;
+			const FVec3 RootOffset = FMath::Sin((FReal)2 * PI * Time / AnimPeriod) * AnimDelta;
+			const FVec3 RootPosition = Test.ParticlePositions[0] + RootOffset;
+
+			Test.GetParticle(0)->X() = RootPosition;
+
+			Test.Evolution.GetCollisionDetector().GetBroadPhase().SetBoundsVelocityInflation(1);
+			Test.Evolution.AdvanceOneTimeStep(Dt);
+			Test.Evolution.EndFrame(Dt);
+
+			// Particles should remain fixed distance apart
+			for (int JointIndex = 0; JointIndex < Test.JointPositions.Num(); ++JointIndex)
+			{
+				const int32 ParticleIndex1 = Test.JointParticleIndices[JointIndex][0];
+				const int32 ParticleIndex2 = Test.JointParticleIndices[JointIndex][1];
+				const FVec3 Delta = Test.GetParticle(ParticleIndex2)->CastToRigidParticle()->P() - Test.GetParticle(ParticleIndex1)->X();
+				const float Distance = Delta.Size();
+				const float ExpectedDistance = (Test.ParticlePositions[ParticleIndex2] - Test.ParticlePositions[ParticleIndex1]).Size();
+				EXPECT_NEAR(Distance, ExpectedDistance, AcceptableDistanceError) << "Joint " << JointIndex << " on frame " << FrameIndex;
+
+				const float DistanceError = FMath::Abs(Distance - ExpectedDistance);
+				if (DistanceError > MaxDistanceError)
+				{
+					MaxDistanceError = DistanceError;
+					MaxDistanceErrorFrameIndex = FrameIndex;
+				}
+			}
+		}
+		EXPECT_LT(MaxDistanceError, AcceptableDistanceError) << "On frame " << MaxDistanceErrorFrameIndex;
+	}
+
+	/**
+	 * Pendulum with animated root.
+	 */
+	void JointConstraint_LongChainAnimated()
+	{
+		const int NumParticles = 10;
+		const int32 NumIterations = 20;
+		const FReal Gravity = 980;
+		const FReal BoxSize = 100;
+		const FReal BoxMass = 1000;
+		const FReal Dt = (FReal)1 / 20;
+		const FReal AnimPeriod = (FReal)1;
+		const FVec3 AnimDelta = FVec3(1 * BoxSize, 0, 0);
+		const FReal AcceptableDistanceError = 5;
+		const FReal Separation = 2 * BoxSize;
+		const FVec3 Begin = FVec3(0, 0, (NumParticles + 10) * Separation);
+		const FVec3 Dir = FVec3(0, 1, 0);
+		const bool bRandomizeConstraintOrder = true;
+
+		FMath::RandInit(1048604845);
+
+		// Create a chain of connected particles, with particle 0 fixed
+		FJointConstraintsTest Test(NumIterations, Gravity);
+		for (int32 ParticleIndex = 0; ParticleIndex < NumParticles; ++ParticleIndex)
+		{
+			Test.ParticlePositions.Add(Begin + ParticleIndex * Separation * Dir);
+			Test.ParticleSizes.Add({ BoxSize, BoxSize, BoxSize });
+			Test.ParticleMasses.Add((ParticleIndex == 0) ? (FReal)0 : BoxMass);
+		}
+		for (int32 JointIndex = 0; JointIndex < NumParticles - 1; ++JointIndex)
+		{
+			Test.JointPositions.Add(Test.ParticlePositions[JointIndex]);
+			Test.JointParticleIndices.Add({ JointIndex, JointIndex + 1 });
+		}
+
+		// Randomize constraint order
+		if (bRandomizeConstraintOrder)
+		{
+			for (int32 JointIndex = 0; JointIndex < Test.JointParticleIndices.Num(); ++JointIndex)
+			{
+				const int32 Index0 = JointIndex;
+				const int32 Index1 = FMath::RandRange(0, Test.JointParticleIndices.Num() - 1);
+				Test.JointPositions.Swap(Index0, Index1);
+				Test.JointParticleIndices.Swap(Index0, Index1);
+			}
+		}
+
+		Test.Create();
+		Test.Evolution.GetCollisionDetector().GetBroadPhase().SetBoundsVelocityInflation(1);
+
+		const FVec3 Box2LocalSpaceJointPosition = Test.JointPositions[0] - Test.ParticlePositions[1];
+
+		FReal MaxDistanceError = 0.0f;
+		int32 MaxDistanceErrorFrameIndex = INDEX_NONE;
+		for (int32 FrameIndex = 0; FrameIndex < 1000; ++FrameIndex)
+		{
+			const FReal Time = FrameIndex * Dt;
+			const FVec3 RootOffset = FMath::Sin((FReal)2 * PI * Time / AnimPeriod) * AnimDelta;
+			const FVec3 RootPosition = Test.ParticlePositions[0] + RootOffset;
+
+			Test.GetParticle(0)->X() = RootPosition;
+
+			Test.Evolution.AdvanceOneTimeStep(Dt);
+			Test.Evolution.EndFrame(Dt);
+
+			// Particles should remain fixed distance apart
+			for (int JointIndex = 0; JointIndex < Test.JointPositions.Num(); ++JointIndex)
+			{
+				const int32 ParticleIndex1 = Test.JointParticleIndices[JointIndex][0];
+				const int32 ParticleIndex2 = Test.JointParticleIndices[JointIndex][1];
+				const FVec3 Delta = Test.GetParticle(ParticleIndex2)->CastToRigidParticle()->P() - Test.GetParticle(ParticleIndex1)->X();
+				const float Distance = Delta.Size();
+				const float ExpectedDistance = (Test.ParticlePositions[ParticleIndex2] - Test.ParticlePositions[ParticleIndex1]).Size();
+				EXPECT_NEAR(Distance, ExpectedDistance, AcceptableDistanceError) << "Joint " << JointIndex << " on frame " << FrameIndex;
+
+				// Track largest error that exceeds threshold
+				const float DistanceError = FMath::Abs(Distance - ExpectedDistance);
+				if (DistanceError > MaxDistanceError)
+				{
+					MaxDistanceError = DistanceError;
+					MaxDistanceErrorFrameIndex = FrameIndex;
+				}
+			}
+		}
+		// Report the largest error and when it occurred if it exceeded the threshold
+		EXPECT_LT(MaxDistanceError, AcceptableDistanceError) << "On frame " << MaxDistanceErrorFrameIndex;
+	}
+
+	void SpringConstraint()
+	{
+		TUniquePtr<FChaosPhysicsMaterial> PhysicalMaterial = MakeUnique<FChaosPhysicsMaterial>();
+		PhysicalMaterial->Friction = 0;
+		PhysicalMaterial->Restitution = 0;
+		PhysicalMaterial->SleepingLinearThreshold = 0;
+		PhysicalMaterial->SleepingAngularThreshold = 0;
+		PhysicalMaterial->DisabledLinearThreshold = 0;
+		PhysicalMaterial->DisabledAngularThreshold = 0;
+
+		TPBDRigidsSOAs<FReal, 3> Particles;
+
+		auto StaticBox = AppendStaticParticleBox<FReal>(Particles, FVec3((FReal)100, (FReal)100, (FReal)100));
+		auto Box2 = AppendDynamicParticleBox<FReal>(Particles, FVec3((FReal)100, (FReal)100, (FReal)100));
+		StaticBox->X() = FVec3((FReal)0, (FReal)0, (FReal)1000);
+
+		Box2->X() = FVec3((FReal)500, (FReal)0, (FReal)1000);
+		Box2->P() = Box2->X();
+
+		TPBDRigidsEvolutionGBF<FReal, 3> Evolution(Particles);
+		TArray<TVector<TGeometryParticleHandle<FReal, 3>*,2>> Constraints = { TVector<TGeometryParticleHandle<FReal,3>*, 2>(StaticBox, Box2) };
+		TArray<FVec3> Points0 = { FVec3((FReal)100, (FReal)0, (FReal)1000) };
+		TArray<FVec3> Points1 = { FVec3((FReal)400, (FReal)0, (FReal)1000) };
+
+		Evolution.SetPhysicsMaterial(StaticBox, MakeSerializable(PhysicalMaterial));
+		Evolution.SetPhysicsMaterial(Box2, MakeSerializable(PhysicalMaterial));
+
+		auto JointConstraints = Chaos::TPBDRigidSpringConstraints<FReal, 3>(Points0, Points1, MoveTemp(Constraints), 1.f);
+		auto JointRule = Chaos::TPBDConstraintIslandRule<Chaos::TPBDRigidSpringConstraints<FReal, 3>>(JointConstraints);
+		Evolution.AddConstraintRule(&JointRule);
+
+		const FReal Dt = 0.01f;
+		for (int32 i = 0; i < 100; ++i)
+		{
+			Evolution.AdvanceOneTimeStep(Dt);
+			Evolution.EndFrame(Dt);
+			EXPECT_LT(FMath::Abs((Box2->R().RotateVector(FVec3((FReal)-100, (FReal)0, (FReal)0)) + Box2->X() - Points0[0]).Size() - 300.f), 0.1);
+		}
+	}
+
+	void DynamicSpringConstraint()
+	{
+		TUniquePtr<FChaosPhysicsMaterial> PhysicalMaterial = MakeUnique<FChaosPhysicsMaterial>();
+		PhysicalMaterial->Friction = 0;
+		PhysicalMaterial->Restitution = 0;
+		PhysicalMaterial->SleepingLinearThreshold = 0;
+		PhysicalMaterial->SleepingAngularThreshold = 0;
+		PhysicalMaterial->DisabledLinearThreshold = 0;
+		PhysicalMaterial->DisabledAngularThreshold = 0;
+
+		{
+			TPBDRigidsSOAs<FReal, 3> Particles;
+
+			auto& StaticBox = *AppendStaticParticleBox<FReal>(Particles, FVec3((FReal)100, (FReal)100, (FReal)100));
+			auto& Box2 = *AppendDynamicParticleBox<FReal>(Particles, FVec3((FReal)100, (FReal)100, (FReal)100));
+			StaticBox.X() = FVec3((FReal)0, (FReal)0, (FReal)500);
+
+			Box2.X() = FVec3((FReal)500, (FReal)0, (FReal)1000);
+			Box2.P() = Box2.X();
+
+			TPBDRigidsEvolutionGBF<FReal, 3> Evolution(Particles);
+			TArray<TVector<TGeometryParticleHandle<FReal,3>*, 2>> Constraints = { TVector<TGeometryParticleHandle<FReal,3>*, 2>(&StaticBox, &Box2) };
+
+			Evolution.SetPhysicsMaterial(&StaticBox, MakeSerializable(PhysicalMaterial));
+			Evolution.SetPhysicsMaterial(&Box2, MakeSerializable(PhysicalMaterial));
+
+			Chaos::TPBDRigidDynamicSpringConstraints<FReal, 3> SpringConstraints(MoveTemp(Constraints));
+			auto SpringRule = Chaos::TPBDConstraintIslandRule<Chaos::TPBDRigidDynamicSpringConstraints<FReal, 3>>(SpringConstraints);
+			Evolution.AddConstraintRule(&SpringRule);
+
+			const FReal Dt = 0.01f;
+			for (int32 i = 0; i < 200; ++i)
+			{
+				Evolution.AdvanceOneTimeStep(Dt);
+				Evolution.EndFrame(Dt);
+			}
+			EXPECT_LT(Box2.X()[2], 0);
+		}
+
+		{
+			TPBDRigidsSOAs<FReal, 3> Particles;
+
+			auto& StaticBox = *AppendStaticParticleBox<FReal>(Particles, FVec3((FReal)100, (FReal)100, (FReal)100));
+			auto& Box2 = *AppendDynamicParticleBox<FReal>(Particles, FVec3((FReal)100, (FReal)100, (FReal)100));
+			StaticBox.X() = FVec3((FReal)0, (FReal)0, (FReal)500);
+
+			Box2.X() = FVec3((FReal)500, (FReal)0, (FReal)1000);
+			Box2.P() = Box2.X();
+
+			TPBDRigidsEvolutionGBF<FReal, 3> Evolution(Particles);
+			TArray<TVector<TGeometryParticleHandle<FReal,3>*, 2>> Constraints = { TVector<TGeometryParticleHandle<FReal,3>*, 2>(&StaticBox, &Box2) };
+
+			Evolution.SetPhysicsMaterial(&StaticBox, MakeSerializable(PhysicalMaterial));
+			Evolution.SetPhysicsMaterial(&Box2, MakeSerializable(PhysicalMaterial));
+
+			Chaos::TPBDRigidDynamicSpringConstraints<FReal, 3> SpringConstraints(MoveTemp(Constraints), 400);
+			auto SpringRule = Chaos::TPBDConstraintIslandRule<Chaos::TPBDRigidDynamicSpringConstraints<FReal, 3>>(SpringConstraints);
+			Evolution.AddConstraintRule(&SpringRule);
+
+			const FReal Dt = 0.01f;
+			for (int32 i = 0; i < 200; ++i)
+			{
+				Evolution.AdvanceOneTimeStep(Dt);
+				Evolution.EndFrame(Dt);
+			}
+			EXPECT_GT(Box2.X()[2], 0);
+		}
+	}
+
+
+
+	TEST(JointTests, TestSingleConstraint) {
+		JointConstraint_Single();
+	}
+
+	TEST(JointTests, TestSingleConstraintWithLateralTranslation) {
+		JointConstraint_SingleMoveRoot();
+	}
+
+	TEST(JointTests, TestSingleConstraintWithAnimatedRoot) {
+		JointConstraint_SingleAnimated();
+	}
+
+	TEST(JointTests, TestShortJointChainWithAnimatedRoot) {
+		JointConstraint_ShortChainAnimated();
+	}
+
+	TEST(JointTests, TestLongJointChainWithAnimatedRoot) {
+		JointConstraint_LongChainAnimated();
+	}
+
+	TEST(JointTests, TestSingleSpringConstraint) {
+		SpringConstraint();
+	}
+
+	TEST(JointTests, TestSingleDynamicSpringConstraint) {
+		DynamicSpringConstraint();
+	}
+}
+
