Rename/move HeadlessChaos.

diff --git a/Engine/Source/Programs/HeadlessChaos/Private/HeadlessChaosTestMostOpposing.cpp b/Engine/Source/Programs/HeadlessChaos/Private/HeadlessChaosTestMostOpposing.cpp
new file mode 100644
--- /dev/null
+++ b/Engine/Source/Programs/HeadlessChaos/Private/HeadlessChaosTestMostOpposing.cpp
@@ -0,0 +1,222 @@
+// Copyright Epic Games, Inc. All Rights Reserved.
+
+#include "HeadlessChaosTestMostOpposing.h"
+
+#include "HeadlessChaos.h"
+#include "Chaos/TriangleMeshImplicitObject.h"
+#include "Chaos/ImplicitObjectScaled.h"
+#include "Chaos/Convex.h"
+
+namespace ChaosTest
+{
+	using namespace Chaos;
+
+	/*We  want to test the following:
+	- Correct face index simple case
+	- Correct face on shared edge
+	*/
+
+	template<class T>
+	void TrimeshMostOpposing()
+	{
+		T Time;
+		TVector<T, 3> Position;
+		TVector<T, 3> Normal;
+		TArray<uint16> DummyMaterials;
+		int32 FaceIndex;
+
+		TParticles<T, 3> Particles;
+		Particles.AddParticles(6);
+		Particles.X(0) = TVector<T, 3>(1, 1, 1);
+		Particles.X(1) = TVector<T, 3>(5, 1, 1);
+		Particles.X(2) = TVector<T, 3>(1, 5, 1);
+
+		Particles.X(3) = TVector<T, 3>(1, 1, 1);
+		Particles.X(4) = TVector<T, 3>(1, 5, 1);
+		Particles.X(5) = TVector<T, 3>(1, 1, -5);
+
+		TArray<TVector<int32, 3>> Indices;
+		Indices.Emplace(0, 1, 2);
+		Indices.Emplace(3, 4, 5);
+		FTriangleMeshImplicitObject Tri(MoveTemp(Particles), MoveTemp(Indices), MoveTemp(DummyMaterials));
+
+		//simple into the triangle
+		bool bHit = Tri.Raycast(TVector<T, 3>(3, 2, 2), TVector<T, 3>(0, 0, -1), 2, 0, Time, Position, Normal, FaceIndex);
+		EXPECT_TRUE(bHit);
+		EXPECT_EQ(Tri.FindMostOpposingFace(Position, TVector<T, 3>(0, 0, -1), FaceIndex, 0.01), 0);
+		EXPECT_EQ(Tri.GetFaceNormal(0).X, Normal.X);
+		EXPECT_EQ(Tri.GetFaceNormal(0).Y, Normal.Y);
+		EXPECT_EQ(Tri.GetFaceNormal(0).Z, Normal.Z);
+
+		//simple into second triangle
+		bHit = Tri.Raycast(TVector<T, 3>(0, 2, 0), TVector<T, 3>(1, 0, 0), 2, 0, Time, Position, Normal, FaceIndex);
+		EXPECT_TRUE(bHit);
+		EXPECT_EQ(FaceIndex, 1);
+		EXPECT_EQ(Tri.FindMostOpposingFace(Position, TVector<T, 3>(1, 0, 0), FaceIndex, 0.01), 1);
+		EXPECT_EQ(Tri.GetFaceNormal(1).X, Normal.X);
+		EXPECT_EQ(Tri.GetFaceNormal(1).Y, Normal.Y);
+		EXPECT_EQ(Tri.GetFaceNormal(1).Z, Normal.Z);
+
+		//very close to edge, for now just return face hit regardless of direction because that's the implementation we currently rely on.
+		//todo: inconsistent with hulls, should make them the same, but may have significant impact on existing content
+
+		bHit = Tri.Raycast(TVector<T, 3>(0, 2, 0.9), TVector<T, 3>(1, 0, 0), 2, 0, Time, Position, Normal, FaceIndex);
+		EXPECT_TRUE(bHit);
+		EXPECT_EQ(FaceIndex, 1);
+		EXPECT_EQ(Tri.FindMostOpposingFace(Position, TVector<T, 3>(1, 0, 0), FaceIndex, 0.01), 1);
+		EXPECT_EQ(Tri.FindMostOpposingFace(Position, TVector<T, 3>(0, 0, -1), FaceIndex, 0.01), 1);	//ignores direction completely as per current implementation
+	}
+
+	template void TrimeshMostOpposing<float>();
+
+	template<class T>
+	void ConvexMostOpposing()
+	{
+		T Time;
+		TVector<T, 3> Position;
+		TVector<T, 3> Normal;
+		int32 FaceIndex;
+
+		TParticles<T, 3> Particles;
+		Particles.AddParticles(6);
+		Particles.X(0) = TVector<T, 3>(1, 1, 1);
+		Particles.X(1) = TVector<T, 3>(5, 1, 1);
+		Particles.X(2) = TVector<T, 3>(1, 5, 1);
+
+		Particles.X(3) = TVector<T, 3>(1, 1, 1);
+		Particles.X(4) = TVector<T, 3>(1, 5, 1);
+		Particles.X(5) = TVector<T, 3>(1, 1, -5);
+
+		FConvex Convex(MoveTemp(Particles));
+
+		//simple into the triangle
+		bool bHit = Convex.Raycast(TVector<T, 3>(3, 2, 2), TVector<T, 3>(0, 0, -1), 2, 0, Time, Position, Normal, FaceIndex);
+		EXPECT_TRUE(bHit);
+		EXPECT_FLOAT_EQ(Time, 1);
+		EXPECT_FLOAT_EQ(Position.X, 3);
+		EXPECT_FLOAT_EQ(Position.Y, 2);
+		EXPECT_FLOAT_EQ(Position.Z, 1);
+		EXPECT_EQ(FaceIndex, INDEX_NONE);	//convex should not compute its own face index as this is too expensive
+		EXPECT_EQ(Convex.FindMostOpposingFace(Position, TVector<T, 3>(0, 0, -1), FaceIndex, 0.01), 1);	//front face, just so happens that convex hull generates the planes in this order
+
+		//simple into second triangle
+		bHit = Convex.Raycast(TVector<T, 3>(0, 2, 0), TVector<T, 3>(1, 0, 0), 2, 0, Time, Position, Normal, FaceIndex);
+		EXPECT_TRUE(bHit);
+		EXPECT_FLOAT_EQ(Time, 1);
+		EXPECT_FLOAT_EQ(Position.X, 1);
+		EXPECT_FLOAT_EQ(Position.Y, 2);
+		EXPECT_FLOAT_EQ(Position.Z, 0);
+		EXPECT_EQ(FaceIndex, INDEX_NONE);	//convex should not compute its own face index as this is too expensive
+		EXPECT_EQ(Convex.FindMostOpposingFace(Position, TVector<T, 3>(1, 0, 0), FaceIndex, 0.01), 3);	//side face, just so happens that convex hull generates the planes in this order
+
+		bHit = Convex.Raycast(TVector<T, 3>(0, 2, 0.99), TVector<T, 3>(1, 0, 0), 2, 0, Time, Position, Normal, FaceIndex);
+		EXPECT_TRUE(bHit);
+		EXPECT_EQ(FaceIndex, INDEX_NONE);
+		EXPECT_EQ(Convex.FindMostOpposingFace(Position, TVector<T, 3>(1, 0, 0), FaceIndex, 0.01), 3);
+		EXPECT_EQ(Convex.FindMostOpposingFace(Position, TVector<T, 3>(0, 0, -1), FaceIndex, 0.01), 1);
+
+		//again but far enough away from edge
+		bHit = Convex.Raycast(TVector<T, 3>(0, 2, 0.9), TVector<T, 3>(1, 0, 0), 2, 0, Time, Position, Normal, FaceIndex);
+		EXPECT_TRUE(bHit);
+		EXPECT_EQ(FaceIndex, INDEX_NONE);
+		EXPECT_EQ(Convex.FindMostOpposingFace(Position, TVector<T, 3>(1, 0, 0), FaceIndex, 0.01), 3);
+		EXPECT_EQ(Convex.FindMostOpposingFace(Position, TVector<T, 3>(0, 0, -1), FaceIndex, 0.01), 3);	//too far to care about other face
+	}
+
+	template void ConvexMostOpposing<float>();
+
+	template<class T>
+	void ScaledMostOpposing()
+	{
+		T Time;
+		TVector<T, 3> Position;
+		TVector<T, 3> Normal;
+		int32 FaceIndex;
+
+		TParticles<T, 3> Particles;
+		Particles.AddParticles(6);
+		Particles.X(0) = TVector<T, 3>(0, -1, 1);
+		Particles.X(1) = TVector<T, 3>(1, -1, -1);
+		Particles.X(2) = TVector<T, 3>(0, 1, 1);
+
+		Particles.X(3) = TVector<T, 3>(0, -1, 1);
+		Particles.X(4) = TVector<T, 3>(0, 1, 1);
+		Particles.X(5) = TVector<T, 3>(-1, -1, -1);
+
+		TUniquePtr<FImplicitObject> Convex = MakeUnique<FConvex>(MoveTemp(Particles));
+
+		//identity scale
+		{
+			TImplicitObjectScaledGeneric<T, 3> Scaled(MakeSerializable(Convex), TVector<T, 3>(1, 1, 1));
+
+			//simple into the triangle
+			bool bHit = Scaled.Raycast(TVector<T, 3>(0.5, 0, 2), TVector<T, 3>(0, 0, -1), 3, 0, Time, Position, Normal, FaceIndex);
+			EXPECT_TRUE(bHit);
+			EXPECT_EQ(Position.X, 0.5);
+			EXPECT_EQ(Position.Y, 0);
+			EXPECT_EQ(Position.Z, 2 - Time);
+			EXPECT_EQ(FaceIndex, INDEX_NONE);	//convex should not compute its own face index as this is too expensive
+			EXPECT_EQ(Scaled.FindMostOpposingFace(Position, TVector<T, 3>(0, 0, -1), FaceIndex, 0.01), 2);	//x+ face, just so happens that convex hull generates the planes in this order
+
+			//simple into second triangle
+			bHit = Scaled.Raycast(TVector<T, 3>(-2, 0, 0.5), TVector<T, 3>(1, 0, 0), 3, 0, Time, Position, Normal, FaceIndex);
+			EXPECT_TRUE(bHit);
+			EXPECT_EQ(Position.X, -2+Time);
+			EXPECT_EQ(Position.Y, 0);
+			EXPECT_EQ(Position.Z, 0.5);
+			EXPECT_EQ(FaceIndex, INDEX_NONE);	//convex should not compute its own face index as this is too expensive
+			EXPECT_EQ(Scaled.FindMostOpposingFace(Position, TVector<T, 3>(1, 0, 0), FaceIndex, 0.01), 3);	//x- face, just so happens that convex hull generates the planes in this order
+
+			bHit = Scaled.Raycast(TVector<T, 3>(-0.001, 0, 2), TVector<T, 3>(0, 0, -1), 3, 0, Time, Position, Normal, FaceIndex);
+			EXPECT_TRUE(bHit);
+			EXPECT_EQ(FaceIndex, INDEX_NONE);
+			EXPECT_EQ(Scaled.FindMostOpposingFace(Position, TVector<T, 3>(1, 0, 0), FaceIndex, 0.01), 3);
+			EXPECT_EQ(Scaled.FindMostOpposingFace(Position, TVector<T, 3>(-1, 0, 0), FaceIndex, 0.01), 2);
+
+			//again but far enough away from edge
+			bHit = Scaled.Raycast(TVector<T, 3>(-0.1, 0, 2), TVector<T, 3>(0, 0, -1), 3, 0, Time, Position, Normal, FaceIndex);
+			EXPECT_TRUE(bHit);
+			EXPECT_EQ(FaceIndex, INDEX_NONE);
+			EXPECT_EQ(Scaled.FindMostOpposingFace(Position, TVector<T, 3>(1, 0, 0), FaceIndex, 0.01), 3);
+			EXPECT_EQ(Scaled.FindMostOpposingFace(Position, TVector<T, 3>(-1, 0, 0), FaceIndex, 0.01), 3);	//too far to care about other face
+		}
+
+		//non-uniform scale
+		{
+			TImplicitObjectScaledGeneric<T, 3> Scaled(MakeSerializable(Convex), TVector<T, 3>(2, 1, 1));
+
+			//simple into the triangle
+			bool bHit = Scaled.Raycast(TVector<T, 3>(0.5, 0, 2), TVector<T, 3>(0, 0, -1), 3, 0, Time, Position, Normal, FaceIndex);
+			EXPECT_TRUE(bHit);
+			EXPECT_EQ(Position.X, 0.5);
+			EXPECT_EQ(Position.Y, 0);
+			EXPECT_EQ(Position.Z, 2 - Time);
+			EXPECT_EQ(FaceIndex, INDEX_NONE);	//convex should not compute its own face index as this is too expensive
+			EXPECT_EQ(Scaled.FindMostOpposingFace(Position, TVector<T, 3>(0, 0, -1), FaceIndex, 0.01), 2);	//x+ face, just so happens that convex hull generates the planes in this order
+
+			//simple into second triangle
+			bHit = Scaled.Raycast(TVector<T, 3>(-2, 0, 0.5), TVector<T, 3>(1, 0, 0), 3, 0, Time, Position, Normal, FaceIndex);
+			EXPECT_TRUE(bHit);
+			EXPECT_EQ(Position.X, -2 + Time);
+			EXPECT_EQ(Position.Y, 0);
+			EXPECT_EQ(Position.Z, 0.5);
+			EXPECT_EQ(FaceIndex, INDEX_NONE);	//convex should not compute its own face index as this is too expensive
+			EXPECT_EQ(Scaled.FindMostOpposingFace(Position, TVector<T, 3>(1, 0, 0), FaceIndex, 0.01), 3);	//x- face, just so happens that convex hull generates the planes in this order
+
+			bHit = Scaled.Raycast(TVector<T, 3>(-0.001, 0, 2), TVector<T, 3>(0, 0, -1), 3, 0, Time, Position, Normal, FaceIndex);
+			EXPECT_TRUE(bHit);
+			EXPECT_EQ(FaceIndex, INDEX_NONE);
+			EXPECT_EQ(Scaled.FindMostOpposingFace(Position, TVector<T, 3>(1, 0, 0), FaceIndex, 0.01), 3);
+			EXPECT_EQ(Scaled.FindMostOpposingFace(Position, TVector<T, 3>(-1, 0, 0), FaceIndex, 0.01), 2);
+
+			//again but far enough away from edge
+			bHit = Scaled.Raycast(TVector<T, 3>(-0.1, 0, 2), TVector<T, 3>(0, 0, -1), 3, 0, Time, Position, Normal, FaceIndex);
+			EXPECT_TRUE(bHit);
+			EXPECT_EQ(FaceIndex, INDEX_NONE);
+			EXPECT_EQ(Scaled.FindMostOpposingFace(Position, TVector<T, 3>(1, 0, 0), FaceIndex, 0.01), 3);
+			EXPECT_EQ(Scaled.FindMostOpposingFace(Position, TVector<T, 3>(-1, 0, 0), FaceIndex, 0.01), 3);	//too far to care about other face
+		}
+	}
+
+	template void ScaledMostOpposing<float>();
+}
\ No newline at end of file
diff --git a/Engine/Source/Programs/HeadlessChaos/Private/HeadlessChaosTestParticleHandle.cpp b/Engine/Source/Programs/HeadlessChaos/Private/HeadlessChaosTestParticleHandle.cpp
new file mode 100644
--- /dev/null
+++ b/Engine/Source/Programs/HeadlessChaos/Private/HeadlessChaosTestParticleHandle.cpp
@@ -0,0 +1,581 @@
+// Copyright Epic Games, Inc. All Rights Reserved.
+
+#include "HeadlessChaosTestParticleHandle.h"
+
+#include "HeadlessChaos.h"
+#include "HeadlessChaosTestUtility.h"
+#include "Chaos/PBDRigidParticles.h"
+#include "Chaos/ParticleHandle.h"
+#include "Chaos/PBDRigidsSOAs.h"
+
+namespace ChaosTest {
+
+	using namespace Chaos;
+
+	template <typename T>
+	void ParticleIteratorTest()
+	{
+		auto Empty = MakeUnique<TGeometryParticles<T, 3>>();
+		auto Five = MakeUnique<TGeometryParticles<T, 3>>();
+		Five->AddParticles(5);
+		auto Two = MakeUnique<TGeometryParticles<T, 3>>();
+		Two->AddParticles(2);
+		//empty soa in the start
+		{
+			TArray<TGeometryParticleHandle<T, 3>*> Handles;
+			TArray<TSOAView<TGeometryParticles<T, 3>>> TmpArray = { Empty.Get(), Five.Get(), Two.Get() };
+			TParticleView<TGeometryParticles<T, 3>> View = MakeParticleView(MoveTemp(TmpArray));
+			for (auto& Particle : View)
+			{
+				Handles.Add(Particle.Handle());
+			}
+			EXPECT_EQ(Handles.Num(), 7);
+
+			THandleView<TGeometryParticles<T, 3>> HandleView = MakeHandleView(Handles);
+			int32 Count = 0;
+			for (auto& Handle : HandleView)
+			{
+				++Count;
+			}
+			EXPECT_EQ(Count, 7);
+		}
+
+		//empty soa in the middle
+		{
+			TArray<TGeometryParticleHandle<T, 3>*> Handles;
+			TArray<TSOAView<TGeometryParticles<T, 3>>> TmpArray = { Five.Get(), Empty.Get(), Two.Get()};
+			TParticleView<TGeometryParticles<T, 3>> View = MakeParticleView(MoveTemp(TmpArray));
+			for (auto& Particle : View)
+			{
+				Handles.Add(Particle.Handle());
+			}
+			EXPECT_EQ(Handles.Num(), 7);
+
+			THandleView<TGeometryParticles<T, 3>> HandleView = MakeHandleView(Handles);
+			int32 Count = 0;
+			for (auto& Handle : HandleView)
+			{
+				++Count;
+			}
+			EXPECT_EQ(Count, 7);
+		}
+
+		//empty soa in the end
+		{
+			TArray<TGeometryParticleHandle<T, 3>*> Handles;
+			TArray<TSOAView<TGeometryParticles<T, 3>>> TmpArray = { Five.Get(), Two.Get(), Empty.Get() };
+			TParticleView<TGeometryParticles<T, 3>> View = MakeParticleView(MoveTemp(TmpArray));
+			for (auto& Particle : View)
+			{
+				Handles.Add(Particle.Handle());
+			}
+			EXPECT_EQ(Handles.Num(), 7);
+
+			THandleView<TGeometryParticles<T, 3>> HandleView = MakeHandleView(Handles);
+			int32 Count = 0;
+			for (auto& Handle : HandleView)
+			{
+				++Count;
+			}
+			EXPECT_EQ(Count, 7);
+		}
+
+		//parallel for
+		{
+			TArray<TSOAView<TGeometryParticles<T, 3>>> TmpArray = { Empty.Get(), Five.Get(), Two.Get() };
+			TParticleView<TGeometryParticles<T, 3>> View = MakeParticleView(MoveTemp(TmpArray));
+			{
+				TArray<bool> AuxArray;
+				AuxArray.SetNumZeroed(View.Num());
+				bool DoubleWrite = false;
+				View.ParallelFor([&AuxArray, &DoubleWrite](const auto& Particle, int32 Idx)
+				{
+					if (AuxArray[Idx])
+					{
+						DoubleWrite = true;
+					}
+					AuxArray[Idx] = true;
+				});
+
+				EXPECT_FALSE(DoubleWrite);
+
+				for (bool Val : AuxArray)
+				{
+					EXPECT_TRUE(Val);
+				}
+			}
+
+			TArray<TGeometryParticleHandle<T, 3>*> Handles;
+			for (auto& Particle : View)
+			{
+				Handles.Add(Particle.Handle());
+			}
+			THandleView<TGeometryParticles<T, 3>> HandleView = MakeHandleView(Handles);
+			
+			{
+				TArray<bool> AuxArray;
+				AuxArray.SetNumZeroed(HandleView.Num());
+				bool DoubleWrite = false;
+				HandleView.ParallelFor([&AuxArray, &DoubleWrite](const auto& Particle, int32 Idx)
+				{
+					if (AuxArray[Idx])
+					{
+						DoubleWrite = true;
+					}
+					AuxArray[Idx] = true;
+				});
+
+				EXPECT_FALSE(DoubleWrite);
+
+				for (bool Val : AuxArray)
+				{
+					EXPECT_TRUE(Val);
+				}
+			}
+		}
+	}
+
+	template void ParticleIteratorTest<float>();
+
+	template <typename T, typename TGeometry, typename TKinematicGeometry, typename TPBDRigid>
+	void ParticleHandleTestHelper(TGeometry* Geometry, TKinematicGeometry* KinematicGeometry, TPBDRigid* PBDRigid)
+	{
+		EXPECT_EQ(Geometry->X()[0], 0);	//default constructor
+		EXPECT_EQ(Geometry->X()[1], 0);
+		EXPECT_EQ(Geometry->X()[2], 0);
+
+		EXPECT_EQ(KinematicGeometry->V()[0], 0);	//default constructor
+		EXPECT_EQ(KinematicGeometry->V()[1], 0);
+		EXPECT_EQ(KinematicGeometry->V()[2], 0);
+
+		EXPECT_EQ(PBDRigid->X()[0], 0);	//default constructor of base
+		EXPECT_EQ(PBDRigid->X()[1], 0);
+		EXPECT_EQ(PBDRigid->X()[2], 0);
+		EXPECT_EQ(PBDRigid->V()[0], 0);
+		EXPECT_EQ(PBDRigid->V()[1], 0);
+		EXPECT_EQ(PBDRigid->V()[2], 0);
+		EXPECT_EQ(PBDRigid->M(), 1);
+
+		PBDRigid->SetX(TVector<T, 3>(1, 2, 3));
+		EXPECT_EQ(PBDRigid->X()[0], 1);
+		KinematicGeometry->SetV(TVector<T, 3>(3, 3, 3));
+		EXPECT_EQ(KinematicGeometry->V()[0], 3);
+
+		EXPECT_EQ(Geometry->ObjectState(), EObjectStateType::Static);
+		EXPECT_EQ(KinematicGeometry->ObjectState(), EObjectStateType::Kinematic);
+		TGeometry* KinematicAsStatic = KinematicGeometry;	//shows polymorphism works
+		EXPECT_EQ(KinematicAsStatic->ObjectState(), EObjectStateType::Kinematic);
+
+		TGeometry* DynamicAsStatic = PBDRigid;
+		EXPECT_EQ(DynamicAsStatic->ObjectState(), EObjectStateType::Dynamic);
+		EXPECT_EQ(DynamicAsStatic->X()[0], 1);
+
+		//more polymorphism
+		ParticleHandleTestHelperObjectState(PBDRigid);
+	}
+
+	template <typename TPBDRigid>
+	void ParticleHandleTestHelperObjectState(TPBDRigid* PBDRigid);
+
+	template <>
+	void ParticleHandleTestHelperObjectState<TPBDRigidParticle<float,3>>(TPBDRigidParticle<float,3>* PBDRigid)
+	{
+		PBDRigid->SetObjectState(EObjectStateType::Dynamic);
+		EXPECT_EQ(PBDRigid->ObjectState(), EObjectStateType::Dynamic);
+	}
+
+	template <>
+	void ParticleHandleTestHelperObjectState<TPBDRigidParticleHandleImp<float,3,true>>(TPBDRigidParticleHandleImp<float,3,true>* PBDRigid)
+	{
+		PBDRigid->SetObjectStateLowLevel(EObjectStateType::Dynamic);
+		EXPECT_EQ(PBDRigid->ObjectState(), EObjectStateType::Dynamic);
+	}
+
+	template <typename T>
+	void ParticleLifetimeAndThreading()
+	{
+		{
+			TPBDRigidsSOAs<T, 3> SOAs;
+
+			TArray<TUniquePtr<TPBDRigidParticle<T, 3>>> GTRawParticles;
+			for (int i = 0; i < 3; ++i)
+			{
+				GTRawParticles.Emplace(TPBDRigidParticle<T, 3>::CreateParticle());
+			}
+			{
+				//for each GT particle, create a physics thread side
+				SOAs.CreateDynamicParticles(3);
+
+				//Solver sets the game thread particle on the physics thread handle
+				int32 Idx = 0;
+				for (auto& Particle : SOAs.GetAllParticlesView())
+				{
+					Particle.GTGeometryParticle() = GTRawParticles[Idx++].Get();
+				}
+				
+				T Count = 0;
+				//fake step and write to physics side
+				for (auto& Particle : SOAs.GetAllParticlesView())
+				{
+					Particle.X() = TVector<T, 3>(Count);
+					Count += 1;
+				}
+			}
+
+			//copy step to GT data
+			{
+				for (const auto& Particle : SOAs.GetAllParticlesView())
+				{
+					Particle.GTGeometryParticle()->SetX(Particle.X());
+				}
+			}
+
+			//consume on GT using raw pointers
+			for (int i = 0; i < 3; ++i)
+			{
+				EXPECT_EQ(GTRawParticles[i]->X()[0], i);
+			}
+
+			//GT destroys a particle by enqueing a command and nulling out its own raw pointer
+			auto RawParticleToDelete = GTRawParticles[1].Get();
+			GTRawParticles[1] = nullptr;
+			//PT does the actual delete
+			//GT would hold a private pointer that the solver can access
+			//SOAs.DestroyParticle(RawParticleToDelete->GetPhysicsTreadHandle());
+			//For now we just search
+			for (auto& Particle : SOAs.GetAllParticlesView())
+			{
+				if (Particle.GTGeometryParticle() == RawParticleToDelete)
+				{
+					SOAs.DestroyParticle(Particle.Handle());	//GT data will be removed
+					break;
+				}
+			}
+
+			//make sure we deleted the right particle
+			EXPECT_EQ(SOAs.GetAllParticlesView().Num(), 2);
+
+			for (auto& Particle : SOAs.GetAllParticlesView())
+			{
+				EXPECT_TRUE(Particle.X()[0] != 1);
+			}
+		}
+	}
+
+	template <typename T>
+	void ParticleDestroyOrdering()
+	{
+		{
+			TPBDRigidsSOAs<T, 3> SOAs;
+			SOAs.CreateDynamicParticles(10);
+			T Count = 0;
+			TGeometryParticleHandle<T, 3>* ThirdParticle = nullptr;
+			for (auto& Particle : SOAs.GetAllParticlesView())
+			{
+				Particle.X() = TVector<T, 3>(Count);
+				if (Count == 2)
+				{
+					ThirdParticle = Particle.Handle();
+				}
+
+				Count += 1;
+			}
+			EXPECT_EQ(ThirdParticle->X()[0], 2);
+
+			SOAs.DestroyParticle(ThirdParticle);
+			//default behavior is swap dynamics at end
+			Count = 0;
+			for (auto& Particle : SOAs.GetAllParticlesView())
+			{
+				if (Count == 2)
+				{
+					EXPECT_EQ(Particle.X()[0], 9);
+				}
+				else
+				{
+					EXPECT_EQ(Particle.X()[0], Count);
+				}
+
+				Count += 1;
+			}
+		}
+
+		//now test non swapping remove
+		{
+			TPBDRigidsSOAs<T, 3> SOAs;
+			SOAs.CreateClusteredParticles(10);
+			T Count = 0;
+			TGeometryParticleHandle<T, 3>* ThirdParticle = nullptr;
+			for (auto& Particle : SOAs.GetAllParticlesView())
+			{
+				Particle.X() = TVector<T, 3>(Count);
+				if (Count == 2)
+				{
+					ThirdParticle = Particle.Handle();
+				}
+
+				Count += 1;
+			}
+			EXPECT_EQ(ThirdParticle->X()[0], 2);
+
+			/*
+			//For now we're just disabling removing clustered all together
+			SOAs.DestroyParticle(ThirdParticle);
+			//default behavior is swap dynamics at end
+			Count = 0;
+			for (auto& Particle : SOAs.GetAllParticlesView())
+			{
+				if (Count < 2)
+				{
+					EXPECT_EQ(Particle.X()[0], Count);
+				}
+				else
+				{
+					EXPECT_EQ(Particle.X()[0], Count+1);
+				}
+
+				Count += 1;
+			}
+			*/
+		}
+	}
+
+	template<class T>
+	void ParticleHandleTest()
+	{
+		{
+			auto GeometryParticles = MakeUnique<TGeometryParticles<T, 3>>();
+			GeometryParticles->AddParticles(1);
+
+			auto KinematicGeometryParticles = MakeUnique<TKinematicGeometryParticles<T, 3>>();
+			KinematicGeometryParticles->AddParticles(1);
+
+			auto PBDRigidParticles = MakeUnique<TPBDRigidParticles<T, 3>>();
+			PBDRigidParticles->AddParticles(1);
+			
+			auto PartialPBDRigids = MakeUnique<TPBDRigidParticles<T, 3>>();
+			PartialPBDRigids->AddParticles(10);
+			
+			auto Geometry = TGeometryParticleHandle<T, 3>::CreateParticleHandle(MakeSerializable(GeometryParticles), 0, INDEX_NONE);
+
+			auto KinematicGeometry = TKinematicGeometryParticleHandle<T, 3>::CreateParticleHandle(MakeSerializable(KinematicGeometryParticles), 0, INDEX_NONE);
+
+			auto PBDRigid = TPBDRigidParticleHandle<T, 3>::CreateParticleHandle(MakeSerializable(PBDRigidParticles), 0, INDEX_NONE);
+
+			ParticleHandleTestHelper<T>(Geometry.Get(), static_cast<TKinematicGeometryParticleHandle<T,3>*>(KinematicGeometry.Get()), static_cast<TPBDRigidParticleHandle<T, 3>*>(PBDRigid.Get()));
+
+			//Test particle iterator
+			{
+				TGeometryParticleHandle<T, 3>* GeomHandles[] = { Geometry.Get(), KinematicGeometry.Get(), PBDRigid.Get() };
+				TArray<TSOAView<TGeometryParticles<T, 3>>> SOAViews = { GeometryParticles.Get(), KinematicGeometryParticles.Get(), PBDRigidParticles.Get() };
+				int32 Count = 0;
+				for (auto Itr = MakeParticleIterator(SOAViews); Itr; ++Itr)
+				{
+					//set X back to 0 for all particles
+					Itr->X() = TVector<T, 3>(0);
+					EXPECT_EQ(Itr->Handle(), GeomHandles[Count]);
+					//implicit const
+					TConstParticleIterator<TGeometryParticles<T, 3>>& ConstItr = Itr;
+					EXPECT_EQ(ConstItr->Handle(), GeomHandles[Count]);
+					++Count;
+				}
+
+				for (auto Itr = MakeConstParticleIterator(SOAViews); Itr; ++Itr)
+				{
+					//check Xs are back to 0
+					EXPECT_EQ(Itr->X()[0], 0);
+				}
+
+				Count = 0;
+				for (auto Itr = MakeConstParticleIterator(SOAViews); Itr; ++Itr)
+				{
+					//check InvM for dynamics
+					const TTransientPBDRigidParticleHandle<T, 3>* PBDRigid2 = Itr->CastToRigidParticle();
+					if (PBDRigid2 && PBDRigid2->ObjectState() == EObjectStateType::Dynamic)
+					{
+						++Count;
+						EXPECT_EQ(PBDRigid2->InvM(), 1);
+						EXPECT_EQ(PBDRigid2->Handle(), PBDRigid.Get());
+					}
+				}
+				EXPECT_EQ(Count, 1);
+			}
+
+			{
+				TArray<TSOAView<TPBDRigidParticles<T, 3>>> SOAViews = { PBDRigidParticles.Get() };
+				TPBDRigidParticleHandle<T, 3>* PBDRigidHandles[] = { static_cast<TPBDRigidParticleHandle<T, 3>*>(PBDRigid.Get()) };
+				int32 Count = 0;
+				for (auto Itr = MakeParticleIterator(MoveTemp(SOAViews)); Itr; ++Itr)
+				{
+					//set P to 1,1,1
+					Itr->P() = TVector<T, 3>(1);
+					EXPECT_EQ(Itr->Handle(), PBDRigidHandles[Count++]);
+					EXPECT_EQ(Itr->Handle()->P()[0], Itr->P()[0]);	//handle type is deduced from iterator type
+				}
+				EXPECT_EQ(Count, 1);
+			}
+
+			//Use an SOA with an active list
+			{
+				TPBDRigidsSOAs<T, 3> SOAsWithHandles;	//todo: create a mock object so we can more easily create handles
+				auto PartialDynamics = SOAsWithHandles.CreateDynamicParticles(10);
+
+				TArray<TPBDRigidParticleHandle<T,3>*> ActiveParticles = { PartialDynamics[3], PartialDynamics[5] };
+				PartialDynamics[3]->X() = TVector<T,3>(3);
+				PartialDynamics[5]->X() = TVector<T,3>(5);
+				
+				TArray<TSOAView<TPBDRigidParticles<T, 3>>> SOAViews = { PBDRigidParticles.Get(), &ActiveParticles, PBDRigidParticles.Get() };
+				int32 Count = 0;
+				for (auto Itr = MakeParticleIterator(MoveTemp(SOAViews)); Itr; ++Itr)
+				{
+					if (Count == 1)
+					{
+						EXPECT_EQ(Itr->X()[0], 3);
+					}
+
+					if (Count == 2)
+					{
+						EXPECT_EQ(Itr->X()[0], 5);
+					}
+					++Count;
+				}
+				EXPECT_EQ(Count, 4);
+			}
+		}
+
+		{
+			// try game thread representation
+			auto Geometry = TGeometryParticle<T, 3>::CreateParticle();
+			auto KinematicGeometry = TKinematicGeometryParticle<T, 3>::CreateParticle();
+			auto PBDRigid = TPBDRigidParticle<T, 3>::CreateParticle();
+			ParticleHandleTestHelper<T>(Geometry.Get(), KinematicGeometry.Get(), PBDRigid.Get());
+		}
+
+		{
+			// try using SOA manager
+			TPBDRigidsSOAs<T, 3> SOAs;
+			SOAs.CreateStaticParticles(3);
+			auto KinematicParticles = SOAs.CreateKinematicParticles(3);
+			SOAs.CreateDynamicParticles(3);
+
+			EXPECT_EQ(SOAs.GetNonDisabledView().Num(), 9);
+
+			//move to disabled
+			T Count = 0;
+			for (auto& Kinematic : KinematicParticles)
+			{
+				Kinematic->X() = TVector<T, 3>(Count);
+				SOAs.DisableParticle(Kinematic);
+				Count += 1;
+			}
+
+			EXPECT_EQ(SOAs.GetNonDisabledView().Num(), 6);
+
+			//values are still set
+			EXPECT_EQ(KinematicParticles[0]->X()[0], 0);
+			EXPECT_EQ(KinematicParticles[1]->X()[0], 1);
+			EXPECT_EQ(KinematicParticles[2]->X()[0], 2);
+
+			//move to enabled
+			for (auto& Kinematic : KinematicParticles)
+			{
+				SOAs.EnableParticle(Kinematic);
+			}
+
+			EXPECT_EQ(SOAs.GetNonDisabledView().Num(), 9);
+
+			//destroy particle
+			SOAs.DestroyParticle(KinematicParticles[0]);
+			KinematicParticles[0] = nullptr;
+
+			EXPECT_EQ(SOAs.GetNonDisabledView().Num(), 8);
+
+			SOAs.DestroyParticle(KinematicParticles[2]);
+			KinematicParticles[2] = nullptr;
+
+			EXPECT_EQ(SOAs.GetNonDisabledView().Num(), 7);
+
+			//disable some and then delete all
+			SOAs.DisableParticle(KinematicParticles[1]);
+
+			TArray<TGeometryParticleHandle<T, 3>*> ToDelete;
+			for (auto& Particle : SOAs.GetAllParticlesView())	//todo: add check that iterator invalidates during delete
+			{
+				ToDelete.Add(Particle.Handle());
+			}
+
+			for (auto& Handle : ToDelete)
+			{
+				SOAs.DestroyParticle(Handle);
+			}
+			EXPECT_EQ(SOAs.GetNonDisabledView().Num(), 0);
+		}
+
+		ParticleLifetimeAndThreading<T>();
+		ParticleDestroyOrdering<T>();
+	}
+	template void ParticleHandleTest<float>();
+
+	void AccelerationStructureHandleComparison()
+	{
+		// When an external particle is created, the handle is retrieved via proxy.
+		// Proxy gets handle initialized async on PT later, so this means a handle 
+		// will always have external particle pointer, and eventually an internal pointer.
+		// because of this, we must be able to compare (external, null) == (external, null) 
+		// and also (external, null) == (external, internal)
+
+
+		auto GTParticle = TGeometryParticle<FReal, 3>::CreateParticle();
+		TAccelerationStructureHandle<FReal, 3> ExternalOnlyHandle(GTParticle.Get());
+
+		TPBDRigidsSOAs<FReal, 3> SOAs;
+		auto Particles = SOAs.CreateStaticParticles(1);
+		TAccelerationStructureHandle<FReal, 3> ExternalInternalHandle(Particles[0], GTParticle.Get());
+
+		TAccelerationStructureHandle<FReal, 3> InternalOnlyHandle(Particles[0], nullptr);
+
+		TAccelerationStructureHandle<FReal, 3> NullHandle;
+
+		EXPECT_EQ(ExternalOnlyHandle, ExternalInternalHandle);
+		EXPECT_EQ(ExternalOnlyHandle, ExternalOnlyHandle);
+
+		// Disabled because operator== ensures on null handle.
+		/*EXPECT_EQ(NullHandle, NullHandle);
+		EXPECT_NE(NullHandle, ExternalOnlyHandle);
+		EXPECT_NE(NullHandle, ExternalInternalHandle);
+		EXPECT_NE(NullHandle, InternalOnlyHandle);*/
+
+		EXPECT_NE(InternalOnlyHandle, ExternalOnlyHandle);
+		EXPECT_EQ(InternalOnlyHandle, ExternalInternalHandle);
+	}
+
+	void HandleObjectStateChangeTest()
+	{
+		TPBDRigidsSOAs<FReal, 3> SOAs;
+
+		// Lambda for adding a particle to the dynamic-backed kinematic SOA
+		const auto CreateDynamicKinematic = [&]()
+		{
+			TPBDRigidParticleHandle<FReal, 3>* Particle = SOAs.CreateDynamicParticles(1)[0];
+			Particle->SetObjectStateLowLevel(EObjectStateType::Kinematic);
+			SOAs.SetDynamicParticleSOA(Particle);
+			return Particle;
+		};
+
+		// Create two dynamic kinematics, move one of them back to the dynamic SOA
+		auto* Particle0 = CreateDynamicKinematic();
+		auto* Particle1 = CreateDynamicKinematic();
+		Particle0->SetObjectStateLowLevel(EObjectStateType::Dynamic);
+		SOAs.SetDynamicParticleSOA(Particle0);
+
+		// Ensure only one dynamic is in Active Particles
+		auto ActiveIt = SOAs.GetActiveParticlesView().Begin();
+		EXPECT_EQ(ActiveIt->ObjectState(), EObjectStateType::Dynamic);
+		EXPECT_EQ(SOAs.GetActiveParticlesView().Num(), 1);
+
+		// Ensure setting to kinematic removes it
+		Particle0->SetObjectStateLowLevel(EObjectStateType::Kinematic);
+		SOAs.SetDynamicParticleSOA(Particle0);
+		EXPECT_EQ(SOAs.GetActiveParticlesView().Num(), 0);
+	}
+}
\ No newline at end of file
diff --git a/Engine/Source/Programs/HeadlessChaos/Private/HeadlessChaosTestRaycast.cpp b/Engine/Source/Programs/HeadlessChaos/Private/HeadlessChaosTestRaycast.cpp
new file mode 100644
--- /dev/null
+++ b/Engine/Source/Programs/HeadlessChaos/Private/HeadlessChaosTestRaycast.cpp
@@ -0,0 +1,606 @@
+// Copyright Epic Games, Inc. All Rights Reserved.
+
+#include "HeadlessChaosTestRaycast.h"
+
+#include "HeadlessChaos.h"
+#include "Chaos/Sphere.h"
+#include "Chaos/Capsule.h"
+#include "Chaos/TriangleMeshImplicitObject.h"
+#include "Chaos/ImplicitObjectScaled.h"
+
+namespace ChaosTest
+{
+	using namespace Chaos;
+
+	/*In general we want to test the following for each geometry type:
+	- time represents how far a swept object travels
+	- position represents the world position where an intersection first occurred. If multiple first intersections we should do something well defined (what?)
+	- normal represents the world normal where an intersection first occurred.
+	- time vs position (i.e. in a thick raycast we want point of impact)
+	- initial overlap blocks
+	- near hit
+	- near miss
+	*/
+
+	template<class T>
+	void SphereRaycast()
+	{
+		TSphere<T,3> Sphere(TVector<T, 3>(1), 15);
+
+		T Time;
+		TVector<T, 3> Position, Normal;
+		int32 FaceIndex;
+
+		//simple
+		bool bHit = Sphere.Raycast(TVector<T,3>(1,1,17), TVector<T, 3>(0, 0, -1), 30, 0, Time, Position, Normal, FaceIndex);
+		EXPECT_TRUE(bHit);
+		EXPECT_FLOAT_EQ(Time, (T)1);
+		
+		EXPECT_FLOAT_EQ(Normal.X, 0);
+		EXPECT_FLOAT_EQ(Normal.Y, 0);
+		EXPECT_FLOAT_EQ(Normal.Z, 1);
+
+		EXPECT_FLOAT_EQ(Position.X, 1);
+		EXPECT_FLOAT_EQ(Position.Y, 1);
+		EXPECT_FLOAT_EQ(Position.Z, 16);
+		EXPECT_EQ(FaceIndex, INDEX_NONE);
+
+		//initial overlap
+		bHit = Sphere.Raycast(TVector<T, 3>(1, 1, 14), TVector<T, 3>(0, 0, -1), 15, 0, Time, Position, Normal, FaceIndex);
+		EXPECT_TRUE(bHit);
+		EXPECT_FLOAT_EQ(Time, 0);
+		EXPECT_EQ(FaceIndex, INDEX_NONE);
+
+		//near hit
+		bHit = Sphere.Raycast(TVector<T, 3>(16, 1, 16), TVector<T, 3>(0, 0, -1), 30, 0, Time, Position, Normal, FaceIndex);
+		EXPECT_TRUE(bHit);
+		EXPECT_FLOAT_EQ(Time, 15);
+
+		EXPECT_FLOAT_EQ(Normal.X, 1);
+		EXPECT_FLOAT_EQ(Normal.Y, 0);
+		EXPECT_FLOAT_EQ(Normal.Z, 0);
+		EXPECT_EQ(FaceIndex, INDEX_NONE);
+
+		EXPECT_FLOAT_EQ(Position.X, 16);
+		EXPECT_FLOAT_EQ(Position.Y, 1);
+		EXPECT_FLOAT_EQ(Position.Z, 1);
+
+		//near miss
+		bHit = Sphere.Raycast(TVector<T, 3>(16 + 1e-4, 1, 16), TVector<T, 3>(0, 0, -1), 30, 0, Time, Position, Normal, FaceIndex);
+		EXPECT_FALSE(bHit);
+
+		//time vs position
+		bHit = Sphere.Raycast(TVector<T, 3>(21, 1, 16), TVector<T, 3>(0, 0, -1), 30, 5, Time, Position, Normal, FaceIndex);
+		EXPECT_TRUE(bHit);
+		EXPECT_FLOAT_EQ(Time, 15);
+		EXPECT_EQ(FaceIndex, INDEX_NONE);
+
+		EXPECT_FLOAT_EQ(Normal.X, 1);
+		EXPECT_FLOAT_EQ(Normal.Y, 0);
+		EXPECT_FLOAT_EQ(Normal.Z, 0);
+
+		EXPECT_FLOAT_EQ(Position.X, 16);
+		EXPECT_FLOAT_EQ(Position.Y, 1);
+		EXPECT_FLOAT_EQ(Position.Z, 1);
+
+		//passed miss
+		bHit = Sphere.Raycast(TVector<T, 3>(1, 1, -14 - 1e-4), TVector<T, 3>(0, 0, -1), 30, 0, Time, Position, Normal, FaceIndex);
+		EXPECT_FALSE(bHit);
+	}
+
+	template<class T>
+	void PlaneRaycast()
+	{
+		TPlane<T, 3> Plane(TVector<T, 3>(1), TVector<T, 3>(1, 0, 0));
+
+		T Time;
+		TVector<T, 3> Position, Normal;
+		int32 FaceIndex;
+
+		//simple
+		bool bHit = Plane.Raycast(TVector<T, 3>(2, 1, 1), TVector<T, 3>(-1, 0, 0), 2, 0, Time, Position, Normal, FaceIndex);
+		EXPECT_TRUE(bHit);
+		EXPECT_FLOAT_EQ(Time, 1);
+		EXPECT_EQ(FaceIndex, INDEX_NONE);
+
+		EXPECT_FLOAT_EQ(Normal.X, 1);
+		EXPECT_FLOAT_EQ(Normal.Y, 0);
+		EXPECT_FLOAT_EQ(Normal.Z, 0);
+
+		EXPECT_FLOAT_EQ(Position.X, 1);
+		EXPECT_FLOAT_EQ(Position.Y, 1);
+		EXPECT_FLOAT_EQ(Position.Z, 1);
+
+		//Other side of plane
+		bHit = Plane.Raycast(TVector<T, 3>(-1, 1, 1), TVector<T, 3>(1, 0, 0), 4, 0, Time, Position, Normal, FaceIndex);
+		EXPECT_TRUE(bHit);
+		EXPECT_FLOAT_EQ(Time, 2);
+		EXPECT_EQ(FaceIndex, INDEX_NONE);
+
+		EXPECT_FLOAT_EQ(Normal.X, -1);
+		EXPECT_FLOAT_EQ(Normal.Y, 0);
+		EXPECT_FLOAT_EQ(Normal.Z, 0);
+
+		EXPECT_FLOAT_EQ(Position.X, 1);
+		EXPECT_FLOAT_EQ(Position.Y, 1);
+		EXPECT_FLOAT_EQ(Position.Z, 1);
+
+		//initial overlap
+		bHit = Plane.Raycast(TVector<T, 3>(2, 1, 1), TVector<T, 3>(1, 0, 0), 2, 3, Time, Position, Normal, FaceIndex);
+		EXPECT_TRUE(bHit);
+		EXPECT_FLOAT_EQ(Time, 0);
+		EXPECT_EQ(FaceIndex, INDEX_NONE);
+
+		//near hit
+		bHit = Plane.Raycast(TVector<T, 3>(1+1, 1, 1), TVector<T, 3>(-1e-2, 0, 1).GetUnsafeNormal(), 100.01, 0, Time, Position, Normal, FaceIndex);
+		EXPECT_TRUE(bHit);
+		EXPECT_EQ(FaceIndex, INDEX_NONE);
+
+		EXPECT_FLOAT_EQ(Normal.X, 1);
+		EXPECT_FLOAT_EQ(Normal.Y, 0);
+		EXPECT_FLOAT_EQ(Normal.Z, 0);
+
+		EXPECT_FLOAT_EQ(Position.X, 1);
+		EXPECT_FLOAT_EQ(Position.Y, 1);
+		EXPECT_FLOAT_EQ(Position.Z, 101);
+
+		//near miss
+		bHit = Plane.Raycast(TVector<T, 3>(1 + 1, 1, 1), TVector<T, 3>(-1e-2, 0, 1).GetUnsafeNormal(), 99.9, 0, Time, Position, Normal, FaceIndex);
+		EXPECT_FALSE(bHit);
+
+		//time vs position
+		bHit = Plane.Raycast(TVector<T, 3>(-1, 1, 1), TVector<T, 3>(1, 0, 0), 4, 1, Time, Position, Normal, FaceIndex);
+		EXPECT_TRUE(bHit);
+		EXPECT_EQ(FaceIndex, INDEX_NONE);
+		EXPECT_FLOAT_EQ(Time, 1);
+		EXPECT_FLOAT_EQ(Position.X, 1);
+		EXPECT_FLOAT_EQ(Position.Y, 1);
+		EXPECT_FLOAT_EQ(Position.Z, 1);
+	}
+
+	template<class T>
+	void CapsuleRaycast()
+	{
+		T Time;
+		TVector<T, 3> Position;
+		TVector<T, 3> Normal;
+		int32 FaceIndex;
+
+		//straight down
+		TCapsule<T> Capsule(TVector<T, 3>(1, 1, 1), TVector<T, 3>(1, 1, 9), 1);
+		bool bHit = Capsule.Raycast(TVector<T, 3>(1, 1, 11), TVector<T, 3>(0, 0, -1), 2, 0, Time, Position, Normal, FaceIndex);
+		
+		EXPECT_TRUE(bHit);
+		EXPECT_FLOAT_EQ(Time, 1);
+		EXPECT_EQ(FaceIndex, INDEX_NONE);
+
+		EXPECT_FLOAT_EQ(Normal.X, 0);
+		EXPECT_FLOAT_EQ(Normal.Y, 0);
+		EXPECT_FLOAT_EQ(Normal.Z, 1);
+
+		EXPECT_FLOAT_EQ(Position.X, 1);
+		EXPECT_FLOAT_EQ(Position.Y, 1);
+		EXPECT_FLOAT_EQ(Position.Z, 10);
+
+		//straight up
+		bHit = Capsule.Raycast(TVector<T, 3>(1, 1, -1), TVector<T, 3>(0, 0, 1), 2, 0, Time, Position, Normal, FaceIndex);
+		EXPECT_TRUE(bHit);
+		EXPECT_FLOAT_EQ(Time, 1);
+		EXPECT_EQ(FaceIndex, INDEX_NONE);
+
+		EXPECT_FLOAT_EQ(Normal.X, 0);
+		EXPECT_FLOAT_EQ(Normal.Y, 0);
+		EXPECT_FLOAT_EQ(Normal.Z, -1);
+
+		EXPECT_FLOAT_EQ(Position.X, 1);
+		EXPECT_FLOAT_EQ(Position.Y, 1);
+		EXPECT_FLOAT_EQ(Position.Z, 0);
+
+		//cylinder
+		bHit = Capsule.Raycast(TVector<T, 3>(3, 1, 7), TVector<T, 3>(-1, 0, 0), 2, 0, Time, Position, Normal, FaceIndex);
+
+		EXPECT_TRUE(bHit);
+		EXPECT_FLOAT_EQ(Time, 1);
+		EXPECT_EQ(FaceIndex, INDEX_NONE);
+
+		EXPECT_FLOAT_EQ(Normal.X, 1);
+		EXPECT_FLOAT_EQ(Normal.Y, 0);
+		EXPECT_FLOAT_EQ(Normal.Z, 0);
+
+		EXPECT_FLOAT_EQ(Position.X, 2);
+		EXPECT_FLOAT_EQ(Position.Y, 1);
+		EXPECT_FLOAT_EQ(Position.Z, 7);
+
+		//cylinder away
+		bHit = Capsule.Raycast(TVector<T, 3>(3, 1, 7), TVector<T, 3>(1, 0, 0), 2, 0, Time, Position, Normal, FaceIndex);
+		EXPECT_FALSE(bHit);
+		
+		// initial overlap cap
+		bHit = Capsule.Raycast(TVector<T, 3>(1, 1, 9.5), TVector<T, 3>(-1, 0, 0), 2, 0, Time, Position, Normal, FaceIndex);
+		EXPECT_TRUE(bHit);
+		EXPECT_EQ(Time, 0);
+		EXPECT_EQ(FaceIndex, INDEX_NONE);
+
+		// initial overlap cylinder
+		bHit = Capsule.Raycast(TVector<T, 3>(1, 1, 7), TVector<T, 3>(-1, 0, 0), 2, 0, Time, Position, Normal, FaceIndex);
+		EXPECT_TRUE(bHit);
+		EXPECT_EQ(Time, 0);
+
+		//cylinder time vs position
+		bHit = Capsule.Raycast(TVector<T, 3>(4, 1, 7), TVector<T, 3>(-1, 0, 0), 4, 1, Time, Position, Normal, FaceIndex);
+
+		EXPECT_TRUE(bHit);
+		EXPECT_FLOAT_EQ(Time, 1);
+		EXPECT_EQ(FaceIndex, INDEX_NONE);
+
+		EXPECT_FLOAT_EQ(Normal.X, 1);
+		EXPECT_FLOAT_EQ(Normal.Y, 0);
+		EXPECT_FLOAT_EQ(Normal.Z, 0);
+
+		EXPECT_FLOAT_EQ(Position.X, 2);
+		EXPECT_FLOAT_EQ(Position.Y, 1);
+		EXPECT_FLOAT_EQ(Position.Z, 7);
+
+		//normal independent of ray dir
+		bHit = Capsule.Raycast(TVector<T, 3>(4, 1, 7), TVector<T, 3>(-1, 0, -1).GetUnsafeNormal(), 4, 1, Time, Position, Normal, FaceIndex);
+
+		EXPECT_TRUE(bHit);
+		EXPECT_EQ(FaceIndex, INDEX_NONE);
+
+		EXPECT_FLOAT_EQ(Normal.X, 1);
+		EXPECT_FLOAT_EQ(Normal.Y, 0);
+		EXPECT_FLOAT_EQ(Normal.Z, 0);
+
+		EXPECT_FLOAT_EQ(Position.X, 2);
+
+		//near hit orthogonal
+		bHit = Capsule.Raycast(TVector<T, 3>(2, 3, 7), TVector<T, 3>(0, -1, 0), 4, 0, Time, Position, Normal, FaceIndex);
+
+		EXPECT_TRUE(bHit);
+		EXPECT_FLOAT_EQ(Time, 2);
+		EXPECT_EQ(FaceIndex, INDEX_NONE);
+
+		EXPECT_FLOAT_EQ(Normal.X, 1);
+		EXPECT_FLOAT_EQ(Normal.Y, 0);
+		EXPECT_FLOAT_EQ(Normal.Z, 0);
+
+		EXPECT_FLOAT_EQ(Position.X, 2);
+		EXPECT_FLOAT_EQ(Position.Y, 1);
+		EXPECT_FLOAT_EQ(Position.Z, 7);
+
+		//near miss
+		bHit = Capsule.Raycast(TVector<T, 3>(2 + 1e-4, 3, 7), TVector<T, 3>(0, -1, 0), 4, 0, Time, Position, Normal, FaceIndex);
+		EXPECT_FALSE(bHit);
+
+		//near hit straight down
+		bHit = Capsule.Raycast(TVector<T, 3>(0, 1, 11), TVector<T, 3>(0, 0, -1), 20, 0, Time, Position, Normal, FaceIndex);
+		EXPECT_TRUE(bHit);
+		EXPECT_EQ(FaceIndex, INDEX_NONE);
+
+		EXPECT_FLOAT_EQ(Normal.X, -1);
+		EXPECT_FLOAT_EQ(Normal.Y, 0);
+		EXPECT_FLOAT_EQ(Normal.Z, 0);
+
+		EXPECT_FLOAT_EQ(Position.X, 0);
+		EXPECT_FLOAT_EQ(Position.Y, 1);
+		EXPECT_FLOAT_EQ(Position.Z, 9);
+
+		bHit = Capsule.Raycast(TVector<T, 3>(-1e-4, 1, 11), TVector<T, 3>(0, 0, -1), 20, 0, Time, Position, Normal, FaceIndex);
+		EXPECT_FALSE(bHit);
+	}
+
+	template <typename T>
+	void TriangleRaycast()
+	{
+		T Time;
+		TVector<T, 3> Position;
+		TVector<T, 3> Normal;
+		TArray<uint16> DummyMaterials;
+		int32 FaceIndex;
+
+		TParticles<T, 3> Particles;
+		Particles.AddParticles(3);
+		Particles.X(0) = TVector<T, 3>(1, 1, 1);
+		Particles.X(1) = TVector<T, 3>(5, 1, 1);
+		Particles.X(2) = TVector<T, 3>(1, 5, 1);
+		TArray<TVector<int32, 3>> Indices;
+		Indices.Emplace(0, 1, 2);
+		FTriangleMeshImplicitObject Tri(MoveTemp(Particles), MoveTemp(Indices), MoveTemp(DummyMaterials));
+
+		//simple into the triangle
+		bool bHit = Tri.Raycast(TVector<T, 3>(3, 2, 2), TVector<T, 3>(0, 0, -1), 2, 0, Time, Position, Normal, FaceIndex);
+		EXPECT_TRUE(bHit);
+		EXPECT_EQ(Time, 1);
+		EXPECT_EQ(FaceIndex, 0);
+
+		EXPECT_FLOAT_EQ(Normal.X, 0);
+		EXPECT_FLOAT_EQ(Normal.Y, 0);
+		EXPECT_FLOAT_EQ(Normal.Z, 1);
+
+		EXPECT_FLOAT_EQ(Position.X, 3);
+		EXPECT_FLOAT_EQ(Position.Y, 2);
+		EXPECT_FLOAT_EQ(Position.Z, 1);
+
+		//double sided
+		bHit = Tri.Raycast(TVector<T, 3>(3, 2, 0), TVector<T, 3>(0, 0, 1), 2, 0, Time, Position, Normal, FaceIndex);
+		EXPECT_TRUE(bHit);
+		EXPECT_EQ(Time, 1);
+		EXPECT_EQ(FaceIndex, 0);
+
+		EXPECT_FLOAT_EQ(Normal.X, 0);
+		EXPECT_FLOAT_EQ(Normal.Y, 0);
+		EXPECT_FLOAT_EQ(Normal.Z, -1);
+
+		EXPECT_FLOAT_EQ(Position.X, 3);
+		EXPECT_FLOAT_EQ(Position.Y, 2);
+		EXPECT_FLOAT_EQ(Position.Z, 1);
+
+		//time vs position
+		bHit = Tri.Raycast(TVector<T, 3>(3, 2, 3), TVector<T, 3>(0, 0, -1), 2, 1, Time, Position, Normal, FaceIndex);
+		EXPECT_TRUE(bHit);
+		EXPECT_EQ(Time, 1);
+		EXPECT_EQ(FaceIndex, 0);
+
+		EXPECT_FLOAT_EQ(Normal.X, 0);
+		EXPECT_FLOAT_EQ(Normal.Y, 0);
+		EXPECT_FLOAT_EQ(Normal.Z, 1);
+
+		EXPECT_FLOAT_EQ(Position.X, 3);
+		EXPECT_FLOAT_EQ(Position.Y, 2);
+		EXPECT_FLOAT_EQ(Position.Z, 1);
+
+		//initial miss, border hit
+		bHit = Tri.Raycast(TVector<T, 3>(0.5, 2, 3), TVector<T, 3>(0, 0, -1), 2, 1, Time, Position, Normal, FaceIndex);
+		EXPECT_TRUE(bHit);
+		EXPECT_EQ(FaceIndex, 0);
+
+		EXPECT_FLOAT_EQ(Normal.X, 0);
+		EXPECT_FLOAT_EQ(Normal.Y, 0);
+		EXPECT_FLOAT_EQ(Normal.Z, 1);
+
+		EXPECT_FLOAT_EQ(Position.X, 1);
+		EXPECT_FLOAT_EQ(Position.Y, 2);
+		EXPECT_FLOAT_EQ(Position.Z, 1);
+
+		//initial overlap with plane, but miss triangle
+		bHit = Tri.Raycast(TVector<T, 3>(10, 1, 1), TVector<T, 3>(0, 0, -1), 2, 1, Time, Position, Normal, FaceIndex);
+		EXPECT_FALSE(bHit);
+
+		//parallel with triangle
+		bHit = Tri.Raycast(TVector<T, 3>(-1, 1, 1), TVector<T, 3>(1, 0, 0), 2, 1, Time, Position, Normal, FaceIndex);
+		EXPECT_TRUE(bHit);
+		EXPECT_FLOAT_EQ(Time, 1);
+		EXPECT_EQ(FaceIndex, 0);
+
+		EXPECT_FLOAT_EQ(Normal.X, 0);
+		EXPECT_FLOAT_EQ(Normal.Y, 0);
+		EXPECT_FLOAT_EQ(Normal.Z, 1);
+
+		EXPECT_FLOAT_EQ(Position.X, 1);
+		EXPECT_FLOAT_EQ(Position.Y, 1);
+		EXPECT_FLOAT_EQ(Position.Z, 1);
+	}
+
+	template <typename T>
+	void BoxRaycast()
+	{
+		T Time;
+		TVector<T, 3> Position;
+		TVector<T, 3> Normal;
+		int32 FaceIndex;
+
+		TAABB<T, 3> Box(TVector<T, 3>(1, 1, 1), TVector<T, 3>(3, 5, 3));
+		
+		//simple into the box
+		bool bHit = Box.Raycast(TVector<T, 3>(2, 3, 4), TVector<T, 3>(0, 0, -1), 2, 0, Time, Position, Normal, FaceIndex);
+		EXPECT_TRUE(bHit);
+		EXPECT_FLOAT_EQ(Time, 1);
+		EXPECT_EQ(FaceIndex, INDEX_NONE);
+
+		EXPECT_FLOAT_EQ(Normal.X, 0);
+		EXPECT_FLOAT_EQ(Normal.Y, 0);
+		EXPECT_FLOAT_EQ(Normal.Z, 1);
+
+		EXPECT_FLOAT_EQ(Position.X, 2);
+		EXPECT_FLOAT_EQ(Position.Y, 3);
+		EXPECT_FLOAT_EQ(Position.Z, 3);
+
+		//time vs position
+		bHit = Box.Raycast(TVector<T, 3>(2, 3, 5), TVector<T, 3>(0, 0, -1), 2, 1, Time, Position, Normal, FaceIndex);
+		EXPECT_TRUE(bHit);
+		EXPECT_FLOAT_EQ(Time, 1);
+		EXPECT_EQ(FaceIndex, INDEX_NONE);
+
+		EXPECT_FLOAT_EQ(Normal.X, 0);
+		EXPECT_FLOAT_EQ(Normal.Y, 0);
+		EXPECT_FLOAT_EQ(Normal.Z, 1);
+
+		EXPECT_FLOAT_EQ(Position.X, 2);
+		EXPECT_FLOAT_EQ(Position.Y, 3);
+		EXPECT_FLOAT_EQ(Position.Z, 3);
+
+		//edge
+		bHit = Box.Raycast(TVector<T, 3>(0.5, 2, -1), TVector<T, 3>(0, 0, 1), 2, 1, Time, Position, Normal, FaceIndex);
+		EXPECT_TRUE(bHit);
+		EXPECT_EQ(FaceIndex, INDEX_NONE);
+				
+		EXPECT_FLOAT_EQ(Position.X, 1);
+		EXPECT_FLOAT_EQ(Position.Y, 2);
+		EXPECT_FLOAT_EQ(Position.Z, 1);
+
+		//corner
+		bHit = Box.Raycast(TVector<T, 3>(0.5, 1, -1), TVector<T, 3>(0, 0, 1), 2, 1, Time, Position, Normal, FaceIndex);
+		EXPECT_TRUE(bHit);
+		EXPECT_EQ(FaceIndex, INDEX_NONE);
+
+		EXPECT_FLOAT_EQ(Position.X, 1);
+		EXPECT_FLOAT_EQ(Position.Y, 1);
+		EXPECT_FLOAT_EQ(Position.Z, 1);
+
+		//near hit by corner edge
+		const TVector<T, 3> StartEmptyRegion(1 - FMath::Sqrt(2) / 2, 1 - FMath::Sqrt(2) / 2, -1);
+		bHit = Box.Raycast(StartEmptyRegion, TVector<T, 3>(0, 0, 1), 2, 1, Time, Position, Normal, FaceIndex);
+		EXPECT_TRUE(bHit);
+		EXPECT_EQ(FaceIndex, INDEX_NONE);
+
+		EXPECT_FLOAT_EQ(Position.X, 1);
+		EXPECT_FLOAT_EQ(Position.Y, 1);
+		EXPECT_FLOAT_EQ(Position.Z, 1);
+
+		//near miss by corner edge
+		const TVector<T, 3> StartEmptyRegionMiss(StartEmptyRegion[0] - 1e-4, StartEmptyRegion[1] - 1e-4, StartEmptyRegion[2]);
+		bHit = Box.Raycast(StartEmptyRegionMiss, TVector<T, 3>(0, 0, 1), 2, 1, Time, Position, Normal, FaceIndex);
+		EXPECT_FALSE(bHit);
+
+		//start in corner voronoi but end in edge voronoi
+		bHit = Box.Raycast(TVector<T,3>(0,0, 0.8), TVector<T, 3>(1, 1, 5).GetUnsafeNormal(), 2, 1, Time, Position, Normal, FaceIndex);
+
+		EXPECT_TRUE(bHit);
+		EXPECT_EQ(FaceIndex, INDEX_NONE);
+
+		EXPECT_FLOAT_EQ(Position.X, 1);
+		EXPECT_FLOAT_EQ(Position.Y, 1);
+		EXPECT_GT(Position.Z, 1);
+
+		//start in voronoi and miss
+		bHit = Box.Raycast(TVector<T, 3>(0, 0, 0.8), TVector<T, 3>(-1, -1, 0).GetUnsafeNormal(), 2, 1, Time, Position, Normal, FaceIndex);
+		EXPECT_FALSE(bHit);
+
+		//initial overlap
+		bHit = Box.Raycast(TVector<T, 3>(1, 1, 2), TVector<T, 3>(-1, -1, 0).GetUnsafeNormal(), 2, 1, Time, Position, Normal, FaceIndex);
+		EXPECT_TRUE(bHit);
+		EXPECT_EQ(FaceIndex, INDEX_NONE);
+		EXPECT_EQ(Time, 0);
+	}
+
+	template <typename T>
+	void ScaledRaycast()
+	{
+		T Time;
+		TVector<T, 3> Position;
+		TVector<T, 3> Normal;
+		int32 FaceIndex;
+		const T Thickness = 0.1;
+
+		TUniquePtr<TSphere<T, 3>> Sphere = MakeUnique<TSphere<T,3>>(TVector<T, 3>(1), 2);
+		TImplicitObjectScaled<TSphere<T, 3>> Unscaled(MakeSerializable(Sphere), TVector<T,3>(1));
+		TImplicitObjectScaled<TSphere<T, 3>> UnscaledThickened(MakeSerializable(Sphere), TVector<T, 3>(1), Thickness);
+		TImplicitObjectScaled<TSphere<T, 3>> UniformScaled(MakeSerializable(Sphere), TVector<T,3>(2));
+		TImplicitObjectScaled<TSphere<T, 3>> UniformScaledThickened(MakeSerializable(Sphere), TVector<T, 3>(2), Thickness);
+		TImplicitObjectScaled<TSphere<T, 3>> NonUniformScaled(MakeSerializable(Sphere), TVector<T,3>(2,1,1));
+		TImplicitObjectScaled<TSphere<T, 3>> NonUniformScaledThickened(MakeSerializable(Sphere), TVector<T, 3>(2, 1, 1), Thickness);
+
+		//simple
+		bool bHit = Unscaled.Raycast(TVector<T, 3>(1, 1, 8), TVector<T, 3>(0, 0, -1), 8, 0, Time, Position, Normal, FaceIndex);
+		EXPECT_TRUE(bHit);
+		EXPECT_FLOAT_EQ(Time, 5);
+		EXPECT_EQ(FaceIndex, INDEX_NONE);
+
+		EXPECT_FLOAT_EQ(Normal.X, 0);
+		EXPECT_FLOAT_EQ(Normal.Y, 0);
+		EXPECT_FLOAT_EQ(Normal.Z, 1);
+
+		EXPECT_FLOAT_EQ(Position.X, 1);
+		EXPECT_FLOAT_EQ(Position.Y, 1);
+		EXPECT_FLOAT_EQ(Position.Z, 3);
+
+		bHit = UnscaledThickened.Raycast(TVector<T, 3>(1, 1, 8), TVector<T, 3>(0, 0, -1), 8, 0, Time, Position, Normal, FaceIndex);
+		EXPECT_TRUE(bHit);
+		EXPECT_FLOAT_EQ(Time, 5 - Thickness);
+		EXPECT_EQ(FaceIndex, INDEX_NONE);
+
+		EXPECT_FLOAT_EQ(Normal.X, 0);
+		EXPECT_FLOAT_EQ(Normal.Y, 0);
+		EXPECT_FLOAT_EQ(Normal.Z, 1);
+
+		EXPECT_FLOAT_EQ(Position.X, 1);
+		EXPECT_FLOAT_EQ(Position.Y, 1);
+		EXPECT_FLOAT_EQ(Position.Z, 3);
+
+		bHit = UniformScaled.Raycast(TVector<T, 3>(2, 2, 8), TVector<T, 3>(0, 0, -1), 8, 0, Time, Position, Normal, FaceIndex);
+		EXPECT_TRUE(bHit);
+		EXPECT_FLOAT_EQ(Time, 2);
+		EXPECT_EQ(FaceIndex, INDEX_NONE);
+
+		EXPECT_FLOAT_EQ(Normal.X, 0);
+		EXPECT_FLOAT_EQ(Normal.Y, 0);
+		EXPECT_FLOAT_EQ(Normal.Z, 1);
+
+		EXPECT_FLOAT_EQ(Position.X, 2);
+		EXPECT_FLOAT_EQ(Position.Y, 2);
+		EXPECT_FLOAT_EQ(Position.Z, 6);
+
+		bHit = UniformScaledThickened.Raycast(TVector<T, 3>(2, 2, 8), TVector<T, 3>(0, 0, -1), 8, 0, Time, Position, Normal, FaceIndex);
+		EXPECT_TRUE(bHit);
+		EXPECT_FLOAT_EQ(Time, 2 - Thickness * 2);
+		EXPECT_EQ(FaceIndex, INDEX_NONE);
+
+		EXPECT_FLOAT_EQ(Normal.X, 0);
+		EXPECT_FLOAT_EQ(Normal.Y, 0);
+		EXPECT_FLOAT_EQ(Normal.Z, 1);
+
+		EXPECT_FLOAT_EQ(Position.X, 2);
+		EXPECT_FLOAT_EQ(Position.Y, 2);
+		EXPECT_FLOAT_EQ(Position.Z, 6);
+
+		bHit = NonUniformScaled.Raycast(TVector<T, 3>(2, 1, 8), TVector<T, 3>(0, 0, -1), 8, 0, Time, Position, Normal, FaceIndex);
+		EXPECT_TRUE(bHit);
+		EXPECT_FLOAT_EQ(Time, 5);
+		EXPECT_EQ(FaceIndex, INDEX_NONE);
+
+		EXPECT_FLOAT_EQ(Normal.X, 0);
+		EXPECT_FLOAT_EQ(Normal.Y, 0);
+		EXPECT_FLOAT_EQ(Normal.Z, 1);
+
+		EXPECT_FLOAT_EQ(Position.X, 2);
+		EXPECT_FLOAT_EQ(Position.Y, 1);
+		EXPECT_FLOAT_EQ(Position.Z, 3);
+
+		bHit = NonUniformScaledThickened.Raycast(TVector<T, 3>(2, 1, 8), TVector<T, 3>(0, 0, -1), 8, 0, Time, Position, Normal, FaceIndex);
+		EXPECT_TRUE(bHit);
+		EXPECT_FLOAT_EQ(Time, 5 - Thickness);
+		EXPECT_EQ(FaceIndex, INDEX_NONE);
+
+		EXPECT_FLOAT_EQ(Normal.X, 0);
+		EXPECT_FLOAT_EQ(Normal.Y, 0);
+		EXPECT_FLOAT_EQ(Normal.Z, 1);
+
+		EXPECT_FLOAT_EQ(Position.X, 2);
+		EXPECT_FLOAT_EQ(Position.Y, 1);
+		EXPECT_FLOAT_EQ(Position.Z, 3);
+
+		//scaled thickness
+		bHit = UniformScaled.Raycast(TVector<T, 3>(2, 2, 8), TVector<T, 3>(0, 0, -1), 8, 1, Time, Position, Normal, FaceIndex);
+		EXPECT_TRUE(bHit);
+		EXPECT_FLOAT_EQ(Time, 1);
+		EXPECT_EQ(FaceIndex, INDEX_NONE);
+
+		EXPECT_FLOAT_EQ(Normal.X, 0);
+		EXPECT_FLOAT_EQ(Normal.Y, 0);
+		EXPECT_FLOAT_EQ(Normal.Z, 1);
+
+		EXPECT_FLOAT_EQ(Position.X, 2);
+		EXPECT_FLOAT_EQ(Position.Y, 2);
+		EXPECT_FLOAT_EQ(Position.Z, 6);
+
+		bHit = UniformScaledThickened.Raycast(TVector<T, 3>(2, 2, 8), TVector<T, 3>(0, 0, -1), 8, 1, Time, Position, Normal, FaceIndex);
+		EXPECT_TRUE(bHit);
+		EXPECT_FLOAT_EQ(Time, 1 - Thickness * 2);
+		EXPECT_EQ(FaceIndex, INDEX_NONE);
+
+		EXPECT_FLOAT_EQ(Normal.X, 0);
+		EXPECT_FLOAT_EQ(Normal.Y, 0);
+		EXPECT_FLOAT_EQ(Normal.Z, 1);
+
+		EXPECT_FLOAT_EQ(Position.X, 2);
+		EXPECT_FLOAT_EQ(Position.Y, 2);
+		EXPECT_FLOAT_EQ(Position.Z, 6);
+	}
+
+	template void SphereRaycast<float>();
+	template void PlaneRaycast<float>();
+	template void CapsuleRaycast<float>();
+	template void TriangleRaycast<float>();
+	template void BoxRaycast<float>();
+	template void ScaledRaycast<float>();
+}
\ No newline at end of file
diff --git a/Engine/Source/Programs/HeadlessChaos/Private/HeadlessChaosTestSerialization.cpp b/Engine/Source/Programs/HeadlessChaos/Private/HeadlessChaosTestSerialization.cpp
new file mode 100644
--- /dev/null
+++ b/Engine/Source/Programs/HeadlessChaos/Private/HeadlessChaosTestSerialization.cpp
@@ -0,0 +1,451 @@
+// Copyright Epic Games, Inc. All Rights Reserved.
+
+#include "HeadlessChaosTestSerialization.h"
+
+#include "HeadlessChaos.h"
+#include "Chaos/ChaosArchive.h"
+#include "Serialization/MemoryWriter.h"
+#include "Serialization/MemoryReader.h"
+#include "Chaos/GeometryParticles.h"
+#include "Chaos/RigidParticles.h"
+#include "Chaos/ImplicitObjectTransformed.h"
+#include "Chaos/ImplicitObjectUnion.h"
+#include "Chaos/SerializationTestUtility.h"
+#include "Chaos/PBDRigidsSOAs.h"
+#include "Chaos/PBDRigidsEvolutionGBF.h"
+#include "Chaos/ChaosPerfTest.h"
+#include "Misc/Paths.h"
+#include "HAL/FileManager.h"
+#include "Chaos/BoundingVolumeHierarchy.h"
+namespace
+{
+	FString GetSerializedBinaryPath();
+}
+
+namespace ChaosTest
+{
+	using namespace Chaos;
+
+	FString GetSerializedBinaryPath()
+	{
+		return FPaths::EngineDir() / TEXT("Source/Programs/NotForLicensees/HeadlessChaos/SerializedBinaries");
+	}
+
+	template<class T>
+	void SimpleObjectsSerialization()
+	{
+
+		TArray<TUniquePtr<TSphere<T, 3>>> OriginalSpheres;
+		OriginalSpheres.Add(TUniquePtr<TSphere<T, 3>>(new TSphere<T, 3>(TVector<T, 3>(), 1)));
+		OriginalSpheres.Add(TUniquePtr<TSphere<T, 3>>(new TSphere<T, 3>(TVector<T, 3>(), 2)));
+		OriginalSpheres.Add(TUniquePtr<TSphere<T, 3>>(new TSphere<T, 3>(TVector<T, 3>(), 3)));
+
+		TArray<uint8> Data;
+		{
+			FMemoryWriter Ar(Data);
+			FChaosArchive Writer(Ar);
+
+			Writer << OriginalSpheres;
+		}
+
+		{
+			FMemoryReader Ar(Data);
+			FChaosArchive Reader(Ar);
+			TArray<TSerializablePtr<TSphere<T, 3>>> SerializedSpheres;
+
+			Reader << SerializedSpheres;
+
+			EXPECT_TRUE(SerializedSpheres.Num() == OriginalSpheres.Num());
+
+			for (int32 Idx = 0; Idx < SerializedSpheres.Num(); ++Idx)
+			{
+				EXPECT_TRUE(SerializedSpheres[Idx]->GetRadius() == OriginalSpheres[Idx]->GetRadius());
+			}
+		}
+	}
+
+	template<class T>
+	void SharedObjectsSerialization()
+	{
+		TArray<TSharedPtr<TSphere<T, 3>>> OriginalSpheres;
+		TSharedPtr<TSphere<T, 3>> Sphere(new TSphere<T, 3>(TVector<T, 3>(0), 1));
+		OriginalSpheres.Add(Sphere);
+		OriginalSpheres.Add(Sphere);
+		TSerializablePtr<TSphere<T, 3>> SerializableSphere = MakeSerializable(Sphere);
+
+		TArray<uint8> Data;
+		{
+			FMemoryWriter Ar(Data);
+			FChaosArchive Writer(Ar);
+
+			Writer << OriginalSpheres;
+			Writer << SerializableSphere;
+		}
+		
+		{
+			TArray<TSharedPtr<TSphere<T, 3>>> SerializedSpheres;
+			TSerializablePtr<TSphere<T, 3>> SerializedSphere;
+			{
+				FMemoryReader Ar(Data);
+				FChaosArchive Reader(Ar);
+
+				Reader << SerializedSpheres;
+				Reader << SerializedSphere;
+
+				EXPECT_TRUE(SerializedSpheres.Num() == OriginalSpheres.Num());
+				EXPECT_EQ(SerializedSphere.Get(), SerializedSpheres[0].Get());
+
+				for (int32 Idx = 0; Idx < SerializedSpheres.Num(); ++Idx)
+				{
+					EXPECT_TRUE(SerializedSpheres[Idx]->GetRadius() == OriginalSpheres[Idx]->GetRadius());
+				}
+
+				EXPECT_EQ(SerializedSpheres[0].Get(), SerializedSpheres[1].Get());
+				EXPECT_EQ(SerializedSpheres[0].GetSharedReferenceCount(), 3);
+			}
+			EXPECT_EQ(SerializedSpheres[0].GetSharedReferenceCount(), 2);	//archive is gone so ref count went down
+		}
+	}
+
+	template<class T>
+	void GraphSerialization()
+	{
+		TArray<TUniquePtr<TSphere<T, 3>>> OriginalSpheres;
+		OriginalSpheres.Emplace(new TSphere<T, 3>{ TVector<T,3>(1,2,3), 1 });
+		OriginalSpheres.Emplace(new TSphere<T, 3>{ TVector<T,3>(1,2,3), 2 });
+
+		TArray<TUniquePtr<TImplicitObjectTransformed<T, 3>>> OriginalChildren;
+		OriginalChildren.Emplace(new TImplicitObjectTransformed<T, 3>(MakeSerializable(OriginalSpheres[0]), TRigidTransform<T,3>::Identity));
+		OriginalChildren.Emplace(new TImplicitObjectTransformed<T, 3>(MakeSerializable(OriginalSpheres[1]), TRigidTransform<T,3>::Identity));
+		OriginalChildren.Emplace(new TImplicitObjectTransformed<T, 3>(MakeSerializable(OriginalSpheres[0]), TRigidTransform<T,3>::Identity));
+
+		TUniquePtr<TImplicitObjectTransformed<T, 3>> Root(new TImplicitObjectTransformed<T, 3>(MakeSerializable(OriginalChildren[1]), TRigidTransform<T, 3>::Identity));
+
+		TArray<uint8> Data;
+		{
+			FMemoryWriter Ar(Data);
+			FChaosArchive Writer(Ar);
+
+			Writer << OriginalSpheres;
+			Writer << OriginalChildren;
+			Writer << Root;
+		}
+
+		{
+			FMemoryReader Ar(Data);
+			FChaosArchive Reader(Ar);
+
+			TArray <TUniquePtr<TSphere<T, 3>>> SerializedSpheres;
+			TArray<TSerializablePtr<TImplicitObjectTransformed<T, 3>>> SerializedChildren;
+			TUniquePtr<TImplicitObjectTransformed<T, 3>> SerializedRoot;
+
+			Reader << SerializedSpheres;
+			Reader << SerializedChildren;
+			Reader << SerializedRoot;
+
+			EXPECT_EQ(SerializedSpheres.Num(), OriginalSpheres.Num());
+			EXPECT_EQ(SerializedChildren.Num(), OriginalChildren.Num());
+
+			EXPECT_EQ(SerializedRoot->GetTransformedObject(), SerializedChildren[1].Get());
+			EXPECT_EQ(SerializedChildren[0]->GetTransformedObject(), SerializedSpheres[0].Get());
+			EXPECT_EQ(SerializedChildren[1]->GetTransformedObject(), SerializedSpheres[1].Get());
+			EXPECT_EQ(SerializedChildren[2]->GetTransformedObject(), SerializedSpheres[0].Get());
+		}
+	}
+
+	template<class T>
+	void ObjectUnionSerialization()
+	{
+		TArray<TUniquePtr<FImplicitObject>> OriginalSpheres;
+		OriginalSpheres.Emplace(new TSphere<T, 3>(TVector<T, 3>(1, 2, 3), 1));
+		OriginalSpheres.Emplace(new TSphere<T, 3>(TVector<T, 3>(1, 2, 3), 2));
+
+		TArray<TUniquePtr<FImplicitObject>> OriginalChildren;
+		OriginalChildren.Emplace(new TImplicitObjectTransformed<T, 3>(MakeSerializable(OriginalSpheres[0]), TRigidTransform<T, 3>::Identity));
+		OriginalChildren.Emplace(new TImplicitObjectTransformed<T, 3>(MakeSerializable(OriginalSpheres[1]), TRigidTransform<T, 3>::Identity));
+		OriginalChildren.Emplace(new TImplicitObjectTransformed<T, 3>(MakeSerializable(OriginalSpheres[0]), TRigidTransform<T, 3>::Identity));
+
+		TUniquePtr<FImplicitObjectUnion> Root(new FImplicitObjectUnion(MoveTemp(OriginalChildren)));
+
+		TArray<uint8> Data;
+		{
+			FMemoryWriter Ar(Data);
+			FChaosArchive Writer(Ar);
+
+			Writer << Root;
+			Writer << OriginalSpheres;
+			Writer << OriginalChildren;
+		}
+
+		{
+			FMemoryReader Ar(Data);
+			FChaosArchive Reader(Ar);
+
+			TArray <TUniquePtr<TSphere<T, 3>>> SerializedSpheres;
+			TArray<TSerializablePtr<TImplicitObjectTransformed<T, 3>>> SerializedChildren;
+			TUniquePtr<FImplicitObjectUnion> SerializedRoot;
+
+			Reader << SerializedRoot;
+			Reader << SerializedSpheres;
+			Reader << SerializedChildren;
+
+			EXPECT_EQ(SerializedSpheres.Num(), OriginalSpheres.Num());
+			EXPECT_EQ(SerializedChildren.Num(), OriginalChildren.Num());
+			EXPECT_EQ(SerializedChildren.Num(), 0);	//We did a move and then serialized, should be empty
+
+			const TArray<TUniquePtr<FImplicitObject>>& UnionObjs = SerializedRoot->GetObjects();
+			TImplicitObjectTransformed<T, 3>* FirstChild = static_cast<TImplicitObjectTransformed<T, 3>*>(UnionObjs[0].Get());
+			TImplicitObjectTransformed<T, 3>* SecondChild = static_cast<TImplicitObjectTransformed<T, 3>*>(UnionObjs[1].Get());
+			TImplicitObjectTransformed<T, 3>* ThirdChild = static_cast<TImplicitObjectTransformed<T, 3>*>(UnionObjs[2].Get());
+
+			EXPECT_EQ(FirstChild->GetTransformedObject(), SerializedSpheres[0].Get());
+			EXPECT_EQ(SecondChild->GetTransformedObject(), SerializedSpheres[1].Get());
+			EXPECT_EQ(ThirdChild->GetTransformedObject(), SerializedSpheres[0].Get());
+			EXPECT_TRUE(FirstChild != ThirdChild);	//First and third point to same sphere, but still unique children
+		}
+	}
+
+	template<class T>
+	void ParticleSerialization()
+	{
+		TArray<TUniquePtr<TSphere<T, 3>>> OriginalSpheres;
+		OriginalSpheres.Emplace(new TSphere<T, 3>(TVector<T, 3>(1, 2, 3), 1));
+		OriginalSpheres.Emplace(new TSphere<T, 3>(TVector<T, 3>(1, 2, 3), 2));
+
+		{
+			TGeometryParticles<T, 3> OriginalParticles;
+			OriginalParticles.AddParticles(2);
+			OriginalParticles.SetGeometry(0, MakeSerializable(OriginalSpheres[0]));
+			OriginalParticles.SetGeometry(1, MakeSerializable(OriginalSpheres[1]));
+
+			TArray<uint8> Data;
+			{
+				FMemoryWriter Ar(Data);
+				FChaosArchive Writer(Ar);
+
+				Writer << OriginalParticles;
+				Writer << OriginalSpheres;
+			}
+
+			{
+				FMemoryReader Ar(Data);
+				FChaosArchive Reader(Ar);
+
+				TArray <TUniquePtr<TSphere<T, 3>>> SerializedSpheres;
+				TGeometryParticles<T, 3> SerializedParticles;
+
+				Reader << SerializedParticles;
+				Reader << SerializedSpheres;
+
+				EXPECT_EQ(SerializedSpheres.Num(), OriginalSpheres.Num());
+				EXPECT_EQ(SerializedParticles.Size(), OriginalParticles.Size());
+
+				EXPECT_EQ(SerializedParticles.Geometry(0).Get(), SerializedSpheres[0].Get());
+				EXPECT_EQ(SerializedParticles.Geometry(1).Get(), SerializedSpheres[1].Get());
+			}
+		}
+
+		//ptr
+		{
+			auto OriginalParticles = MakeUnique<TGeometryParticles<T, 3>>();
+			OriginalParticles->AddParticles(2);
+			OriginalParticles->SetGeometry(0, MakeSerializable(OriginalSpheres[0]));
+			OriginalParticles->SetGeometry(1, MakeSerializable(OriginalSpheres[1]));
+
+			TArray<uint8> Data;
+			{
+				FMemoryWriter Ar(Data);
+				FChaosArchive Writer(Ar);
+
+				Writer << OriginalParticles;
+				Writer << OriginalSpheres;
+			}
+
+			{
+				FMemoryReader Ar(Data);
+				FChaosArchive Reader(Ar);
+
+				TArray <TUniquePtr<TSphere<T, 3>>> SerializedSpheres;
+				TUniquePtr<TGeometryParticles<T, 3>> SerializedParticles;
+
+				Reader << SerializedParticles;
+				Reader << SerializedSpheres;
+
+				EXPECT_EQ(SerializedSpheres.Num(), OriginalSpheres.Num());
+				EXPECT_EQ(SerializedParticles->Size(), OriginalParticles->Size());
+
+				EXPECT_EQ(SerializedParticles->Geometry(0).Get(), SerializedSpheres[0].Get());
+				EXPECT_EQ(SerializedParticles->Geometry(1).Get(), SerializedSpheres[1].Get());
+			}
+		}
+	}
+
+	template<class T>
+	void BVHSerialization()
+	{
+		TArray<uint8> Data;
+		{
+			TArray<TUniquePtr<TSphere<T, 3>>> OriginalSpheres;
+			OriginalSpheres.Emplace(new TSphere<T, 3>(TVector<T, 3>(0, 0, 0), 1));
+			OriginalSpheres.Emplace(new TSphere<T, 3>(TVector<T, 3>(0, 0, 0), 2));
+
+			TGeometryParticles<T, 3> OriginalParticles;
+			OriginalParticles.AddParticles(2);
+			OriginalParticles.SetGeometry(0, MakeSerializable(OriginalSpheres[0]));
+			OriginalParticles.SetGeometry(1, MakeSerializable(OriginalSpheres[1]));
+			OriginalParticles.X(0) = TVector<T, 3>(100, 1, 2);
+			OriginalParticles.X(1) = TVector<T, 3>(0, 1, 2);
+			OriginalParticles.R(0) = TRotation<T, 3>::Identity;
+			OriginalParticles.R(1) = TRotation<T, 3>::Identity;
+
+			TBoundingVolumeHierarchy<TGeometryParticles<T, 3>, TArray<int32>, T, 3> OriginalBVH(OriginalParticles);
+
+			FMemoryWriter Ar(Data);
+			FChaosArchive Writer(Ar);
+
+			Writer << OriginalBVH;
+			Writer << OriginalSpheres;
+			Writer << OriginalParticles;
+		}
+
+		{
+			TArray <TUniquePtr<TSphere<T, 3>>> SerializedSpheres;
+			TGeometryParticles<T, 3> SerializedParticles;
+			TBoundingVolumeHierarchy<TGeometryParticles<T, 3>, TArray<int32>, T, 3> SerializedBVH(SerializedParticles);
+			FMemoryReader Ar(Data);
+			FChaosArchive Reader(Ar);
+
+			Reader << SerializedBVH;
+			Reader << SerializedSpheres;
+			Reader << SerializedParticles;
+
+			const TAABB<T, 3> QueryBox{ {-1,0,0}, {1,10,20} };
+			const TArray<int32>& PotentialIntersections = SerializedBVH.FindAllIntersections(QueryBox);
+			TArray<int32> FinalIntersections;
+			for (int32 Potential : PotentialIntersections)
+			{
+				TRigidTransform<T, 3> TM(SerializedParticles.X(Potential), SerializedParticles.R(Potential));
+				const TAABB<T, 3> Bounds = SerializedParticles.Geometry(Potential)->BoundingBox().TransformedAABB(TM);
+				if (Bounds.Intersects(QueryBox))
+				{
+					FinalIntersections.Add(Potential);
+				}
+			}
+
+			EXPECT_EQ(FinalIntersections.Num(), 1);
+			EXPECT_EQ(FinalIntersections[0], 1);
+		}
+	}
+
+	template<class T>
+	void RigidParticlesSerialization()
+	{
+		TArray<TVector<T, 3>> F;
+		F.Emplace(TVector<T, 3>(1, 2, 3));
+		F.Emplace(TVector<T, 3>(3, 2, 1));
+
+		TArray<TVector<T, 3>> X;
+		X.Emplace(TVector<T, 3>(0, 2, 1));
+		X.Emplace(TVector<T, 3>(100, 15, 0));
+
+		TRigidParticles<T, 3> Particles;
+		Particles.AddParticles(2);
+		Particles.F(0) = F[0];
+		Particles.F(1) = F[1];
+		Particles.X(0) = X[0];
+		Particles.X(1) = X[1];
+
+
+		TCHAR const * BinaryFolderName = TEXT("RigidParticles");
+		bool bSaveBinaryToDisk = false; // Flip to true and run to save current binary to disk for future tests.
+		TArray<TRigidParticles<T, 3>> ObjectsToTest;
+		bool bResult = SaveLoadUtility<T,TRigidParticles<T, 3>>(Particles, *GetSerializedBinaryPath(), BinaryFolderName, bSaveBinaryToDisk, ObjectsToTest);
+		EXPECT_TRUE(bResult);
+
+		for (TRigidParticles<T, 3> const &TestParticles : ObjectsToTest)
+		{
+			EXPECT_EQ(TestParticles.Size(), Particles.Size());
+			EXPECT_EQ(TestParticles.F(0), Particles.F(0));
+			EXPECT_EQ(TestParticles.F(1), Particles.F(1));
+			EXPECT_EQ(TestParticles.X(0), Particles.X(0));
+			EXPECT_EQ(TestParticles.X(1), Particles.X(1));
+		}
+	}
+
+	template<class T>
+	void BVHParticlesSerialization()
+	{
+		TArray<uint8> Data;
+		TArray<TUniquePtr<TSphere<T, 3>>> Spheres;
+		Spheres.Emplace(new TSphere<T, 3>(TVector<T, 3>(0, 0, 0), 1));
+		Spheres.Emplace(new TSphere<T, 3>(TVector<T, 3>(0, 0, 0), 1));
+		Spheres.Emplace(new TSphere<T, 3>(TVector<T, 3>(0, 0, 0), 1));
+
+		TGeometryParticles<T, 3> Particles;
+		Particles.AddParticles(3);
+		Particles.SetGeometry(0, MakeSerializable(Spheres[0]));
+		Particles.SetGeometry(1, MakeSerializable(Spheres[1]));
+		Particles.SetGeometry(2, MakeSerializable(Spheres[2]));
+		Particles.X(0) = TVector<T, 3>(15, 1, 2);
+		Particles.X(1) = TVector<T, 3>(0, 2, 2);
+		Particles.X(2) = TVector<T, 3>(0, 2, 2);
+		Particles.R(0) = TRotation<T, 3>::Identity;
+		Particles.R(1) = TRotation<T, 3>::Identity;
+		Particles.R(2) = TRotation<T, 3>::Identity;
+
+		TBVHParticles<T, 3> BVHParticles(MoveTemp(Particles));
+
+		TCHAR const *BinaryFolderName = TEXT("BVHParticles");
+		bool bSaveBinaryToDisk = false; // Flip to true and run to save current binary to disk for future tests.
+		TArray<TBVHParticles<T, 3>> ObjectsToTest;
+		bool bResult = SaveLoadUtility<T, TBVHParticles<T, 3>>(BVHParticles, *GetSerializedBinaryPath(), BinaryFolderName, bSaveBinaryToDisk, ObjectsToTest);
+		EXPECT_TRUE(bResult);
+
+		for (TBVHParticles<T, 3> const &TestBVHP: ObjectsToTest)
+		{
+			const TAABB<T, 3> Box{ {-1,-1,-1}, {1,3,3} };
+			TArray<int32> PotentialIntersections = BVHParticles.FindAllIntersections(Box);
+
+			EXPECT_EQ(TestBVHP.Size(), BVHParticles.Size());
+			EXPECT_EQ(PotentialIntersections.Num(), 2);
+			EXPECT_EQ(PotentialIntersections[0], 1);
+			EXPECT_EQ(PotentialIntersections[1], 2);
+		}
+	}
+
+	void EvolutionPerfHelper(const FString& FilePath)
+	{
+		CHAOS_PERF_TEST(EvolutionPerf, EChaosPerfUnits::Us);
+
+		for (int i = 0; i < 1000; ++i)
+		{
+			TUniquePtr<FArchive> File(IFileManager::Get().CreateFileReader(*FilePath));
+			if (File)
+			{
+				Chaos::FChaosArchive ChaosAr(*File);
+				TPBDRigidsSOAs<float, 3> Particles;
+				TPBDRigidsEvolutionGBF<float, 3> Evolution(Particles);
+
+				Evolution.Serialize(ChaosAr);
+				Evolution.AdvanceOneTimeStep(1 / 60.f);
+				Evolution.EndFrame(1 / 60.0f);
+			}
+		}
+	}
+
+	void EvolutionPerfHarness()
+	{
+		//Load evolutions and step them over and over (with rewind) to measure perf of different components in the system
+		//EvolutionPerfHelper(FPaths::EngineDir() / TEXT("Source/Programs/NotForLicensees/HeadlessPhysicsSQ/Captures/ChaosEvolution_76.bin"));
+	}
+	
+
+	template void SimpleObjectsSerialization<float>();
+	template void SharedObjectsSerialization<float>();
+	template void GraphSerialization<float>();
+	template void ObjectUnionSerialization<float>();
+	template void ParticleSerialization<float>();
+	template void BVHSerialization<float>();
+	template void RigidParticlesSerialization<float>();
+	template void BVHParticlesSerialization<float>();
+}
\ No newline at end of file
diff --git a/Engine/Source/Programs/HeadlessChaos/Private/HeadlessChaosTestSim.cpp b/Engine/Source/Programs/HeadlessChaos/Private/HeadlessChaosTestSim.cpp
new file mode 100644
--- /dev/null
+++ b/Engine/Source/Programs/HeadlessChaos/Private/HeadlessChaosTestSim.cpp
@@ -0,0 +1,164 @@
+// Copyright Epic Games, Inc. All Rights Reserved.
+
+#include "HeadlessChaos.h"
+#include "Chaos/PBDRigidsEvolutionGBF.h"
+#include "Chaos/Box.h"
+#include "Chaos/Sphere.h"
+#include "Chaos/Utilities.h"
+#include "Modules/ModuleManager.h"
+
+namespace ChaosTest {
+
+	using namespace Chaos;
+
+	GTEST_TEST(SimTests, SphereSphereSimTest)
+	{
+		TPBDRigidsSOAs<FReal, 3> Particles;
+		TPBDRigidsEvolutionGBF<FReal, 3> Evolution(Particles);
+		auto Static = Evolution.CreateStaticParticles(1)[0];
+		auto Dynamic = Evolution.CreateDynamicParticles(1)[0];
+
+		TUniquePtr<FImplicitObject> Box(new TSphere<FReal, 3>(FVec3(0, 0, 0), 50));
+		Static->SetGeometry(MakeSerializable(Box));
+		Dynamic->SetGeometry(MakeSerializable(Box));
+
+		Static->X() = FVec3(10, 10, 10);
+		Dynamic->X() = FVec3(10, 10, 300);
+		Dynamic->I() = FMatrix33(100000.0f, 100000.0f, 100000.0f);
+		Dynamic->InvI() = FMatrix33(1.0f / 100000.0f, 1.0f / 100000.0f, 1.0f / 100000.0f);
+
+		const FReal Dt = 1 / 60.f;
+		for (int i = 0; i < 100; ++i)
+		{
+			Evolution.AdvanceOneTimeStep(1 / 60.f);
+			Evolution.EndFrame(Dt);
+		}
+
+		EXPECT_NEAR(Dynamic->X().Z, 110, 1);
+	}
+
+	GTEST_TEST(SimTests, BoxBoxSimTest)
+	{
+		TPBDRigidsSOAs<FReal, 3> Particles;
+		TPBDRigidsEvolutionGBF<FReal, 3> Evolution(Particles);
+		auto Static = Evolution.CreateStaticParticles(1)[0];
+		auto Dynamic = Evolution.CreateDynamicParticles(1)[0];
+
+		TUniquePtr<FImplicitObject> StaticBox(new TBox<FReal, 3>(FVec3(-50, -50, -50), FVec3(50, 50, 50)));
+		TUniquePtr<FImplicitObject> DynamicBox(new TBox<FReal, 3>(FVec3(-50, -50, -50), FVec3(50, 50, 50)));
+		Static->SetGeometry(MakeSerializable(StaticBox));
+		Dynamic->SetGeometry(MakeSerializable(DynamicBox));
+
+		Static->X() = FVec3(10, 10, 10);
+		Dynamic->X() = FVec3(10, 10, 300);
+		Dynamic->I() = FMatrix33(100000.0f, 100000.0f, 100000.0f);
+		Dynamic->InvI() = FMatrix33(1.0f / 100000.0f, 1.0f / 100000.0f, 1.0f / 100000.0f);
+
+		const FReal Dt = 1 / 60.f;
+		for (int i = 0; i < 100; ++i)
+		{
+			Evolution.AdvanceOneTimeStep(Dt);
+			Evolution.EndFrame(Dt);
+		}
+
+		EXPECT_NEAR(Dynamic->X().Z, 110, 5);
+	}
+
+	// This test will fail because the inertia of the dynamic box is very low. The mass and inertia are both 1.0, 
+	// but the box is 100x100x100. When we detect collisions, we get points around the edge of the box. The impulse
+	// required to stop the velocity at that point is tiny because a tiny impulse can impart a large angular velocity
+	// at that position. Therefore we would need a very large number of iterations to resolve it.
+	// 
+	// This will be fixed if/when we have a multi-contact manifold between particle pairs and we simultaneously
+	// resolve contacts in that manifold.
+	//
+	GTEST_TEST(SimTests, DISABLED_VeryLowInertiaSimTest)
+	{
+		TPBDRigidsSOAs<FReal, 3> Particles;
+		TPBDRigidsEvolutionGBF<FReal, 3> Evolution(Particles);
+		auto Static = Evolution.CreateStaticParticles(1)[0];
+		auto Dynamic = Evolution.CreateDynamicParticles(1)[0];
+
+		TUniquePtr<FImplicitObject> StaticBox(new TBox<FReal, 3>(FVec3(-50, -50, -50), FVec3(50, 50, 50)));
+		TUniquePtr<FImplicitObject> DynamicBox(new TBox<FReal, 3>(FVec3(-50, -50, -50), FVec3(50, 50, 50)));
+		Static->SetGeometry(MakeSerializable(StaticBox));
+		Dynamic->SetGeometry(MakeSerializable(DynamicBox));
+
+		Static->X() = FVec3(10, 10, 10);
+		Dynamic->X() = FVec3(10, 10, 300);
+		Dynamic->I() = FMatrix33(1, 1, 1);
+		Dynamic->InvI() = FMatrix33(1, 1, 1);
+
+		for (int i = 0; i < 100; ++i)
+		{
+			Evolution.AdvanceOneTimeStep(1 / 60.f);
+			Evolution.EndFrame(1 / 60.f);
+		}
+
+		EXPECT_NEAR(Dynamic->X().Z, 110, 10);
+	}
+
+	GTEST_TEST(SimTests, SleepAndWakeSimTest)
+	{
+		TPBDRigidsSOAs<FReal, 3> Particles;
+		TPBDRigidsEvolutionGBF<FReal, 3> Evolution(Particles);
+		auto Static = Evolution.CreateStaticParticles(1)[0];
+		auto Dynamic1 = Evolution.CreateDynamicParticles(1)[0];
+		auto Dynamic2 = Evolution.CreateDynamicParticles(1)[0];
+
+		TArrayCollectionArray<TSerializablePtr<FChaosPhysicsMaterial>> PhysicsMaterials;
+
+		Particles.GetParticleHandles().AddArray(&PhysicsMaterials);
+
+		TUniquePtr<FChaosPhysicsMaterial> PhysicsMaterial = MakeUnique<FChaosPhysicsMaterial>();
+		PhysicsMaterial->SleepingLinearThreshold = 20;
+		PhysicsMaterial->SleepingAngularThreshold = 20;
+
+		TUniquePtr<FImplicitObject> StaticBox(new TBox<FReal, 3>(FVec3(-500, -500, -50), FVec3(500, 500, 50)));
+		TUniquePtr<FImplicitObject> DynamicBox(new TBox<FReal, 3>(FVec3(-50, -50, -50), FVec3(50, 50, 50)));
+		Static->SetGeometry(MakeSerializable(StaticBox));
+		Dynamic1->SetGeometry(MakeSerializable(DynamicBox));
+		Dynamic2->SetGeometry(MakeSerializable(DynamicBox));
+
+		Evolution.SetPhysicsMaterial(Dynamic1, MakeSerializable(PhysicsMaterial));
+		Evolution.SetPhysicsMaterial(Dynamic2, MakeSerializable(PhysicsMaterial));
+
+		Static->X() = FVec3(10, 10, 10);
+		Dynamic1->X() = FVec3(10, 10, 120);
+		Dynamic2->X() = FVec3(10, 10, 400);
+		Dynamic1->I() = FMatrix33(100000.0f, 100000.0f, 100000.0f);
+		Dynamic1->InvI() = FMatrix33(1.0f / 100000.0f, 1.0f / 100000.0f, 1.0f / 100000.0f);
+		Dynamic2->I() = FMatrix33(100000.0f, 100000.0f, 100000.0f);
+		Dynamic2->InvI() = FMatrix33(1.0f / 100000.0f, 1.0f / 100000.0f, 1.0f / 100000.0f);
+
+		bool Dynamic1WentToSleep = false;
+		bool Dynamic1HasWokeAgain = false;
+		for (int i = 0; i < 1000; ++i)
+		{
+			Evolution.AdvanceOneTimeStep(1 / 60.f);
+			Evolution.EndFrame(1 / 60.f);
+
+			// at some point Dynamic1 should come to rest and go to sleep on static particle
+			if (Dynamic1WentToSleep == false && Dynamic1->ObjectState() == EObjectStateType::Sleeping)
+			{
+				Dynamic1WentToSleep = true;
+
+				EXPECT_LT(Dynamic1->X().Z, 120);
+				EXPECT_GT(Dynamic1->X().Z, 100);
+			}
+
+			// later the Dynamic2 collides with Dynamic1 waking it up again
+			if (Dynamic1WentToSleep)
+			{
+				if (Dynamic1->ObjectState() == EObjectStateType::Dynamic)
+				{
+					Dynamic1HasWokeAgain = true;
+				}
+			}
+		}
+
+		EXPECT_TRUE(Dynamic1WentToSleep);
+		EXPECT_TRUE(Dynamic1HasWokeAgain);
+	}
+}
+
diff --git a/Engine/Source/Programs/HeadlessChaos/Private/HeadlessChaosTestSolverCommandList.cpp b/Engine/Source/Programs/HeadlessChaos/Private/HeadlessChaosTestSolverCommandList.cpp
new file mode 100644
--- /dev/null
+++ b/Engine/Source/Programs/HeadlessChaos/Private/HeadlessChaosTestSolverCommandList.cpp
@@ -0,0 +1,77 @@
+// Copyright Epic Games, Inc. All Rights Reserved.
+
+#include "HeadlessChaosTestSolverCommandList.h"
+#include "HeadlessChaos.h"
+
+#include "Chaos/ParticleHandle.h"
+#include "Chaos/ErrorReporter.h"
+#include "Chaos/Sphere.h"
+#include "PhysicsProxy/SingleParticlePhysicsProxy.h"
+#include "Chaos/Utilities.h"
+#include "PBDRigidsSolver.h"
+#include "ChaosSolversModule.h"
+
+#include "Modules/ModuleManager.h"
+#include "Framework/PhysicsTickTask.h"
+#include "Framework/CommandBuffer.h"
+
+
+namespace ChaosTest {
+
+    using namespace Chaos;
+
+	template<class T>
+	void CommandListTest()
+	{
+		auto Sphere = TSharedPtr<FImplicitObject, ESPMode::ThreadSafe>(new TSphere<float, 3>(TVector<float, 3>(0), 10));
+		typedef Chaos::TPBDRigidParticle<T, 3> FRigidParticle;
+		typedef TUniquePtr<FRigidParticle> FRigidParticlePtr;
+
+		FChaosSolversModule* Module = FChaosSolversModule::GetModule();
+		Module->ChangeThreadingMode(EChaosThreadingMode::DedicatedThread);
+
+		// Make a solver
+		Chaos::FPhysicsSolver* Solver = FChaosSolversModule::GetModule()->CreateSolver(true);
+		Solver->SetEnabled(true);
+
+		// Make a particle
+		FRigidParticlePtr Particle = TPBDRigidParticle<T,3>::CreateParticle();
+		Particle->SetGeometry(Sphere);
+		Particle->SetX(TVector<float, 3>(0, 0, 0));
+		Particle->SetV(TVector<float, 3>(0, 0, 10));
+
+		Solver->RegisterObject(Particle.Get());
+
+		{
+			TAtomic<int32> TestSequenceAtmoic(5);
+
+			Chaos::FCommandList NewCommandList;
+			NewCommandList.Enqueue([&TestSequenceAtmoic]() {
+				TestSequenceAtmoic = TestSequenceAtmoic.Load() * 2;
+			});
+
+			NewCommandList.Enqueue([&TestSequenceAtmoic]() {
+				TestSequenceAtmoic = TestSequenceAtmoic.Load() - 2;
+			});
+
+			NewCommandList.Enqueue([&TestSequenceAtmoic]() {
+				TestSequenceAtmoic = TestSequenceAtmoic.Load() * 2;
+			});
+
+			EXPECT_EQ(TestSequenceAtmoic.Load(), 5);
+
+			NewCommandList.Flush();
+
+			FPlatformProcess::Sleep(2.0f);
+
+			EXPECT_EQ(TestSequenceAtmoic.Load(), 16);
+		}
+
+		// Throw out the proxy
+		Solver->UnregisterObject(Particle.Get());
+
+		FChaosSolversModule::GetModule()->DestroySolver(Solver);
+	}
+
+	template void CommandListTest<float>();
+}
diff --git a/Engine/Source/Programs/HeadlessChaos/Private/HeadlessChaosTestSolverProxies.cpp b/Engine/Source/Programs/HeadlessChaos/Private/HeadlessChaosTestSolverProxies.cpp
new file mode 100644
--- /dev/null
+++ b/Engine/Source/Programs/HeadlessChaos/Private/HeadlessChaosTestSolverProxies.cpp
@@ -0,0 +1,133 @@
+// Copyright Epic Games, Inc. All Rights Reserved.
+
+#include "HeadlessChaosTestSolverProxies.h"
+#include "HeadlessChaos.h"
+
+#include "Chaos/ParticleHandle.h"
+#include "Chaos/ErrorReporter.h"
+#include "PhysicsProxy/SingleParticlePhysicsProxy.h"
+#include "PhysicsProxy/GeometryCollectionPhysicsProxy.h"
+#include "Chaos/Utilities.h"
+#include "PBDRigidsSolver.h"
+#include "ChaosSolversModule.h"
+
+#include "Modules/ModuleManager.h"
+#include "Framework/PhysicsTickTask.h"
+
+
+namespace ChaosTest {
+
+    using namespace Chaos;
+
+	template<class T>
+	void SingleParticleProxySingleThreadTest()
+	{
+		auto Sphere = TSharedPtr<FImplicitObject, ESPMode::ThreadSafe>(new TSphere<float, 3>(TVector<float, 3>(0), 10));
+
+		FChaosSolversModule* Module = FChaosSolversModule::GetModule();
+		Module->ChangeThreadingMode(EChaosThreadingMode::SingleThread);
+
+		// Make a solver
+		Chaos::FPhysicsSolver* Solver = Module->CreateSolver(true);
+		Solver->SetEnabled(true);
+
+		// Make a particle
+
+		TUniquePtr<Chaos::TPBDRigidParticle<float, 3>> Particle = Chaos::TPBDRigidParticle<float, 3>::CreateParticle();
+		Particle->SetGeometry(Sphere);
+		Particle->SetX(TVector<float, 3>(0, 0, 0));
+
+		Solver->RegisterObject(Particle.Get());
+
+		Particle->SetV(TVector<float, 3>(0, 0, 10));
+		Solver->AddDirtyProxy(Particle->Proxy);
+
+		Solver->PushPhysicsState(Module->GetDispatcher());
+
+		FPhysicsSolverAdvanceTask AdvanceTask(Solver, 100.0f);
+		AdvanceTask.DoTask(ENamedThreads::GameThread, FGraphEventRef());
+
+		Solver->BufferPhysicsResults();
+		Solver->FlipBuffers();
+		Solver->UpdateGameThreadStructures();
+
+		// Make sure game thread data has changed
+		TVector<float, 3> V = Particle->V();
+		EXPECT_EQ(V.X, 0.f);
+		EXPECT_GT(V.Z, 0.f);
+
+		TVector<float, 3> X = Particle->X();
+		EXPECT_EQ(X.X, 0.f);
+		EXPECT_GT(X.Z, 0.f);
+
+		// Throw out the proxy
+		Solver->UnregisterObject(Particle.Get());
+
+		Module->DestroySolver(Solver);		
+	}
+
+	template void SingleParticleProxySingleThreadTest<float>();
+
+	
+	template<class T>
+	void SingleParticleProxyTaskGraphTest()
+	{
+		//
+		// DISABLED TEST
+		//
+		// There is currently no way to execute a TaskGraph or DedicatedThread simulation in a unit test.
+		// This test should be enabled when TaskGraph simulation is supported for unit tests.
+		//
+
+		auto Sphere = TSharedPtr<FImplicitObject, ESPMode::ThreadSafe>(new TSphere<float, 3>(TVector<float, 3>(0), 10));
+
+		FChaosSolversModule* Module = FChaosSolversModule::GetModule();
+		Module->ChangeThreadingMode(EChaosThreadingMode::DedicatedThread);
+
+		// Make a solver
+		Chaos::FPhysicsSolver* Solver = Module->CreateSolver(true);
+		Solver->SetEnabled(true);
+
+		// Make a particle
+
+		TUniquePtr<Chaos::TPBDRigidParticle<float, 3>> Particle = Chaos::TPBDRigidParticle<float, 3>::CreateParticle();
+		Particle->SetGeometry(Sphere);
+		Particle->SetX(TVector<float, 3>(0, 0, 0));
+		Solver->RegisterObject(Particle.Get());
+
+		Particle->SetV(TVector<float, 3>(0, 0, 10));
+		Solver->AddDirtyProxy(Particle->Proxy);
+
+		int32 Counter = 0;
+		while (Particle->X().Size() == 0.f)
+		{
+			Solver->PushPhysicsState(Module->GetDispatcher()); 
+
+			// This might not be the correct way to advance when using the TaskGraph.
+			FPhysicsSolverAdvanceTask AdvanceTask(Solver, 100.0f);
+			AdvanceTask.DoTask(ENamedThreads::GameThread, FGraphEventRef());
+
+			Solver->BufferPhysicsResults();
+			Solver->FlipBuffers();
+			Solver->UpdateGameThreadStructures();
+
+			EXPECT_LE(Counter++, 5);
+		}
+
+		// Make sure game thread data has changed
+		TVector<float, 3> V = Particle->V();
+		EXPECT_EQ(V.X, 0.f);
+		EXPECT_GT(V.Z, 0.f);
+
+		TVector<float, 3> X = Particle->X();
+		EXPECT_EQ(X.X, 0.f);
+		EXPECT_GT(X.Z, 0.f);
+
+		// Throw out the proxy
+		Solver->UnregisterObject(Particle.Get());
+
+		Module->DestroySolver(Solver);
+	}
+
+	template void SingleParticleProxyTaskGraphTest<float>();
+}
diff --git a/Engine/Source/Programs/HeadlessChaos/Private/HeadlessChaosTestSpatialHashing.cpp b/Engine/Source/Programs/HeadlessChaos/Private/HeadlessChaosTestSpatialHashing.cpp
new file mode 100644
--- /dev/null
+++ b/Engine/Source/Programs/HeadlessChaos/Private/HeadlessChaosTestSpatialHashing.cpp
@@ -0,0 +1,76 @@
+// Copyright Epic Games, Inc. All Rights Reserved.
+
+#include "HeadlessChaosTestSpatialHashing.h"
+
+#include "HeadlessChaos.h"
+#include "HeadlessChaosTestUtility.h"
+#include "Modules/ModuleManager.h"
+#include "Chaos/Matrix.h"
+#include "Chaos/PBDRigidClustering.h"
+
+DEFINE_LOG_CATEGORY_STATIC(AHSP_Test, Verbose, All);
+
+
+namespace ChaosTest {
+	using namespace Chaos;
+
+	template<class T, int Dimension>
+	TArray< TVector<T, 3> > GenerateSamplePoints(int32 NumPoints, int32 InitialDistance) {
+
+		TArray< TVector<T, 3> > ReturnValue;
+		for (int i = 1; i <= NumPoints/2; i++) {
+			int32 iCubed = i * i * i;
+			TVector<T, 3> Vec1(0), Vec2(0);
+			if (0<Dimension && Dimension >= 3)
+			{
+				Vec1 = TVector<T, 3>(float(FMath::Rand()) / RAND_MAX, float(FMath::Rand()) / RAND_MAX, float(FMath::Rand()) / RAND_MAX)*InitialDistance;
+				Vec2 = TVector<T, 3>(float(FMath::Rand()) / RAND_MAX, float(FMath::Rand()) / RAND_MAX, float(FMath::Rand()) / RAND_MAX)*InitialDistance;
+				TVector<T,3> Mid((Vec1 - Vec2) / 2);
+				Vec1 = Mid + (Mid - Vec1).Normalize()*(InitialDistance / float(iCubed));
+				Vec2 = Mid - (Mid - Vec1).Normalize()*(InitialDistance / float(iCubed));
+			}
+			else
+			{
+				int32 OtherDim = (Dimension + 1) % 3;
+				Vec1[Dimension] = InitialDistance / float(iCubed);
+				Vec2[Dimension] = -InitialDistance / float(iCubed);
+				Vec1[OtherDim] = Vec2[OtherDim] = i;
+			}
+			ReturnValue.Add(Vec1);
+			ReturnValue.Add(Vec2);
+		}
+		return ReturnValue;
+	}
+
+	template<class T>
+	void SpatialHashing()
+	{
+		UE_LOG(AHSP_Test, Verbose, TEXT("SpatialHashing"));
+		for (int Dimension = 0; Dimension <= 3; Dimension++)
+		{
+			UE_LOG(AHSP_Test, Verbose, TEXT("::Dimension[%d]"), Dimension);
+			TArray< TVector<T, 3> > Samples = Chaos::CleanCollisionParticles(GenerateSamplePoints<T, 3>(100, 1000), 1.0);
+			for (int32 Index1 = 0; Index1 < Samples.Num(); Index1++)
+			{
+				for (int32 Index2 = 0; Index2 < Samples.Num(); Index2++)
+				{
+					if (Index1 != Index2 && Index1<Index2)
+					{
+						TVector<T, 3> Sample1(Samples[Index1]), Sample2(Samples[Index2]);
+						float Delta = (Sample2 - Sample1).Size();
+						if (Delta < 1.0)
+						{
+							UE_LOG(AHSP_Test, Verbose, TEXT("... [%d](%3.5f,%3.5f,%3.5f)"), Index1, Sample1.X, Sample1.Y, Sample1.Z);
+							UE_LOG(AHSP_Test, Verbose, TEXT("... [%d](%3.5f,%3.5f,%3.5f)"), Index2, Sample2.X, Sample2.Y, Sample2.Z);
+							UE_LOG(AHSP_Test, Verbose, TEXT("... ... %3.5f"), Delta);
+							EXPECT_TRUE(false);
+						}
+					}
+				}
+			}
+		}
+	}
+	template void SpatialHashing<float>();
+
+
+}
\ No newline at end of file
