Rename/move HeadlessChaos.

diff --git a/Engine/Source/Programs/HeadlessChaos/Private/HeadlessChaosTestUtility.cpp b/Engine/Source/Programs/HeadlessChaos/Private/HeadlessChaosTestUtility.cpp
new file mode 100644
--- /dev/null
+++ b/Engine/Source/Programs/HeadlessChaos/Private/HeadlessChaosTestUtility.cpp
@@ -0,0 +1,758 @@
+// Copyright Epic Games, Inc. All Rights Reserved.
+
+#include "HeadlessChaosTestUtility.h"
+
+#include "HeadlessChaos.h"
+#include "Chaos/Box.h"
+#include "Chaos/Convex.h"
+#include "Chaos/ErrorReporter.h"
+#include "Chaos/Levelset.h"
+#include "Chaos/ParticleHandle.h"
+#include "Chaos/PBDRigidParticles.h"
+#include "Chaos/PBDRigidsEvolutionGBF.h"
+#include "Chaos/Plane.h"
+#include "Chaos/Sphere.h"
+#include "Chaos/Cylinder.h"
+#include "Chaos/TaperedCylinder.h"
+#include "Chaos/Utilities.h"
+
+namespace ChaosTest {
+
+	using namespace Chaos;
+
+	template<class T>
+	int32 AppendAnalyticSphere(TPBDRigidParticles<T, 3> & InParticles, T Scale)
+	{
+		InParticles.AddParticles(1);
+		int32 RigidBodyIndex = InParticles.Size() - 1;
+
+		InParticles.X(RigidBodyIndex) = TVector<T, 3>(0.f, 0.f, 0.f);
+		InParticles.V(RigidBodyIndex) = TVector<T, 3>(0.f, 0.f, 0.f);
+		InParticles.R(RigidBodyIndex) = TRotation<T, 3>::MakeFromEuler(TVector<T, 3>(0.f, 0.f, 0.f)).GetNormalized();
+		InParticles.W(RigidBodyIndex) = TVector<T, 3>(0.f, 0.f, 0.f);
+		InParticles.P(RigidBodyIndex) = InParticles.X(RigidBodyIndex);
+		InParticles.Q(RigidBodyIndex) = InParticles.R(RigidBodyIndex);
+
+		InParticles.M(RigidBodyIndex) = 1.f;
+		InParticles.InvM(RigidBodyIndex) = 1.f;
+		InParticles.I(RigidBodyIndex) = PMatrix<T, 3, 3>(1.f, 0.f, 0.f, 0.f, 1.f, 0.f, 0.f, 0.f, 1.f);
+		InParticles.InvI(RigidBodyIndex) = PMatrix<T, 3, 3>(1.f, 0.f, 0.f, 0.f, 1.f, 0.f, 0.f, 0.f, 1.f);
+		InParticles.SetDynamicGeometry(RigidBodyIndex, MakeUnique<TSphere<T, 3>>(TVector<T, 3>(0), Scale));
+		InParticles.SetObjectState(RigidBodyIndex, EObjectStateType::Dynamic);
+
+		return RigidBodyIndex;
+	}
+	template int32 AppendAnalyticSphere(TPBDRigidParticles<float, 3> &, float Scale);
+
+	template<class T>
+	TPBDRigidParticleHandle<T, 3>* AppendAnalyticSphere2(TPBDRigidsEvolutionGBF<T, 3>& Evolution, T Scale)
+	{
+		TArray<TPBDRigidParticleHandle<T, 3>*> Particles = Evolution.CreateDynamicParticles(1);
+		TPBDRigidParticleHandle<T, 3>* Particle = Particles[0];
+
+		Particle->X() = TVector<T, 3>(0.f, 0.f, 0.f);
+		Particle->V() = TVector<T, 3>(0.f, 0.f, 0.f);
+		Particle->R() = TRotation<T, 3>::MakeFromEuler(TVector<T, 3>(0.f, 0.f, 0.f)).GetNormalized();
+		Particle->W() = TVector<T, 3>(0.f, 0.f, 0.f);
+		Particle->P() = Particle->X();
+		Particle->Q() = Particle->R();
+
+		Particle->M() = 1.f;
+		Particle->InvM() = 1.f;
+		Particle->I() = PMatrix<T, 3, 3>(1.f, 0.f, 0.f, 0.f, 1.f, 0.f, 0.f, 0.f, 1.f);
+		Particle->InvI() = PMatrix<T, 3, 3>(1.f, 0.f, 0.f, 0.f, 1.f, 0.f, 0.f, 0.f, 1.f);
+		Particle->SetDynamicGeometry(MakeUnique<TSphere<T, 3>>(TVector<T, 3>(0), Scale));
+
+		return Particle;
+	}
+	template TPBDRigidParticleHandle<float, 3>* AppendAnalyticSphere2(TPBDRigidsEvolutionGBF<float, 3>& Evolution, float Scale);
+
+
+	template<class T>
+	int32 AppendAnalyticBox(TPBDRigidParticles<T, 3> & InParticles, TVector<T, 3> Scale)
+	{
+		InParticles.AddParticles(1);
+		int32 RigidBodyIndex = InParticles.Size() - 1;
+
+		InParticles.X(RigidBodyIndex) = TVector<T, 3>(0.f, 0.f, 0.f);
+		InParticles.V(RigidBodyIndex) = TVector<T, 3>(0.f, 0.f, 0.f);
+		InParticles.R(RigidBodyIndex) = TRotation<T, 3>::MakeFromEuler(TVector<T, 3>(0.f, 0.f, 0.f)).GetNormalized();
+		InParticles.W(RigidBodyIndex) = TVector<T, 3>(0.f, 0.f, 0.f);
+		InParticles.P(RigidBodyIndex) = InParticles.X(RigidBodyIndex);
+		InParticles.Q(RigidBodyIndex) = InParticles.R(RigidBodyIndex);
+
+		InParticles.M(RigidBodyIndex) = 1.f;
+		InParticles.InvM(RigidBodyIndex) = 1.f;
+		InParticles.I(RigidBodyIndex) = PMatrix<T, 3, 3>(1.f, 0.f, 0.f, 0.f, 1.f, 0.f, 0.f, 0.f, 1.f);
+		InParticles.InvI(RigidBodyIndex) = PMatrix<T, 3, 3>(1.f, 0.f, 0.f, 0.f, 1.f, 0.f, 0.f, 0.f, 1.f);
+		InParticles.SetDynamicGeometry(RigidBodyIndex, MakeUnique<TBox<T, 3>>(-Scale / 2.f, Scale / 2.f));
+		InParticles.SetObjectState(RigidBodyIndex, EObjectStateType::Dynamic);
+
+		return RigidBodyIndex;
+	}
+	template int32 AppendAnalyticBox(TPBDRigidParticles<float, 3> &, TVector<float, 3>);
+
+
+	template<class T>
+	void InitAnalyticBox2(TKinematicGeometryParticleHandle<T, 3>* Particle, TVector<T, 3> Scale)
+	{
+		Particle->X() = TVector<T, 3>(0.f, 0.f, 0.f);
+		Particle->V() = TVector<T, 3>(0.f, 0.f, 0.f);
+		Particle->R() = TRotation<T, 3>::MakeFromEuler(TVector<T, 3>(0.f, 0.f, 0.f)).GetNormalized();
+		Particle->W() = TVector<T, 3>(0.f, 0.f, 0.f);
+		Particle->SetDynamicGeometry(MakeUnique<TBox<T, 3>>(-Scale / 2.f, Scale / 2.f));
+
+		TPBDRigidParticleHandle<T, 3>* DynamicParticle = Particle->CastToRigidParticle();
+		if(DynamicParticle && DynamicParticle->ObjectState() == EObjectStateType::Dynamic)
+		{
+			DynamicParticle->P() = Particle->X();
+			DynamicParticle->Q() = Particle->R();
+
+			DynamicParticle->M() = 1.f;
+			DynamicParticle->InvM() = 1.f;
+			DynamicParticle->I() = PMatrix<T, 3, 3>(1.f, 0.f, 0.f, 0.f, 1.f, 0.f, 0.f, 0.f, 1.f);
+			DynamicParticle->InvI() = PMatrix<T, 3, 3>(1.f, 0.f, 0.f, 0.f, 1.f, 0.f, 0.f, 0.f, 1.f);
+		}
+	}
+
+	template<class T>
+	TKinematicGeometryParticleHandle<T, 3>* AppendKinematicAnalyticBox2(TPBDRigidsEvolutionGBF<T, 3>& Evolution, TVector<T, 3> Scale)
+	{
+		TArray<TKinematicGeometryParticleHandle<T, 3>*> Particles = Evolution.CreateKinematicParticles(1);
+		InitAnalyticBox2(Particles[0], Scale);
+		return Particles[0];
+	}
+	template TKinematicGeometryParticleHandle<float, 3>* AppendKinematicAnalyticBox2(TPBDRigidsEvolutionGBF<float, 3>& Evolution, TVector<float, 3> Scale);
+
+	template<class T>
+	TPBDRigidParticleHandle<T, 3>* AppendDynamicAnalyticBox2(TPBDRigidsEvolutionGBF<T, 3>& Evolution, TVector<T, 3> Scale)
+	{
+		TArray<TPBDRigidParticleHandle<T, 3>*> Particles = Evolution.CreateDynamicParticles(1);
+		InitAnalyticBox2(Particles[0], Scale);
+		return Particles[0];
+	}
+	template TPBDRigidParticleHandle<float, 3>* AppendDynamicAnalyticBox2(TPBDRigidsEvolutionGBF<float, 3>& Evolution, TVector<float, 3> Scale);
+
+
+	template<class T>
+	int32 AppendParticleBox(TPBDRigidParticles<T, 3> & InParticles, TVector<T, 3> Scale, TArray<TVector<int32, 3>> * elements)
+	{
+		InParticles.AddParticles(1);
+		int32 RigidBodyIndex = InParticles.Size() - 1;
+
+		InParticles.X(RigidBodyIndex) = TVector<T, 3>(0.f, 0.f, 0.f);
+		InParticles.V(RigidBodyIndex) = TVector<T, 3>(0.f, 0.f, 0.f);
+		InParticles.R(RigidBodyIndex) = TRotation<T, 3>::MakeFromEuler(TVector<T, 3>(0.f, 0.f, 0.f)).GetNormalized();
+		InParticles.W(RigidBodyIndex) = TVector<T, 3>(0.f, 0.f, 0.f);
+		InParticles.P(RigidBodyIndex) = InParticles.X(RigidBodyIndex);
+		InParticles.Q(RigidBodyIndex) = InParticles.R(RigidBodyIndex);
+
+		check(Scale.X == Scale.Y && Scale.X == Scale.Z);
+		T ScaleSq = Scale.X * Scale.X;
+		InParticles.M(RigidBodyIndex) = 1.f;
+		InParticles.InvM(RigidBodyIndex) = 1.f;
+		InParticles.I(RigidBodyIndex) = PMatrix<T, 3, 3>(ScaleSq / 6.f, 0.f, 0.f, 0.f, ScaleSq / 6.f, 0.f, 0.f, 0.f, ScaleSq / 6.f);
+		InParticles.InvI(RigidBodyIndex) = PMatrix<T, 3, 3>(6.f / ScaleSq, 0.f, 0.f, 0.f, 6.f / ScaleSq, 0.f, 0.f, 0.f, 6.f / ScaleSq);
+		InParticles.SetDynamicGeometry(RigidBodyIndex, MakeUnique<TBox<T, 3>>(-Scale / 2.f, Scale / 2.f));
+		InParticles.SetObjectState(RigidBodyIndex, EObjectStateType::Dynamic);
+
+		int32 CollisionIndex = 0;
+		InParticles.CollisionParticlesInitIfNeeded(RigidBodyIndex);
+		InParticles.CollisionParticles(RigidBodyIndex)->AddParticles(8);
+		InParticles.CollisionParticles(RigidBodyIndex)->X(CollisionIndex++) = TVector<T, 3>(-Scale[0] / 2.f, -Scale[1] / 2.f, -Scale[2] / 2.f);
+		InParticles.CollisionParticles(RigidBodyIndex)->X(CollisionIndex++) = TVector<T, 3>(+Scale[0] / 2.f, -Scale[1] / 2.f, -Scale[2] / 2.f);
+		InParticles.CollisionParticles(RigidBodyIndex)->X(CollisionIndex++) = TVector<T, 3>(-Scale[0] / 2.f, +Scale[1] / 2.f, -Scale[2] / 2.f);
+		InParticles.CollisionParticles(RigidBodyIndex)->X(CollisionIndex++) = TVector<T, 3>(+Scale[0] / 2.f, +Scale[1] / 2.f, -Scale[2] / 2.f);
+		InParticles.CollisionParticles(RigidBodyIndex)->X(CollisionIndex++) = TVector<T, 3>(-Scale[0] / 2.f, -Scale[1] / 2.f, +Scale[2] / 2.f);
+		InParticles.CollisionParticles(RigidBodyIndex)->X(CollisionIndex++) = TVector<T, 3>(+Scale[0] / 2.f, -Scale[1] / 2.f, +Scale[2] / 2.f);
+		InParticles.CollisionParticles(RigidBodyIndex)->X(CollisionIndex++) = TVector<T, 3>(-Scale[0] / 2.f, +Scale[1] / 2.f, +Scale[2] / 2.f);
+		InParticles.CollisionParticles(RigidBodyIndex)->X(CollisionIndex++) = TVector<T, 3>(+Scale[0] / 2.f, +Scale[1] / 2.f, +Scale[2] / 2.f);
+
+		if (elements != nullptr)
+		{
+			/*
+			//cw
+			elements->Add(TVector<int32, 3>(4, 1, 5)); // Front
+			elements->Add(TVector<int32, 3>(1, 4, 0));
+			elements->Add(TVector<int32, 3>(7, 2, 6)); // Back
+			elements->Add(TVector<int32, 3>(2, 7, 3));
+			elements->Add(TVector<int32, 3>(6, 0, 4)); // Right
+			elements->Add(TVector<int32, 3>(0, 6, 2));
+			elements->Add(TVector<int32, 3>(5, 3, 7)); // Left
+			elements->Add(TVector<int32, 3>(3, 5, 1));
+			elements->Add(TVector<int32, 3>(6, 5, 7)); // Top
+			elements->Add(TVector<int32, 3>(5, 6, 4));
+			elements->Add(TVector<int32, 3>(0, 2, 1)); // Front
+			elements->Add(TVector<int32, 3>(2, 0, 3));
+			/**/
+			//ccw
+			elements->Add(TVector<int32, 3>(1,4,5)); // Front
+			elements->Add(TVector<int32, 3>(4,1,0));
+			elements->Add(TVector<int32, 3>(2,7,6)); // Back
+			elements->Add(TVector<int32, 3>(7,2,3));
+			elements->Add(TVector<int32, 3>(0,6,4)); // Right
+			elements->Add(TVector<int32, 3>(6,0,2));
+			elements->Add(TVector<int32, 3>(3,5,7)); // Left
+			elements->Add(TVector<int32, 3>(5,3,1));
+			elements->Add(TVector<int32, 3>(5,6,7)); // Top
+			elements->Add(TVector<int32, 3>(6,5,4));
+			elements->Add(TVector<int32, 3>(2,0,1)); // Front
+			elements->Add(TVector<int32, 3>(0,2,3));
+		}
+
+		return RigidBodyIndex;
+	}
+	template int32 AppendParticleBox(TPBDRigidParticles<float, 3> &, TVector<float, 3>, TArray<Chaos::TVector<int32, 3>> *);
+
+
+	template<class T>
+	void InitDynamicParticleBox2(TPBDRigidParticleHandle<T, 3>* Particle, const TVector<T, 3>& Scale, TArray<TVector<int32, 3>>* OutElements)
+	{
+		Particle->X() = TVector<T, 3>(0.f, 0.f, 0.f);
+		Particle->V() = TVector<T, 3>(0.f, 0.f, 0.f);
+		Particle->R() = TRotation<T, 3>::MakeFromEuler(TVector<T, 3>(0.f, 0.f, 0.f)).GetNormalized();
+		Particle->W() = TVector<T, 3>(0.f, 0.f, 0.f);
+
+		Particle->P() = Particle->X();
+		Particle->Q() = Particle->R();
+
+		check(Scale.X == Scale.Y && Scale.X == Scale.Z);
+		T ScaleSq = Scale.X * Scale.X;
+		Particle->M() = 1.f;
+		Particle->InvM() = 1.f;
+		Particle->I() = PMatrix<T, 3, 3>(ScaleSq / 6.f, 0.f, 0.f, 0.f, ScaleSq / 6.f, 0.f, 0.f, 0.f, ScaleSq / 6.f);
+		Particle->InvI() = PMatrix<T, 3, 3>(6.f / ScaleSq, 0.f, 0.f, 0.f, 6.f / ScaleSq, 0.f, 0.f, 0.f, 6.f / ScaleSq);
+
+		Particle->SetDynamicGeometry(MakeUnique<TBox<T, 3>>(-Scale / 2.f, Scale / 2.f));
+
+		int32 CollisionIndex = 0;
+		Particle->CollisionParticlesInitIfNeeded();
+		Particle->CollisionParticles()->AddParticles(8);
+		Particle->CollisionParticles()->X(CollisionIndex++) = TVector<T, 3>(-Scale[0] / 2.f, -Scale[1] / 2.f, -Scale[2] / 2.f);
+		Particle->CollisionParticles()->X(CollisionIndex++) = TVector<T, 3>(+Scale[0] / 2.f, -Scale[1] / 2.f, -Scale[2] / 2.f);
+		Particle->CollisionParticles()->X(CollisionIndex++) = TVector<T, 3>(-Scale[0] / 2.f, +Scale[1] / 2.f, -Scale[2] / 2.f);
+		Particle->CollisionParticles()->X(CollisionIndex++) = TVector<T, 3>(+Scale[0] / 2.f, +Scale[1] / 2.f, -Scale[2] / 2.f);
+		Particle->CollisionParticles()->X(CollisionIndex++) = TVector<T, 3>(-Scale[0] / 2.f, -Scale[1] / 2.f, +Scale[2] / 2.f);
+		Particle->CollisionParticles()->X(CollisionIndex++) = TVector<T, 3>(+Scale[0] / 2.f, -Scale[1] / 2.f, +Scale[2] / 2.f);
+		Particle->CollisionParticles()->X(CollisionIndex++) = TVector<T, 3>(-Scale[0] / 2.f, +Scale[1] / 2.f, +Scale[2] / 2.f);
+		Particle->CollisionParticles()->X(CollisionIndex++) = TVector<T, 3>(+Scale[0] / 2.f, +Scale[1] / 2.f, +Scale[2] / 2.f);
+
+		if (OutElements != nullptr)
+		{
+			/*
+			//cw
+			OutElements->Add(TVector<int32, 3>(4, 1, 5)); // Front
+			OutElements->Add(TVector<int32, 3>(1, 4, 0));
+			OutElements->Add(TVector<int32, 3>(7, 2, 6)); // Back
+			OutElements->Add(TVector<int32, 3>(2, 7, 3));
+			OutElements->Add(TVector<int32, 3>(6, 0, 4)); // Right
+			OutElements->Add(TVector<int32, 3>(0, 6, 2));
+			OutElements->Add(TVector<int32, 3>(5, 3, 7)); // Left
+			OutElements->Add(TVector<int32, 3>(3, 5, 1));
+			OutElements->Add(TVector<int32, 3>(6, 5, 7)); // Top
+			OutElements->Add(TVector<int32, 3>(5, 6, 4));
+			OutElements->Add(TVector<int32, 3>(0, 2, 1)); // Front
+			OutElements->Add(TVector<int32, 3>(2, 0, 3));
+			/**/
+			//ccw
+			OutElements->Add(TVector<int32, 3>(1, 4, 5)); // Front
+			OutElements->Add(TVector<int32, 3>(4, 1, 0));
+			OutElements->Add(TVector<int32, 3>(2, 7, 6)); // Back
+			OutElements->Add(TVector<int32, 3>(7, 2, 3));
+			OutElements->Add(TVector<int32, 3>(0, 6, 4)); // Right
+			OutElements->Add(TVector<int32, 3>(6, 0, 2));
+			OutElements->Add(TVector<int32, 3>(3, 5, 7)); // Left
+			OutElements->Add(TVector<int32, 3>(5, 3, 1));
+			OutElements->Add(TVector<int32, 3>(5, 6, 7)); // Top
+			OutElements->Add(TVector<int32, 3>(6, 5, 4));
+			OutElements->Add(TVector<int32, 3>(2, 0, 1)); // Front
+			OutElements->Add(TVector<int32, 3>(0, 2, 3));
+		}
+	}
+
+	
+	template<class T>
+	void InitDynamicParticleSphere2(TPBDRigidParticleHandle<T, 3>* Particle, const TVector<T, 3>& Scale, TArray<TVector<int32, 3>>* OutElements) {
+		Particle->X() = TVector<T, 3>(0.f, 0.f, 0.f);
+		Particle->V() = TVector<T, 3>(0.f, 0.f, 0.f);
+		Particle->R() = TRotation<T, 3>::MakeFromEuler(TVector<T, 3>(0.f, 0.f, 0.f)).GetNormalized();
+		Particle->W() = TVector<T, 3>(0.f, 0.f, 0.f);
+
+		Particle->P() = Particle->X();
+		Particle->Q() = Particle->R();
+
+		check(Scale.X == Scale.Y && Scale.X == Scale.Z);
+		T ScaleSq = Scale.X * Scale.X;
+		Particle->M() = 1.f;
+		Particle->InvM() = 1.f;
+		Particle->I() = PMatrix<T, 3, 3>(ScaleSq / 6.f, 0.f, 0.f, 0.f, ScaleSq / 6.f, 0.f, 0.f, 0.f, ScaleSq / 6.f);
+		Particle->InvI() = PMatrix<T, 3, 3>(6.f / ScaleSq, 0.f, 0.f, 0.f, 6.f / ScaleSq, 0.f, 0.f, 0.f, 6.f / ScaleSq);
+
+		Particle->SetDynamicGeometry(MakeUnique<TSphere<T, 3>>(TVector<T, 3>(0), Scale.X / 2.f));
+
+		int32 CollisionIndex = 0;
+		Particle->CollisionParticlesInitIfNeeded();
+		Particle->CollisionParticles()->AddParticles(6);
+
+		Particle->CollisionParticles()->X(CollisionIndex++) = TVector<T, 3>(-Scale[0] / 2.f, 0, 0);
+		Particle->CollisionParticles()->X(CollisionIndex++) = TVector<T, 3>(+Scale[0] / 2.f, 0, 0);
+		Particle->CollisionParticles()->X(CollisionIndex++) = TVector<T, 3>(0, -Scale[1] / 2.f, 0);
+		Particle->CollisionParticles()->X(CollisionIndex++) = TVector<T, 3>(0, +Scale[1] / 2.f, 0);
+		Particle->CollisionParticles()->X(CollisionIndex++) = TVector<T, 3>(0, 0, -Scale[2] / 2.f);
+		Particle->CollisionParticles()->X(CollisionIndex++) = TVector<T, 3>(0, 0, +Scale[2] / 2.f);
+
+		if (OutElements != nullptr)
+		{
+			/*
+			//cw
+			OutElements->Add(TVector<int32, 3>(4, 1, 5)); // Front
+			OutElements->Add(TVector<int32, 3>(1, 4, 0));
+			OutElements->Add(TVector<int32, 3>(7, 2, 6)); // Back
+			OutElements->Add(TVector<int32, 3>(2, 7, 3));
+			OutElements->Add(TVector<int32, 3>(6, 0, 4)); // Right
+			OutElements->Add(TVector<int32, 3>(0, 6, 2));
+			OutElements->Add(TVector<int32, 3>(5, 3, 7)); // Left
+			OutElements->Add(TVector<int32, 3>(3, 5, 1));
+			OutElements->Add(TVector<int32, 3>(6, 5, 7)); // Top
+			OutElements->Add(TVector<int32, 3>(5, 6, 4));
+			OutElements->Add(TVector<int32, 3>(0, 2, 1)); // Front
+			OutElements->Add(TVector<int32, 3>(2, 0, 3));
+			/**/
+			//ccw
+			OutElements->Add(TVector<int32, 3>(1, 4, 5)); // Front
+			OutElements->Add(TVector<int32, 3>(4, 1, 0));
+			OutElements->Add(TVector<int32, 3>(2, 7, 6)); // Back
+			OutElements->Add(TVector<int32, 3>(7, 2, 3));
+			OutElements->Add(TVector<int32, 3>(0, 6, 4)); // Right
+			OutElements->Add(TVector<int32, 3>(6, 0, 2));
+			OutElements->Add(TVector<int32, 3>(3, 5, 7)); // Left
+			OutElements->Add(TVector<int32, 3>(5, 3, 1));
+			OutElements->Add(TVector<int32, 3>(5, 6, 7)); // Top
+			OutElements->Add(TVector<int32, 3>(6, 5, 4));
+			OutElements->Add(TVector<int32, 3>(2, 0, 1)); // Front
+			OutElements->Add(TVector<int32, 3>(0, 2, 3));
+		}
+	}
+
+
+	template<class T>
+	void InitDynamicParticleCylinder2(TPBDRigidParticleHandle<T, 3>* Particle, const TVector<T, 3>& Scale, TArray<TVector<int32, 3>>* OutElements, bool Tapered) {
+		Particle->X() = TVector<T, 3>(0.f, 0.f, 0.f);
+		Particle->V() = TVector<T, 3>(0.f, 0.f, 0.f);
+		Particle->R() = TRotation<T, 3>::MakeFromEuler(TVector<T, 3>(0.f, 0.f, 0.f)).GetNormalized();
+		Particle->W() = TVector<T, 3>(0.f, 0.f, 0.f);
+
+		Particle->P() = Particle->X();
+		Particle->Q() = Particle->R();
+
+		check(Scale.X == Scale.Y && Scale.X == Scale.Z);
+		T ScaleSq = Scale.X * Scale.X;
+		Particle->M() = 1.f;
+		Particle->InvM() = 1.f;
+		Particle->I() = PMatrix<T, 3, 3>(ScaleSq / 6.f, 0.f, 0.f, 0.f, ScaleSq / 6.f, 0.f, 0.f, 0.f, ScaleSq / 6.f);
+		Particle->InvI() = PMatrix<T, 3, 3>(6.f / ScaleSq, 0.f, 0.f, 0.f, 6.f / ScaleSq, 0.f, 0.f, 0.f, 6.f / ScaleSq);
+		
+		if (Tapered)
+		{
+			Particle->SetDynamicGeometry(MakeUnique<TTaperedCylinder<T>>(TVector<T, 3>(0, 0, Scale.X / 2.0f), TVector<T, 3>(0, 0, -Scale.X / 2.0f), Scale.X / 2.f, Scale.X / 2.f));
+		}
+		else 
+		{
+			Particle->SetDynamicGeometry(MakeUnique<TCylinder<T>>(TVector<T, 3>(0, 0, Scale.X / 2.0f), TVector<T, 3>(0, 0, -Scale.X / 2.0f), Scale.X / 2.f));
+		}
+
+		int32 CollisionIndex = 0;
+		Particle->CollisionParticlesInitIfNeeded();
+		Particle->CollisionParticles()->AddParticles(8);
+
+		Particle->CollisionParticles()->X(CollisionIndex++) = TVector<T, 3>(-Scale[0] / 2.f, 0, +Scale[2] / 2.f);
+		Particle->CollisionParticles()->X(CollisionIndex++) = TVector<T, 3>(-Scale[0] / 2.f, 0, -Scale[2] / 2.f);
+		Particle->CollisionParticles()->X(CollisionIndex++) = TVector<T, 3>(+Scale[0] / 2.f, 0, +Scale[2] / 2.f);
+		Particle->CollisionParticles()->X(CollisionIndex++) = TVector<T, 3>(+Scale[0] / 2.f, 0, -Scale[2] / 2.f);
+		Particle->CollisionParticles()->X(CollisionIndex++) = TVector<T, 3>(0, -Scale[1] / 2.f, +Scale[2] / 2.f);
+		Particle->CollisionParticles()->X(CollisionIndex++) = TVector<T, 3>(0, -Scale[1] / 2.f, -Scale[2] / 2.f);
+		Particle->CollisionParticles()->X(CollisionIndex++) = TVector<T, 3>(0, +Scale[1] / 2.f, +Scale[2] / 2.f);
+		Particle->CollisionParticles()->X(CollisionIndex++) = TVector<T, 3>(0, +Scale[1] / 2.f, -Scale[2] / 2.f);
+
+		if (OutElements != nullptr)
+		{
+			/*
+			//cw
+			OutElements->Add(TVector<int32, 3>(4, 1, 5)); // Front
+			OutElements->Add(TVector<int32, 3>(1, 4, 0));
+			OutElements->Add(TVector<int32, 3>(7, 2, 6)); // Back
+			OutElements->Add(TVector<int32, 3>(2, 7, 3));
+			OutElements->Add(TVector<int32, 3>(6, 0, 4)); // Right
+			OutElements->Add(TVector<int32, 3>(0, 6, 2));
+			OutElements->Add(TVector<int32, 3>(5, 3, 7)); // Left
+			OutElements->Add(TVector<int32, 3>(3, 5, 1));
+			OutElements->Add(TVector<int32, 3>(6, 5, 7)); // Top
+			OutElements->Add(TVector<int32, 3>(5, 6, 4));
+			OutElements->Add(TVector<int32, 3>(0, 2, 1)); // Front
+			OutElements->Add(TVector<int32, 3>(2, 0, 3));
+			/**/
+			//ccw
+			OutElements->Add(TVector<int32, 3>(1, 4, 5)); // Front
+			OutElements->Add(TVector<int32, 3>(4, 1, 0));
+			OutElements->Add(TVector<int32, 3>(2, 7, 6)); // Back
+			OutElements->Add(TVector<int32, 3>(7, 2, 3));
+			OutElements->Add(TVector<int32, 3>(0, 6, 4)); // Right
+			OutElements->Add(TVector<int32, 3>(6, 0, 2));
+			OutElements->Add(TVector<int32, 3>(3, 5, 7)); // Left
+			OutElements->Add(TVector<int32, 3>(5, 3, 1));
+			OutElements->Add(TVector<int32, 3>(5, 6, 7)); // Top
+			OutElements->Add(TVector<int32, 3>(6, 5, 4));
+			OutElements->Add(TVector<int32, 3>(2, 0, 1)); // Front
+			OutElements->Add(TVector<int32, 3>(0, 2, 3));
+		}
+	}
+
+
+	template<class T>
+	TPBDRigidParticleHandle<T, 3>* AppendDynamicParticleBox(TPBDRigidsEvolutionGBF<T, 3>& Evolution, const TVector<T, 3>& Scale, TArray<TVector<int32, 3>>* OutElements)
+	{
+		TArray<TPBDRigidParticleHandle<T, 3>*> Particles = Evolution.CreateDynamicParticles(1);
+		InitDynamicParticleBox2(Particles[0], Scale, OutElements);
+		return Particles[0];
+	}
+	template TPBDRigidParticleHandle<float, 3>* AppendDynamicParticleBox(TPBDRigidsEvolutionGBF<float, 3>& Evolution, const TVector<float, 3>& Scale, TArray<Chaos::TVector<int32, 3>>* OutElements);
+
+	template<class T>
+	TPBDRigidParticleHandle<T, 3>* AppendDynamicParticleBox(TPBDRigidsSOAs<T, 3>& SOAs, const TVector<T, 3>& Scale, TArray<TVector<int32, 3>>* OutElements)
+	{
+		TArray<TPBDRigidParticleHandle<T, 3>*> Particles = SOAs.CreateDynamicParticles(1);
+		InitDynamicParticleBox2(Particles[0], Scale, OutElements);
+		return Particles[0];
+	}
+	template TPBDRigidParticleHandle<float, 3>* AppendDynamicParticleBox(TPBDRigidsSOAs<float, 3>& Evolution, const TVector<float, 3>& Scale, TArray<Chaos::TVector<int32, 3>>* OutElements);
+
+	template<class T>
+	TPBDRigidParticleHandle<T, 3>* AppendDynamicParticleSphere(TPBDRigidsSOAs<T, 3>& SOAs, const TVector<T, 3>& Scale, TArray<TVector<int32, 3>>* OutElements)
+	{
+		TArray<TPBDRigidParticleHandle<T, 3>*> Particles = SOAs.CreateDynamicParticles(1);
+		InitDynamicParticleSphere2(Particles[0], Scale, OutElements);
+		return Particles[0];
+	}
+	template TPBDRigidParticleHandle<float, 3>* AppendDynamicParticleSphere(TPBDRigidsSOAs<float, 3>& Evolution, const TVector<float, 3>& Scale, TArray<Chaos::TVector<int32, 3>>* OutElements);
+
+	template<class T>
+	TPBDRigidParticleHandle<T, 3>* AppendDynamicParticleCylinder(TPBDRigidsSOAs<T, 3>& SOAs, const TVector<T, 3>& Scale, TArray<TVector<int32, 3>>* OutElements)
+	{
+		TArray<TPBDRigidParticleHandle<T, 3>*> Particles = SOAs.CreateDynamicParticles(1);
+		InitDynamicParticleCylinder2(Particles[0], Scale, OutElements, false);
+		return Particles[0];
+	}
+	template TPBDRigidParticleHandle<float, 3>* AppendDynamicParticleCylinder(TPBDRigidsSOAs<float, 3>& Evolution, const TVector<float, 3>& Scale, TArray<Chaos::TVector<int32, 3>>* OutElements);
+
+	template<class T>
+	TPBDRigidParticleHandle<T, 3>* AppendDynamicParticleTaperedCylinder(TPBDRigidsSOAs<T, 3>& SOAs, const TVector<T, 3>& Scale, TArray<TVector<int32, 3>>* OutElements)
+	{
+		TArray<TPBDRigidParticleHandle<T, 3>*> Particles = SOAs.CreateDynamicParticles(1);
+		InitDynamicParticleCylinder2(Particles[0], Scale, OutElements, true);
+		return Particles[0];
+	}
+	template TPBDRigidParticleHandle<float, 3>* AppendDynamicParticleTaperedCylinder(TPBDRigidsSOAs<float, 3>& Evolution, const TVector<float, 3>& Scale, TArray<Chaos::TVector<int32, 3>>* OutElements);
+
+	template<class T>
+	TPBDRigidParticleHandle<T, 3>* AppendClusteredParticleBox(TPBDRigidsEvolutionGBF<T, 3>& Evolution, const TVector<T, 3>& Scale, TArray<TVector<int32, 3>>* OutElements)
+	{
+		auto Particles = Evolution.CreateClusteredParticles(1);
+		InitDynamicParticleBox2(Particles[0], Scale, OutElements);
+		return Particles[0];
+	}
+	template TPBDRigidParticleHandle<float, 3>* AppendClusteredParticleBox(TPBDRigidsEvolutionGBF<float, 3>& Evolution, const TVector<float, 3>& Scale, TArray<Chaos::TVector<int32, 3>>* OutElements);
+
+	template<class T>
+	TPBDRigidParticleHandle<T, 3>* AppendClusteredParticleBox(TPBDRigidsSOAs<T, 3>& SOAs, const TVector<T, 3>& Scale, TArray<TVector<int32, 3>>* OutElements)
+	{
+		auto Particles = SOAs.CreateClusteredParticles(1);
+		InitDynamicParticleBox2(Particles[0], Scale, OutElements);
+		return Particles[0];
+	}
+	template TPBDRigidParticleHandle<float, 3>* AppendClusteredParticleBox(TPBDRigidsSOAs<float, 3>& SOAs, const TVector<float, 3>& Scale, TArray<Chaos::TVector<int32, 3>>* OutElements);
+
+	template<class T>
+	void InitStaticParticleBox(TGeometryParticleHandle<T, 3>* Particle, const TVector<T, 3>& Scale, TArray<TVector<int32, 3>>* OutElements)
+	{
+		Particle->X() = TVector<T, 3>(0.f, 0.f, 0.f);
+		Particle->R() = TRotation<T, 3>::MakeFromEuler(TVector<T, 3>(0.f, 0.f, 0.f)).GetNormalized();
+
+		check(Scale.X == Scale.Y && Scale.X == Scale.Z);
+		T ScaleSq = Scale.X * Scale.X;
+
+		Particle->SetDynamicGeometry(MakeUnique<TBox<T, 3>>(-Scale / 2.f, Scale / 2.f));
+
+		if (OutElements != nullptr)
+		{
+			/*
+			//cw
+			OutElements->Add(TVector<int32, 3>(4, 1, 5)); // Front
+			OutElements->Add(TVector<int32, 3>(1, 4, 0));
+			OutElements->Add(TVector<int32, 3>(7, 2, 6)); // Back
+			OutElements->Add(TVector<int32, 3>(2, 7, 3));
+			OutElements->Add(TVector<int32, 3>(6, 0, 4)); // Right
+			OutElements->Add(TVector<int32, 3>(0, 6, 2));
+			OutElements->Add(TVector<int32, 3>(5, 3, 7)); // Left
+			OutElements->Add(TVector<int32, 3>(3, 5, 1));
+			OutElements->Add(TVector<int32, 3>(6, 5, 7)); // Top
+			OutElements->Add(TVector<int32, 3>(5, 6, 4));
+			OutElements->Add(TVector<int32, 3>(0, 2, 1)); // Front
+			OutElements->Add(TVector<int32, 3>(2, 0, 3));
+			/**/
+			//ccw
+			OutElements->Add(TVector<int32, 3>(1, 4, 5)); // Front
+			OutElements->Add(TVector<int32, 3>(4, 1, 0));
+			OutElements->Add(TVector<int32, 3>(2, 7, 6)); // Back
+			OutElements->Add(TVector<int32, 3>(7, 2, 3));
+			OutElements->Add(TVector<int32, 3>(0, 6, 4)); // Right
+			OutElements->Add(TVector<int32, 3>(6, 0, 2));
+			OutElements->Add(TVector<int32, 3>(3, 5, 7)); // Left
+			OutElements->Add(TVector<int32, 3>(5, 3, 1));
+			OutElements->Add(TVector<int32, 3>(5, 6, 7)); // Top
+			OutElements->Add(TVector<int32, 3>(6, 5, 4));
+			OutElements->Add(TVector<int32, 3>(2, 0, 1)); // Front
+			OutElements->Add(TVector<int32, 3>(0, 2, 3));
+		}
+	}
+
+	template<class T>
+	TGeometryParticleHandle<T, 3>* AppendStaticParticleBox(TPBDRigidsSOAs<T, 3>& SOAs, const TVector<T, 3>& Scale, TArray<TVector<int32, 3>>* OutElements)
+	{
+		TArray<TGeometryParticleHandle<T, 3>*> Particles = SOAs.CreateStaticParticles(1);
+		InitStaticParticleBox(Particles[0], Scale, OutElements);
+		return Particles[0];
+	}
+	template TGeometryParticleHandle<float, 3>* AppendStaticParticleBox(TPBDRigidsSOAs<float, 3>& Evolution, const TVector<float, 3>& Scale, TArray<Chaos::TVector<int32, 3>>* OutElements);
+
+	template<class T>
+	int AppendStaticAnalyticFloor(TPBDRigidParticles<T, 3> & InParticles)
+	{
+		InParticles.AddParticles(1);
+		int32 RigidBodyIndex = InParticles.Size() - 1;
+
+		InParticles.X(RigidBodyIndex) = TVector<T, 3>(0.f, 0.f, 0.f);
+		InParticles.V(RigidBodyIndex) = TVector<T, 3>(0.f, 0.f, 0.f);
+		InParticles.R(RigidBodyIndex) = TRotation<T, 3>::MakeFromEuler(TVector<T, 3>(0.f, 0.f, 0.f)).GetNormalized();
+		InParticles.W(RigidBodyIndex) = TVector<T, 3>(0.f, 0.f, 0.f);
+		InParticles.M(RigidBodyIndex) = 1.f;
+		InParticles.InvM(RigidBodyIndex) = 0.f;
+		InParticles.I(RigidBodyIndex) = PMatrix<T, 3, 3>(1.f, 0.f, 0.f, 0.f, 1.f, 0.f, 0.f, 0.f, 1.f);
+		InParticles.InvI(RigidBodyIndex) = PMatrix<T, 3, 3>(0.f, 0.f, 0.f, 0.f, 0.f, 0.f, 0.f, 0.f, 0.f);
+		InParticles.SetDynamicGeometry(RigidBodyIndex, MakeUnique<TPlane<T, 3>>(TVector<T, 3>(0.f, 0.f, 0.f), TVector<T, 3>(0.f, 0.f, 1.f)));
+		InParticles.SetObjectState(RigidBodyIndex, EObjectStateType::Kinematic);
+
+		InParticles.P(RigidBodyIndex) = InParticles.X(RigidBodyIndex);
+		InParticles.Q(RigidBodyIndex) = InParticles.R(RigidBodyIndex);
+
+		return RigidBodyIndex;
+	}
+	template int AppendStaticAnalyticFloor(TPBDRigidParticles<float, 3> &);
+
+	template<class T>
+	TKinematicGeometryParticleHandle<T, 3>* AppendStaticAnalyticFloor(TPBDRigidsEvolutionGBF<T, 3>& Evolution)
+	{
+		TArray<TKinematicGeometryParticleHandle<T, 3>*> Particles = Evolution.CreateKinematicParticles(1);
+		TKinematicGeometryParticleHandle<T, 3>* Particle = Particles[0];
+
+		Particle->X() = TVector<T, 3>(0.f, 0.f, 0.f);
+		Particle->V() = TVector<T, 3>(0.f, 0.f, 0.f);
+		Particle->R() = TRotation<T, 3>::MakeFromEuler(TVector<T, 3>(0.f, 0.f, 0.f)).GetNormalized();
+		Particle->W() = TVector<T, 3>(0.f, 0.f, 0.f);
+		Particle->SetDynamicGeometry( MakeUnique<TPlane<T, 3>>(TVector<T, 3>(0.f, 0.f, 0.f), TVector<T, 3>(0.f, 0.f, 1.f)));
+
+		return Particle;
+	}
+	template TKinematicGeometryParticleHandle<float, 3>* AppendStaticAnalyticFloor(TPBDRigidsEvolutionGBF<float, 3>& Evolution);
+
+	template<class T>
+	TKinematicGeometryParticleHandle<T, 3>* AppendStaticAnalyticFloor(TPBDRigidsSOAs<T, 3>& SOAs)
+	{
+		TArray<TKinematicGeometryParticleHandle<T, 3>*> Particles = SOAs.CreateKinematicParticles(1);
+		TKinematicGeometryParticleHandle<T, 3>* Particle = Particles[0];
+
+		Particle->X() = TVector<T, 3>(0.f, 0.f, 0.f);
+		Particle->V() = TVector<T, 3>(0.f, 0.f, 0.f);
+		Particle->R() = TRotation<T, 3>::MakeFromEuler(TVector<T, 3>(0.f, 0.f, 0.f)).GetNormalized();
+		Particle->W() = TVector<T, 3>(0.f, 0.f, 0.f);
+		Particle->SetDynamicGeometry(MakeUnique<TPlane<T, 3>>(TVector<T, 3>(0.f, 0.f, 0.f), TVector<T, 3>(0.f, 0.f, 1.f)));
+
+		return Particle;
+	}
+	template TKinematicGeometryParticleHandle<float, 3>* AppendStaticAnalyticFloor(TPBDRigidsSOAs<float, 3>& Evolution);
+
+
+	template<class T>
+	TKinematicGeometryParticleHandle<T, 3>* AppendStaticConvexFloor(TPBDRigidsSOAs<T, 3>& SOAs)
+	{
+		TArray<TKinematicGeometryParticleHandle<T, 3>*> Particles = SOAs.CreateKinematicParticles(1);
+		TKinematicGeometryParticleHandle<T, 3>* Particle = Particles[0];
+
+		Particle->X() = TVector<T, 3>(0.f, 0.f, 0.f);
+		Particle->V() = TVector<T, 3>(0.f, 0.f, 0.f);
+		Particle->R() = TRotation<T, 3>::MakeFromEuler(TVector<T, 3>(0.f, 0.f, 0.f)).GetNormalized();
+		Particle->W() = TVector<T, 3>(0.f, 0.f, 0.f);
+
+		Chaos::TParticles<T, 3> Cube;
+		Cube.AddParticles(9);
+		Cube.X(0) = TVector<float, 3>(-1000, -1000, -20);
+		Cube.X(1) = TVector<float, 3>(-1000, -1000, 0);
+		Cube.X(2) = TVector<float, 3>(-1000, 1000, -20);
+		Cube.X(3) = TVector<float, 3>(-1000, 1000, 0);
+		Cube.X(4) = TVector<float, 3>(1000, -1000, -20);
+		Cube.X(5) = TVector<float, 3>(1000, -1000, 0);
+		Cube.X(6) = TVector<float, 3>(1000, 1000, -20);
+		Cube.X(7) = TVector<float, 3>(1000, 1000, 0);
+		Cube.X(8) = TVector<float, 3>(0, 0, 0);
+
+		Particle->SetDynamicGeometry(MakeUnique<FConvex>(Cube));
+
+		return Particle;
+	}
+	template TKinematicGeometryParticleHandle<float, 3>* AppendStaticConvexFloor(TPBDRigidsSOAs<float, 3>& Evolution);
+
+	/**/
+	template<class T>
+	TLevelSet<T, 3> ConstructLevelset(TParticles<T, 3> & SurfaceParticles, TArray<TVector<int32, 3>> & Elements)
+	{
+		// build Particles and bounds
+		Chaos::TAABB<float, 3> BoundingBox(TVector<T, 3>(0), TVector<T, 3>(0));
+		for (int32 CollisionParticleIndex = 0; CollisionParticleIndex < (int32)SurfaceParticles.Size(); CollisionParticleIndex++)
+		{
+			BoundingBox.GrowToInclude(SurfaceParticles.X(CollisionParticleIndex));
+		}
+
+		// build cell domain
+		int32 MaxAxisSize = 10;
+		int MaxAxis = BoundingBox.LargestAxis();
+		TVector<T, 3> Extents = BoundingBox.Extents();
+		Chaos::TVector<int32, 3> Counts(MaxAxisSize * Extents[0] / Extents[MaxAxis], MaxAxisSize * Extents[1] / Extents[MaxAxis], MaxAxisSize * Extents[2] / Extents[MaxAxis]);
+		Counts[0] = Counts[0] < 1 ? 1 : Counts[0];
+		Counts[1] = Counts[1] < 1 ? 1 : Counts[1];
+		Counts[2] = Counts[2] < 1 ? 1 : Counts[2];
+
+		TUniformGrid<float, 3> Grid(BoundingBox.Min(), BoundingBox.Max(), Counts, 1);
+		TTriangleMesh<float> CollisionMesh(MoveTemp(Elements));
+		FErrorReporter ErrorReporter;
+		return TLevelSet<T,3>(ErrorReporter, Grid, SurfaceParticles, CollisionMesh);
+	}
+	template TLevelSet<float, 3> ConstructLevelset(TParticles<float, 3> &, TArray<TVector<int32, 3>> &);
+
+	/**/
+	template<class T>
+	void AppendDynamicParticleConvexBox(TPBDRigidParticleHandle<T, 3> & InParticles, const TVector<T, 3>& Scale)
+	{
+		Chaos::TParticles<T, 3> Cube;
+		Cube.AddParticles(9);
+		Cube.X(0) = TVector<float, 3>(-1, -1, -1)*Scale;
+		Cube.X(1) = TVector<float, 3>(-1, -1, 1)*Scale;
+		Cube.X(2) = TVector<float, 3>(-1, 1, -1)*Scale;
+		Cube.X(3) = TVector<float, 3>(-1, 1, 1)*Scale;
+		Cube.X(4) = TVector<float, 3>(1, -1, -1)*Scale;
+		Cube.X(5) = TVector<float, 3>(1, -1, 1)*Scale;
+		Cube.X(6) = TVector<float, 3>(1, 1, -1)*Scale;
+		Cube.X(7) = TVector<float, 3>(1, 1, 1)*Scale;
+		Cube.X(8) = TVector<float, 3>(0, 0, 0);
+
+		InParticles.X() = TVector<T, 3>(0.f, 0.f, 0.f);
+		InParticles.V() = TVector<T, 3>(0.f, 0.f, 0.f);
+		InParticles.R() = TRotation<T, 3>::MakeFromEuler(TVector<T, 3>(0.f, 0.f, 0.f)).GetNormalized();
+		InParticles.W() = TVector<T, 3>(0.f, 0.f, 0.f);
+		InParticles.P() = InParticles.X();
+		InParticles.Q() = InParticles.R();
+
+		InParticles.M() = 1.f;
+		InParticles.InvM() = 1.f;
+		InParticles.I() = PMatrix<T, 3, 3>(1.f, 0.f, 0.f, 0.f, 1.f, 0.f, 0.f, 0.f, 1.f);
+		InParticles.InvI() = PMatrix<T, 3, 3>(1.f, 0.f, 0.f, 0.f, 1.f, 0.f, 0.f, 0.f, 1.f);
+		InParticles.SetDynamicGeometry(MakeUnique<FConvex>(Cube));
+		InParticles.SetObjectStateLowLevel(EObjectStateType::Dynamic);
+	}
+	template void AppendDynamicParticleConvexBox(TPBDRigidParticleHandle<float, 3> &, const TVector<float, 3>& Scale);
+
+	template<class T>
+	TPBDRigidParticleHandle<T, 3>* AppendDynamicParticleConvexBox(TPBDRigidsSOAs<T, 3>& SOAs, const TVector<T, 3>& Scale)
+	{
+		TArray<TPBDRigidParticleHandle<T, 3>*> Particles = SOAs.CreateDynamicParticles(1);
+		AppendDynamicParticleConvexBox(*Particles[0], Scale);
+		return Particles[0];
+	}
+	template TPBDRigidParticleHandle<float, 3>* AppendDynamicParticleConvexBox(TPBDRigidsSOAs<float, 3>& Evolution, const TVector<float, 3>& Scale);
+
+
+	/**/
+	template<class T>
+	TVector<T,3> ObjectSpacePoint(TPBDRigidParticles<T, 3> & InParticles, const int32 Index, const TVector<T, 3>& WorldSpacePoint)
+	{
+		TRigidTransform<T, 3> LocalToWorld(InParticles.X(Index), InParticles.R(Index));
+		return LocalToWorld.InverseTransformPosition(WorldSpacePoint);
+	}
+	template TVector<float, 3> ObjectSpacePoint(TPBDRigidParticles<float, 3> & InParticles, const int32 Index, const TVector <float , 3 > & WorldSpacePoint);
+
+	template<class T>
+	TVector<T, 3> ObjectSpacePoint(TGeometryParticleHandle<T, 3>& Particle, const TVector<T, 3>& WorldSpacePoint)
+	{
+		TRigidTransform<T, 3> LocalToWorld(Particle.X(), Particle.R());
+		return LocalToWorld.InverseTransformPosition(WorldSpacePoint);
+	}
+	template TVector<float, 3> ObjectSpacePoint(TPBDRigidParticles<float, 3> & InParticles, const int32 Index, const TVector <float, 3 > & WorldSpacePoint);
+
+
+	/**/
+	template<class T>
+	T PhiWithNormal(TPBDRigidParticles<T, 3> & InParticles, const int32 Index, const TVector<T, 3>& WorldSpacePoint, TVector<T, 3>& Normal)
+	{
+		TRigidTransform<T, 3>(InParticles.X(Index), InParticles.R(Index));
+		TVector<T, 3> BodySpacePoint = ObjectSpacePoint(InParticles, Index, WorldSpacePoint);
+		T LocalPhi = InParticles.Geometry(Index)->PhiWithNormal(BodySpacePoint, Normal);
+		Normal = TRigidTransform<T, 3>(InParticles.X(Index), InParticles.R(Index)).TransformVector(Normal);
+		return LocalPhi;
+	}
+	template float PhiWithNormal(TPBDRigidParticles<float, 3> & InParticles, const int32 Index, const TVector<float, 3>& WorldSpacePoint, TVector<float, 3>& Normal);
+
+	/**/
+	template<class T>
+	T SignedDistance(TPBDRigidParticles<T, 3> & InParticles, const int32 Index, const TVector<T, 3>& WorldSpacePoint)
+	{
+		TVector<T, 3> Normal;
+		return PhiWithNormal(InParticles, Index, WorldSpacePoint, Normal);
+	}
+	template float SignedDistance(TPBDRigidParticles<float, 3> & InParticles, const int32 Index, const TVector<float, 3>& WorldSpacePoint);
+
+
+	/**/
+	template<class T>
+	T PhiWithNormal(TGeometryParticleHandle<T, 3> & Particle, const TVector<T, 3>& WorldSpacePoint, TVector<T, 3>& Normal)
+	{
+		TRigidTransform<T, 3>(Particle.X(), Particle.R());
+		TVector<T, 3> BodySpacePoint = ObjectSpacePoint(Particle, WorldSpacePoint);
+		T LocalPhi = Particle.Geometry()->PhiWithNormal(BodySpacePoint, Normal);
+		Normal = TRigidTransform<T, 3>(Particle.X(), Particle.R()).TransformVector(Normal);
+		return LocalPhi;
+	}
+	template float PhiWithNormal(TPBDRigidParticles<float, 3> & InParticles, const int32 Index, const TVector<float, 3>& WorldSpacePoint, TVector<float, 3>& Normal);
+
+	/**/
+	template<class T>
+	T SignedDistance(TGeometryParticleHandle<T, 3>& Particle, const TVector<T, 3>& WorldSpacePoint)
+	{
+		TVector<T, 3> Normal;
+		return PhiWithNormal(Particle, WorldSpacePoint, Normal);
+	}
+	template float SignedDistance(TGeometryParticleHandle<float, 3> & Particle, const TVector<float, 3>& WorldSpacePoint);
+
+	/**/
+	FVec3 RandAxis()
+	{
+		for (int32 It = 0; It < 1000; ++It)
+		{
+			FVec3 Point = FVec3(FMath::RandRange(-1.0f, 1.0f), FMath::RandRange(-1.0f, 1.0f), FMath::RandRange(-1.0f, 1.0f));
+			if (Point.Size() > KINDA_SMALL_NUMBER)
+			{
+				return FVec3(Point.GetSafeNormal());
+			}
+		}
+		return FVec3(FVector::UpVector);
+	}
+
+
+}
diff --git a/Engine/Source/Programs/HeadlessChaos/Public/HeadlessChaosCollisionConstraints.h b/Engine/Source/Programs/HeadlessChaos/Public/HeadlessChaosCollisionConstraints.h
new file mode 100644
--- /dev/null
+++ b/Engine/Source/Programs/HeadlessChaos/Public/HeadlessChaosCollisionConstraints.h
@@ -0,0 +1,123 @@
+// Copyright Epic Games, Inc. All Rights Reserved.
+
+#pragma once
+
+#include "Math/Vector.h"
+#include "Chaos/Box.h"
+#include "Chaos/PBDRigidParticles.h"
+#include "Chaos/PBDCollisionConstraints.h"
+#include "Chaos/PBDCollisionConstraintsContact.h"
+#include "Chaos/CollisionResolution.h"
+#include "Chaos/Collision/CollisionDetector.h"
+#include "Chaos/Plane.h"
+#include "Chaos/Sphere.h"
+#include "Chaos/Utilities.h"
+#include "Chaos/Vector.h"
+#include "Chaos/PBDRigidsSOAs.h"
+
+namespace Chaos
+{	
+
+/**
+ * Test collision constraints.
+ */
+template <class T, int d>
+class TPBDCollisionConstraintAccessor
+{
+public:
+	using FCollisionConstraints = TPBDCollisionConstraints<T, d>;
+	using FConstraintContainerHandle = TPBDCollisionConstraintHandle<T, d>;
+	using FContactConstraintBase = TCollisionConstraintBase<T, d>;
+	using FMultiPointContactConstraint = TRigidBodyMultiPointContactConstraint<T, d>;
+	using FPointContactConstraint = TRigidBodyPointContactConstraint<T, d>;
+	using FMultiPointContactConstraint = TRigidBodyMultiPointContactConstraint<T, d>;
+	using FConstraintHandleAllocator = TConstraintHandleAllocator<TPBDCollisionConstraints<T, d>>;
+	using FConstraintHandleID = TPair<const TGeometryParticleHandle<T, d>*, const TGeometryParticleHandle<T, d>*>;
+	using FCollisionDetector = TCollisionDetector<FSpatialAccelerationBroadPhase, FNarrowPhase, FAsyncCollisionReceiver, FCollisionConstraints>;
+	using FAccelerationStructure = TBoundingVolume<TAccelerationStructureHandle<T, d>, T, d>;
+
+	TPBDCollisionConstraintAccessor() 
+		: SpatialAcceleration(EmptyParticles.GetNonDisabledView())
+		, BroadPhase(EmptyParticles, (T)1, (T)0)
+		, CollisionConstraints(EmptyParticles, EmptyCollided, EmptyPhysicsMaterials, 1, 1)
+		, CollisionDetector(BroadPhase, CollisionConstraints)
+	{}
+
+	TPBDCollisionConstraintAccessor(const TPBDRigidsSOAs<T, d>& InParticles, TArrayCollectionArray<bool>& Collided, const TArrayCollectionArray<TSerializablePtr<FChaosPhysicsMaterial>>& PerParticleMaterials,
+	const int32 PushOutIterations, const int32 PushOutPairIterations, const T Thickness) 
+		: SpatialAcceleration(InParticles.GetNonDisabledView())
+		, BroadPhase(InParticles, (T)1, (T)0)
+		, CollisionConstraints(InParticles, Collided, PerParticleMaterials, 1, 1, Thickness)
+		, CollisionDetector(BroadPhase, CollisionConstraints)
+	{}
+
+	virtual ~TPBDCollisionConstraintAccessor() {}
+	
+	void ComputeConstraints(T Dt)
+	{
+		CollisionStats::FStatData StatData(false);
+		CollisionDetector.GetBroadPhase().SetSpatialAcceleration(&SpatialAcceleration);
+		CollisionDetector.DetectCollisions(Dt, StatData);
+	}
+
+	void Update(FContactConstraintBase& Constraint, T BoundsThickness = T(0) )
+	{
+		Collisions::Update<ECollisionUpdateType::Deepest,T,d>(BoundsThickness, Constraint);
+	}
+
+	void UpdateManifold(FContactConstraintBase& Constraint, T BoundsThickness = T(0))
+	{
+		Collisions::UpdateManifold<T, d>(BoundsThickness, Constraint);
+	}
+
+
+	void UpdateLevelsetConstraint(FPointContactConstraint& Constraint) 
+	{
+		Collisions::UpdateLevelsetLevelsetConstraint<ECollisionUpdateType::Deepest>(T(0), Constraint);
+	}
+
+	int32 NumConstraints() const
+	{
+		return CollisionConstraints.NumConstraints();
+	}
+
+	TCollisionConstraintBase<T,d>& GetConstraint(int32 Index)
+	{
+		if (Index < CollisionConstraints.NumConstraints())
+		{
+			return GetConstraintHandle(Index)->GetContact();
+		}
+		return EmptyConstraint;
+	}
+
+	const FConstraintContainerHandle* GetConstraintHandle(int32 ConstraintIndex) const
+	{
+		return CollisionConstraints.GetAllConstraintHandles()[ConstraintIndex];
+	}
+
+	FConstraintContainerHandle* GetConstraintHandle(int32 ConstraintIndex)
+	{
+		return CollisionConstraints.GetAllConstraintHandles()[ConstraintIndex];
+	}
+
+	void Apply(const T Dt, const TArray<FConstraintContainerHandle*>& InConstraintHandles, const int32 It, const int32 NumIts)
+	{
+		CollisionConstraints.Apply(Dt, InConstraintHandles, It, NumIts);
+	}
+
+	bool ApplyPushOut(const T Dt, const TArray<FConstraintContainerHandle*>& InConstraintHandles, const TSet<const TGeometryParticleHandle<T, d>*>& IsTemporarilyStatic, int32 Iteration, int32 NumIterations)
+	{
+		return CollisionConstraints.ApplyPushOut(Dt, InConstraintHandles, IsTemporarilyStatic, Iteration, NumIterations);
+	}
+
+	FPointContactConstraint EmptyConstraint;
+	TPBDRigidsSOAs<T, d> EmptyParticles;
+	TArrayCollectionArray<bool> EmptyCollided;
+	TArrayCollectionArray<TSerializablePtr<FChaosPhysicsMaterial>> EmptyPhysicsMaterials;
+
+	FAccelerationStructure SpatialAcceleration;
+	FSpatialAccelerationBroadPhase BroadPhase;
+	FCollisionConstraints CollisionConstraints;
+	FCollisionDetector CollisionDetector;
+};
+}
diff --git a/Engine/Source/Programs/HeadlessChaos/Public/HeadlessChaosTestBroadphase.h b/Engine/Source/Programs/HeadlessChaos/Public/HeadlessChaosTestBroadphase.h
new file mode 100644
--- /dev/null
+++ b/Engine/Source/Programs/HeadlessChaos/Public/HeadlessChaosTestBroadphase.h
@@ -0,0 +1,26 @@
+// Copyright Epic Games, Inc. All Rights Reserved.
+
+
+#pragma once
+
+#include "HeadlessChaosTestUtility.h"
+
+namespace ChaosTest {
+
+	template<typename T>
+	void GridBPTest();
+
+	template<typename T>
+	void GridBPTest2();
+
+	template<typename T>
+	void AABBTreeTest();
+
+	template<typename T>
+	void AABBTreeTimesliceTest();
+
+	template<typename T>
+	void BroadphaseCollectionTest();
+
+	void TestPendingSpatialDataHandlePointerConflict();
+}
\ No newline at end of file
diff --git a/Engine/Source/Programs/HeadlessChaos/Public/HeadlessChaosTestCollisions.h b/Engine/Source/Programs/HeadlessChaos/Public/HeadlessChaosTestCollisions.h
new file mode 100644
--- /dev/null
+++ b/Engine/Source/Programs/HeadlessChaos/Public/HeadlessChaosTestCollisions.h
@@ -0,0 +1,45 @@
+// Copyright Epic Games, Inc. All Rights Reserved.
+
+#pragma once
+
+#include "HeadlessChaosTestUtility.h"
+#include "Modules/ModuleManager.h"
+#include "Chaos/PBDCollisionConstraints.h"
+#include "Chaos/PBDRigidsEvolution.h"
+#include "Chaos/PBDRigidParticles.h"
+#include "Chaos/Sphere.h"
+#include "Chaos/Utilities.h"
+
+namespace ChaosTest {
+
+	template<class T>
+	void LevelsetConstraint();
+
+	template<class T>
+	void LevelsetConstraintGJK();
+		
+	template<class T>
+	void CollisionBoxPlane();
+
+	template<class T>
+	void CollisionConvexConvex();
+
+	template<class T>
+	void CollisionBoxPlaneZeroResitution();
+
+	template<class T>
+	void CollisionBoxPlaneRestitution();
+
+	template<class T>
+	void CollisionBoxToStaticBox();
+
+	template<class T>
+	void CollisionPGS();
+
+	template<class T>
+	void CollisionPGS2();
+
+	template<class T>
+	void CollisionBroadphaseVelocityExpansion();
+
+}
\ No newline at end of file
diff --git a/Engine/Source/Programs/HeadlessChaos/Public/HeadlessChaosTestConstraints.h b/Engine/Source/Programs/HeadlessChaos/Public/HeadlessChaosTestConstraints.h
new file mode 100644
--- /dev/null
+++ b/Engine/Source/Programs/HeadlessChaos/Public/HeadlessChaosTestConstraints.h
@@ -0,0 +1,73 @@
+// Copyright Epic Games, Inc. All Rights Reserved.
+
+
+#pragma once
+
+#include "HeadlessChaosTestUtility.h"
+#include "Modules/ModuleManager.h"
+#include "Chaos/PBDRigidsEvolutionGBF.h"
+#include "Chaos/PBDRigidParticles.h"
+#include "Chaos/Sphere.h"
+#include "Chaos/Utilities.h"
+
+namespace ChaosTest 
+{
+	/**
+	 * Base class for constraint tests. Provides a basic sim with no builtin constraint support.
+	 */
+	class FConstraintsTest
+	{
+	public:
+
+		FConstraintsTest(const int32 NumIterations, const FReal Gravity)
+			: Evolution(SOAs, NumIterations)
+		{
+			PhysicalMaterial = MakeUnique<FChaosPhysicsMaterial>();
+			PhysicalMaterial->Friction = 0;
+			PhysicalMaterial->Restitution = 0;
+			PhysicalMaterial->SleepingLinearThreshold = 0;
+			PhysicalMaterial->SleepingAngularThreshold = 0;
+			PhysicalMaterial->DisabledLinearThreshold = 0;
+			PhysicalMaterial->DisabledAngularThreshold = 0;
+
+			Evolution.GetGravityForces().SetAcceleration(Gravity * FVec3(0, 0, -1));
+		}
+
+		virtual ~FConstraintsTest()
+		{
+		}
+
+		auto AddParticleBox(const FVec3& Position, const FRotation3& Rotation, const FVec3& Size, FReal Mass)
+		{
+			TGeometryParticleHandle<FReal, 3>& Particle = Mass > SMALL_NUMBER ? *AppendDynamicParticleBox<FReal>(SOAs, Size) : *AppendStaticParticleBox<FReal>(SOAs, Size);
+
+			Particle.X() = Position;
+			Particle.R() = Rotation;
+			auto PBDParticlePtr = Particle.CastToRigidParticle();
+			if(PBDParticlePtr && PBDParticlePtr->ObjectState() == EObjectStateType::Dynamic)
+			{
+				auto& PBDParticle = *PBDParticlePtr;
+				PBDParticle.P() = PBDParticle.X();
+				PBDParticle.Q() = PBDParticle.R();
+				PBDParticle.M() = PBDParticle.M() * Mass;
+				PBDParticle.I() = PBDParticle.I() * Mass;
+				PBDParticle.InvM() = PBDParticle.InvM() * ((FReal)1 / Mass);
+				PBDParticle.InvI() = PBDParticle.InvI() * ((FReal)1 / Mass);
+			}
+			Evolution.SetPhysicsMaterial(&Particle, MakeSerializable(PhysicalMaterial));
+
+			return &Particle;
+		}
+
+		// Solver state
+		TPBDRigidsSOAs<FReal, 3> SOAs;
+		TPBDRigidsEvolutionGBF<FReal, 3> Evolution;
+		TUniquePtr<FChaosPhysicsMaterial> PhysicalMaterial;
+
+		TGeometryParticleHandle<FReal, 3>* GetParticle(const int32 Idx)
+		{
+			return SOAs.GetParticleHandles().Handle(Idx).Get();
+		}
+	};
+
+}
\ No newline at end of file
diff --git a/Engine/Source/Programs/HeadlessChaos/Public/HeadlessChaosTestGJK.h b/Engine/Source/Programs/HeadlessChaos/Public/HeadlessChaosTestGJK.h
new file mode 100644
--- /dev/null
+++ b/Engine/Source/Programs/HeadlessChaos/Public/HeadlessChaosTestGJK.h
@@ -0,0 +1,61 @@
+// Copyright Epic Games, Inc. All Rights Reserved.
+
+#pragma once
+
+#include "HeadlessChaosTestUtility.h"
+#include "Modules/ModuleManager.h"
+
+namespace ChaosTest {
+
+	/**/
+	template<typename T>
+	void SimplexLine();
+
+	template<typename T>
+	void SimplexTriangle();
+
+	template <typename T>
+	void SimplexTetrahedron();
+
+	template <typename T>
+	void GJKSphereSphereTest();
+
+	template <typename T>
+	void GJKSphereBoxTest();
+
+	template <typename T>
+	void GJKSphereCapsuleTest();
+
+	template <typename T>
+	void GJKSphereConvexTest();
+
+	template <typename T>
+	void GJKSphereScaledSphereTest();
+
+	template <typename T>
+	void GJKSphereSphereSweep();
+
+	template <typename T>
+	void GJKSphereBoxSweep();
+
+	template <typename T>
+	void GJKSphereCapsuleSweep();
+
+	template <typename T>
+	void GJKSphereConvexSweep();
+
+	template <typename T>
+	void GJKSphereScaledSphereSweep();
+
+	template <typename T>
+	void GJKSphereTransformedSphereSweep();
+
+	template <typename T>
+	void GJKBoxCapsuleSweep();
+
+	template <typename T>
+	void GJKBoxBoxSweep();
+
+	template <typename T>
+	void GJKCapsuleConvexInitialOverlapSweep();
+}
\ No newline at end of file
diff --git a/Engine/Source/Programs/HeadlessChaos/Public/HeadlessChaosTestHandles.h b/Engine/Source/Programs/HeadlessChaos/Public/HeadlessChaosTestHandles.h
new file mode 100644
--- /dev/null
+++ b/Engine/Source/Programs/HeadlessChaos/Public/HeadlessChaosTestHandles.h
@@ -0,0 +1,21 @@
+// Copyright Epic Games, Inc. All Rights Reserved.
+
+#pragma once
+
+namespace ChaosTest
+{
+	namespace Handles
+	{
+		//template<typename T>
+		//void SparseSerializationTest();
+
+		template<typename T>
+		void HandleArrayTest();
+
+		template<typename T>
+		void HandleHeapTest();
+
+		template<typename T>
+		void HandleSerializeTest();
+	}
+}
\ No newline at end of file
diff --git a/Engine/Source/Programs/HeadlessChaos/Public/HeadlessChaosTestRaycast.h b/Engine/Source/Programs/HeadlessChaos/Public/HeadlessChaosTestRaycast.h
new file mode 100644
--- /dev/null
+++ b/Engine/Source/Programs/HeadlessChaos/Public/HeadlessChaosTestRaycast.h
@@ -0,0 +1,42 @@
+// Copyright Epic Games, Inc. All Rights Reserved.
+
+
+#pragma once
+
+#include "HeadlessChaosTestUtility.h"
+
+namespace ChaosTest {
+
+	template<typename T>
+	void SphereRaycast();
+
+	template<typename T>
+	void PlaneRaycast();
+
+	template<typename T>
+	void CylinderRaycast();
+
+	template<typename T>
+	void TaperedCylinderRaycast();
+
+	template<typename T>
+	void CapsuleRaycast();
+
+	template<typename T>
+	void TriangleRaycast();
+
+	template<typename T>
+	void BoxRaycast();
+
+	template<typename T>
+	void ScaledRaycast();
+
+	template<typename T>
+	void TransformedRaycast();
+
+	template<typename T>
+	void UnionRaycast();
+
+	template<typename T>
+	void IntersectionRaycast();
+}
\ No newline at end of file
diff --git a/Engine/Source/Programs/HeadlessChaos/Public/HeadlessChaosTestUtility.h b/Engine/Source/Programs/HeadlessChaos/Public/HeadlessChaosTestUtility.h
new file mode 100644
--- /dev/null
+++ b/Engine/Source/Programs/HeadlessChaos/Public/HeadlessChaosTestUtility.h
@@ -0,0 +1,167 @@
+// Copyright Epic Games, Inc. All Rights Reserved.
+
+#pragma once
+
+#include "HeadlessChaos.h"
+#include "Chaos/PBDRigidParticles.h"
+#include "Chaos/Vector.h"
+#include "Math/Vector.h"
+
+#include "Chaos/Box.h"
+#include "Chaos/PBDRigidParticles.h"
+#include "Chaos/Plane.h"
+#include "Chaos/Sphere.h"
+#include "Chaos/ImplicitObject.h"
+#include "Chaos/Levelset.h"
+#include "Chaos/UniformGrid.h"
+#include "Chaos/Utilities.h"
+#include "Chaos/ParticleHandleFwd.h"
+
+namespace Chaos
+{
+	template<class T, int d> class TPBDRigidsEvolutionGBF;
+	template <typename T, int d> class TPBDRigidsSOAs;
+}
+
+namespace ChaosTest {
+	
+	using namespace Chaos;
+
+	MATCHER_P2(VectorNear, V, Tolerance, "") { return arg.Equals(V, Tolerance); }
+
+	// Expects each component of the vector is within T of its corresponding component in A. 
+#define EXPECT_VECTOR_NEAR(A,B,T) EXPECT_THAT(A, VectorNear(B, T)) << *FString("Expected: " + B.ToString() + "\nActual:   " + A.ToString());
+	// Default comparison to KINDA_SMALL_NUMBER.
+#define EXPECT_VECTOR_NEAR_DEFAULT(A,B) EXPECT_THAT(A, VectorNear(B, KINDA_SMALL_NUMBER)) << *FString("Expected: " + B.ToString() + "\nActual:   " + A.ToString())
+	// Print an additional error string if the expect fails.
+#define EXPECT_VECTOR_NEAR_ERR(A,B,T,E) EXPECT_THAT(A, VectorNear(B, T)) << *(FString("Expected: " + B.ToString() + "\nActual:   " + A.ToString() + "\n") + E);
+
+	// Similar to EXPECT_VECTOR_NEAR_DEFAULT but only prints the component(s) that are wrong, and prints with more precision.
+#define EXPECT_VECTOR_FLOAT_EQ(A, B) EXPECT_FLOAT_EQ(A.X, B.X); EXPECT_FLOAT_EQ(A.Y, B.Y); EXPECT_FLOAT_EQ(A.Z, B.Z);
+	// Print an additional error string if the expect fails. 
+#define EXPECT_VECTOR_FLOAT_EQ_ERR(A, B, E) EXPECT_FLOAT_EQ(A.X, B.X) << *E; EXPECT_FLOAT_EQ(A.Y, B.Y) << *E; EXPECT_FLOAT_EQ(A.Z, B.Z) << *E;
+
+	/**/
+	template<class T>
+	int32 AppendAnalyticSphere(TPBDRigidParticles<T, 3> & InParticles, T Scale = (T)1);
+	template<class T>
+	TPBDRigidParticleHandle<T, 3>* AppendAnalyticSphere2(TPBDRigidsEvolutionGBF<T, 3>& Evolution, T Scale = (T)1);
+
+	/**/
+	template<class T>
+	int32 AppendAnalyticBox(TPBDRigidParticles<T, 3>& InParticles, TVector<T, 3> Scale = TVector<T, 3>(1));
+	template<class T>
+	TKinematicGeometryParticleHandle<T, 3>* AppendKinematicAnalyticBox2(TPBDRigidsEvolutionGBF<T, 3>& Evolution, TVector<T, 3> Scale = TVector<T, 3>(1));
+	template<class T>
+	TPBDRigidParticleHandle<T, 3>* AppendDynamicAnalyticBox2(TPBDRigidsEvolutionGBF<T, 3>& Evolution, TVector<T, 3> Scale = TVector<T, 3>(1));
+
+	/**/
+	template<class T>
+	int32 AppendParticleBox(TPBDRigidParticles<T, 3>& InParticles, TVector<T, 3> Scale = TVector<T, 3>(1), TArray<TVector<int32, 3>>* OutElements = nullptr);
+	template<class T>
+	TPBDRigidParticleHandle<T, 3>* AppendDynamicParticleBox(TPBDRigidsEvolutionGBF<T, 3>& Evolution, const TVector<T, 3>& Scale = TVector<T, 3>(1), TArray<TVector<int32, 3>>* OutElements = nullptr);
+
+	template<class T>
+	TPBDRigidParticleHandle<T, 3>* AppendDynamicParticleBox(TPBDRigidsSOAs<T, 3>& SOAs, const TVector<T, 3>& Scale = TVector<T, 3>(1), TArray<TVector<int32, 3>>* OutElements = nullptr);
+
+	template<class T>
+	TPBDRigidParticleHandle<T, 3>* AppendDynamicParticleSphere(TPBDRigidsSOAs<T, 3>& SOAs, const TVector<T, 3>& Scale = TVector<T, 3>(1), TArray<TVector<int32, 3>>* OutElements = nullptr);
+
+	template<class T>
+	TPBDRigidParticleHandle<T, 3>* AppendDynamicParticleCylinder(TPBDRigidsSOAs<T, 3>& SOAs, const TVector<T, 3>& Scale = TVector<T, 3>(1), TArray<TVector<int32, 3>>* OutElements = nullptr);
+	
+	template<class T>
+	TPBDRigidParticleHandle<T, 3>* AppendDynamicParticleTaperedCylinder(TPBDRigidsSOAs<T, 3>& SOAs, const TVector<T, 3>& Scale = TVector<T, 3>(1), TArray<TVector<int32, 3>>* OutElements = nullptr);
+
+	template<class T>
+	TGeometryParticleHandle<T, 3>* AppendStaticParticleBox(TPBDRigidsSOAs<T, 3>& SOAs, const TVector<T, 3>& Scale = TVector<T, 3>(1), TArray<TVector<int32, 3>>* OutElements = nullptr);
+
+	template<class T>
+	TPBDRigidParticleHandle<T, 3>* AppendClusteredParticleBox(TPBDRigidsEvolutionGBF<T, 3>& Evolution, const TVector<T, 3>& Scale = TVector<T, 3>(1), TArray<TVector<int32, 3>>* OutElements = nullptr);
+
+	template<class T>
+	TPBDRigidParticleHandle<T, 3>* AppendClusteredParticleBox(TPBDRigidsSOAs<T, 3>& SOAs, const TVector<T, 3>& Scale = TVector<T, 3>(1), TArray<TVector<int32, 3>>* OutElements = nullptr);
+		
+	/**/
+	template<class T>
+	int32 AppendStaticAnalyticFloor(TPBDRigidParticles<T, 3>& InParticles);
+	template<class T>
+	TKinematicGeometryParticleHandle<T, 3>* AppendStaticAnalyticFloor(TPBDRigidsEvolutionGBF<T, 3>& Evolution);
+
+	template<class T>
+	TKinematicGeometryParticleHandle<T, 3>* AppendStaticAnalyticFloor(TPBDRigidsSOAs<T, 3>& SOAs);
+
+	template<class T>
+	TKinematicGeometryParticleHandle<T, 3>* AppendStaticConvexFloor(TPBDRigidsSOAs<T, 3>& SOAs);
+		
+	/**/
+	template<class T>
+	TLevelSet<T, 3> ConstructLevelset(TParticles<T, 3> & SurfaceParticles, TArray<TVector<int32, 3>> & Elements);
+
+	/**/
+	template<class T>
+	void AppendDynamicParticleConvexBox(TPBDRigidParticleHandle<T, 3> & InParticles, const TVector<T, 3>& Scale);
+	
+	template<class T>
+	TPBDRigidParticleHandle<T, 3>* AppendDynamicParticleConvexBox(TPBDRigidsSOAs<T, 3>& SOAs, const TVector<T, 3>& Scale = TVector<T, 3>(1));
+
+	/**/
+	template<class T>
+	TVector<T,3>  ObjectSpacePoint(TPBDRigidParticles<T, 3> & InParticles, const int32 Index, const TVector<T, 3>& WorldSpacePoint);
+
+	/**/
+	template<class T>
+	T PhiWithNormal(TPBDRigidParticles<T, 3> & InParticles, const int32 Index, const TVector<T, 3>& WorldSpacePoint, TVector<T, 3>& Normal);
+
+	/**/
+	template<class T>
+	T SignedDistance(TPBDRigidParticles<T, 3> & InParticles, const int32 Index, const TVector<T, 3>& WorldSpacePoint);
+
+	/**/
+	template<class T>
+	TVector<T, 3>  ObjectSpacePoint(TGeometryParticleHandle<T, 3> & Particle, const TVector<T, 3>& WorldSpacePoint);
+
+	/**/
+	template<class T>
+	T PhiWithNormal(TGeometryParticleHandle<T, 3> & Particle, const TVector<T, 3>& WorldSpacePoint, TVector<T, 3>& Normal);
+
+	/**/
+	template<class T>
+	T SignedDistance(TGeometryParticleHandle<T, 3> & Particle, const TVector<T, 3>& WorldSpacePoint);
+
+	/**
+	 * Return a random normalized axis.
+	 * @note: not spherically distributed - actually calculates a point on a box and normalizes.
+	 */
+	FVec3 RandAxis();
+
+
+	/**/
+	inline FVec3 RandomVector(FReal MinValue, FReal MaxValue)
+	{
+		FVec3 V;
+		V.X = FMath::RandRange(MinValue, MaxValue);
+		V.Y = FMath::RandRange(MinValue, MaxValue);
+		V.Z = FMath::RandRange(MinValue, MaxValue);
+		return V;
+	}
+
+	/**/
+	inline FMatrix33 RandomMatrix(FReal MinValue, FReal MaxValue)
+	{
+		return FMatrix33(
+			RandomVector(MinValue, MaxValue),
+			RandomVector(MinValue, MaxValue),
+			RandomVector(MinValue, MaxValue));
+	}
+
+	/**/
+	inline FRotation3 RandomRotation(FReal XMax, FReal YMax, FReal ZMax)
+	{
+		FReal X = FMath::DegreesToRadians(FMath::RandRange(-XMax, XMax));
+		FReal Y = FMath::DegreesToRadians(FMath::RandRange(-YMax, YMax));
+		FReal Z = FMath::DegreesToRadians(FMath::RandRange(-ZMax, ZMax));
+		return FRotation3::FromAxisAngle(FVec3(1, 0, 0), X) * FRotation3::FromAxisAngle(FVec3(0, 1, 0), Y) * FRotation3::FromAxisAngle(FVec3(0, 0, 1), Z);
+	}
+
+}
diff --git a/Engine/Source/Programs/HeadlessChaos/Resources/Windows/Debug-PCLaunch-Win32.manifest b/Engine/Source/Programs/HeadlessChaos/Resources/Windows/Debug-PCLaunch-Win32.manifest
new file mode 100644
--- /dev/null
+++ b/Engine/Source/Programs/HeadlessChaos/Resources/Windows/Debug-PCLaunch-Win32.manifest
@@ -0,0 +1,3 @@
+<?xml version='1.0' encoding='UTF-8' standalone='yes'?>
+<assembly xmlns='urn:schemas-microsoft-com:asm.v1' manifestVersion='1.0'>
+</assembly>
diff --git a/Engine/Source/Programs/HeadlessChaos/Resources/Windows/Debug-PCLaunch-Win64.manifest b/Engine/Source/Programs/HeadlessChaos/Resources/Windows/Debug-PCLaunch-Win64.manifest
new file mode 100644
--- /dev/null
+++ b/Engine/Source/Programs/HeadlessChaos/Resources/Windows/Debug-PCLaunch-Win64.manifest
@@ -0,0 +1,3 @@
+<?xml version='1.0' encoding='UTF-8' standalone='yes'?>
+<assembly xmlns='urn:schemas-microsoft-com:asm.v1' manifestVersion='1.0'>
+</assembly>
diff --git a/Engine/Source/Programs/HeadlessChaos/Resources/Windows/PCLaunch-Win32.manifest b/Engine/Source/Programs/HeadlessChaos/Resources/Windows/PCLaunch-Win32.manifest
new file mode 100644
--- /dev/null
+++ b/Engine/Source/Programs/HeadlessChaos/Resources/Windows/PCLaunch-Win32.manifest
@@ -0,0 +1,15 @@
+﻿<?xml version="1.0" encoding="UTF-8" standalone="yes"?>
+<!-- Copyright 1981-2001 Microsoft Corporation -->
+<assembly xmlns="urn:schemas-microsoft-com:asm.v1" manifestVersion="1.0">
+
+  <!-- Tell Vista we don't need administrator level access-->
+  <ms_asmv2:trustInfo xmlns:ms_asmv2="urn:schemas-microsoft-com:asm.v2">
+    <ms_asmv2:security>
+      <ms_asmv2:requestedPrivileges>
+        <ms_asmv2:requestedExecutionLevel level="asInvoker" uiAccess="false" />
+      </ms_asmv2:requestedPrivileges>
+    </ms_asmv2:security>
+  </ms_asmv2:trustInfo>
+
+</assembly>
+
diff --git a/Engine/Source/Programs/HeadlessChaos/Resources/Windows/PCLaunch-Win64.manifest b/Engine/Source/Programs/HeadlessChaos/Resources/Windows/PCLaunch-Win64.manifest
new file mode 100644
--- /dev/null
+++ b/Engine/Source/Programs/HeadlessChaos/Resources/Windows/PCLaunch-Win64.manifest
@@ -0,0 +1,15 @@
+﻿<?xml version="1.0" encoding="UTF-8" standalone="yes"?>
+<!-- Copyright 1981-2001 Microsoft Corporation -->
+<assembly xmlns="urn:schemas-microsoft-com:asm.v1" manifestVersion="1.0">
+
+  <!-- Tell Vista we don't need administrator level access-->
+  <ms_asmv2:trustInfo xmlns:ms_asmv2="urn:schemas-microsoft-com:asm.v2">
+    <ms_asmv2:security>
+      <ms_asmv2:requestedPrivileges>
+        <ms_asmv2:requestedExecutionLevel level="asInvoker" uiAccess="false" />
+      </ms_asmv2:requestedPrivileges>
+    </ms_asmv2:security>
+  </ms_asmv2:trustInfo>
+
+</assembly>
+
diff --git a/Engine/Source/Programs/HeadlessChaos/Resources/Windows/PCLaunch.rc b/Engine/Source/Programs/HeadlessChaos/Resources/Windows/PCLaunch.rc
new file mode 100644
--- /dev/null
+++ b/Engine/Source/Programs/HeadlessChaos/Resources/Windows/PCLaunch.rc
@@ -0,0 +1,80 @@
+// Copyright Epic Games, Inc. All Rights Reserved.
+
+#include <windows.h>
+#include "Runtime/Launch/Resources/Version.h"
+#include "Runtime/Launch/Resources/Windows/resource.h"
+
+/////////////////////////////////////////////////////////////////////////////
+//
+// Version
+//
+
+VS_VERSION_INFO VERSIONINFO
+ FILEVERSION ENGINE_MAJOR_VERSION,ENGINE_MINOR_VERSION,ENGINE_PATCH_VERSION,0
+ PRODUCTVERSION ENGINE_MAJOR_VERSION,ENGINE_MINOR_VERSION,ENGINE_PATCH_VERSION,0
+ FILEFLAGSMASK 0x17L
+#ifdef _DEBUG
+ FILEFLAGS 0x1L
+#else
+ FILEFLAGS 0x0L
+#endif
+ FILEOS 0x4L
+ FILETYPE 0x2L
+ FILESUBTYPE 0x0L
+BEGIN
+    BLOCK "StringFileInfo"
+    BEGIN
+        BLOCK "040904b0"
+        BEGIN
+            VALUE "CompanyName", EPIC_COMPANY_NAME
+            VALUE "LegalCopyright", EPIC_COPYRIGHT_STRING
+			VALUE "ProductName", EPIC_PRODUCT_NAME
+			VALUE "ProductVersion", ENGINE_VERSION_STRING
+            VALUE "FileDescription", "SlateViewer"
+            VALUE "InternalName", "SlateViewer"
+            VALUE "OriginalFilename", "SlateViewer.exe"
+        END
+    END
+    BLOCK "VarFileInfo"
+    BEGIN
+        VALUE "Translation", 0x409, 1200
+    END
+END
+
+
+
+/////////////////////////////////////////////////////////////////////////////
+//
+// Application Manifest
+//
+
+#if _DEBUG
+
+// Debug manifest
+#if _WIN64
+1						RT_MANIFEST				"Debug-PCLaunch-Win64.manifest"
+#else
+1						RT_MANIFEST				"Debug-PCLaunch-Win32.manifest"
+#endif
+
+#else
+
+// Release manifest
+#if _WIN64
+1						RT_MANIFEST				"PCLaunch-Win64.manifest"
+#else
+1						RT_MANIFEST				"PCLaunch-Win32.manifest"
+#endif
+
+#endif
+
+/////////////////////////////////////////////////////////////////////////////
+//
+// Icon
+//
+
+// Icon with lowest ID value placed first to ensure application icon
+// remains consistent on all systems.
+IDICON_UE4Game          ICON                    "HeadlessChaos.ico"
+
+
