Chaos - fix joint tests

diff --git a/Engine/Source/Programs/HeadlessChaos/Private/HeadlessChaosTestJointSolver.cpp b/Engine/Source/Programs/HeadlessChaos/Private/HeadlessChaosTestJointSolver.cpp
--- a/Engine/Source/Programs/HeadlessChaos/Private/HeadlessChaosTestJointSolver.cpp
+++ b/Engine/Source/Programs/HeadlessChaos/Private/HeadlessChaosTestJointSolver.cpp
@@ -354,17 +354,20 @@ namespace ChaosTest
 		//   DW = DR/DT = -InvI * Stiffness * AngleError * Dt
 		//   DR = -InvI * Stiffness * AngleError * Dt * Dt
 
+		FReal EffectiveInertiaA1 = SolverTestA.Inertia1.Z + SolverTestA.Mass1 * Offset1.X * Offset1.X;
+		FReal EffectiveInertiaB1 = SolverTestB.Inertia1.Z + SolverTestB.Mass1 * Offset1.X * Offset1.X;
+
 		// Verify that the angle change is as expected
 		FReal OutAngle1A = 2.0f * FMath::Asin(SolverTestA.Q1.Z);
 		FReal OutAngleDelta1A = OutAngle1A - Angle1;
-		FReal IIAxis1A = SolverTestA.II1.Z;
+		FReal IIAxis1A = 1.0f / EffectiveInertiaA1;
 		FReal ExpectedAngle1DeltaA = -IIAxis1A * SolverTestA.JointSettings.SoftSwingStiffness * Angle1 * Dt * Dt;
 		EXPECT_NEAR(OutAngleDelta1A, ExpectedAngle1DeltaA, 1.e-6f);
 
 		// Verify that the angle change is as expected
 		FReal OutAngle1B = 2.0f * FMath::Asin(SolverTestB.Q1.Z);
 		FReal OutAngleDelta1B = OutAngle1B - Angle1;
-		FReal IIAxis1B = SolverTestB.II1.Z;
+		FReal IIAxis1B = 1.0f / EffectiveInertiaB1;
 		FReal ExpectedAngle1DeltaB = -IIAxis1B * SolverTestB.JointSettings.SoftSwingStiffness * Angle1 * Dt * Dt;
 		EXPECT_NEAR(OutAngleDelta1B, ExpectedAngle1DeltaB, 1.e-6f);
 
@@ -446,44 +449,44 @@ namespace ChaosTest
 	// assuming the drive torque is T = -K.Theta.
 	GTEST_TEST(JointSolverTests, TestJointSolver_KinematicDynamic_SLerpDrive_ForceMode)
 	{
-		FReal Distances[] = {
-			0.0f,
-			1.0f,
-			10.0f,
-			100.0f,
-			1000.0f
-		};
-
-		FReal AngAccs[] = {
-			10.0f,
-			100.0f,
-			200.0f,
+		FReal DistanceAngAccs[][2] = {
+			{ 0.0f, 10.0f },
+			{ 0.0f, 100.0f },
+			{ 0.0f, 200.0f },
+			{ 1.0f, 10.0f },
+			{ 1.0f, 100.0f },
+			{ 1.0f, 200.0f },
+			{ 10.0f, 10.0f },
+			{ 10.0f, 100.0f },
+			{ 10.0f, 200.0f },
 		};
 
-		for (int32 AccIndex = 0; AccIndex < UE_ARRAY_COUNT(AngAccs); ++AccIndex)
+		for (int32 Index = 0; Index < UE_ARRAY_COUNT(DistanceAngAccs); ++Index)
 		{
-			for (int32 DistanceIndex = 0; DistanceIndex < UE_ARRAY_COUNT(Distances); ++DistanceIndex)
-			{
-				FReal Distance = Distances[DistanceIndex];
-				FReal AngAcc = AngAccs[AccIndex];
-				FReal Mass = 5.0f;
-				FReal Inertia = 200.0f;
-				FReal Stiffness = AngAcc * 1000.0f;
-				FReal Damping = AngAcc * 300.0f;
-
-				FJointSolverTest SolverTestA;
-				KinematicDynamic_SLerpDrive(SolverTestA, Mass, Inertia, EJointForceMode::Force, Stiffness, Damping, AngAcc, Distance);
-
-				FVec3 OutAngles1 = FVec3(
-					2.0f * FMath::Asin(SolverTestA.Q1.X),
-					2.0f * FMath::Asin(SolverTestA.Q1.Y),
-					2.0f * FMath::Asin(SolverTestA.Q1.Z));
-
-				// Calculate the expected angle for the given torque
-				FReal ExpectedAngleDeg = FMath::RadiansToDegrees(SolverTestA.Inertia1.Y * AngAcc / Stiffness);
-				FReal AngleDeg = FMath::RadiansToDegrees(OutAngles1.Y);
-				EXPECT_NEAR(AngleDeg, ExpectedAngleDeg, 0.5f);
-			}
+			FReal Distance = DistanceAngAccs[Index][0];
+			FReal AngAcc = DistanceAngAccs[Index][1];
+			FReal Mass = 5.0f;
+			FReal Inertia = 200.0f;
+			FReal Stiffness = AngAcc * 1000.0f;
+			FReal Damping = AngAcc * 300.0f;
+
+			FJointSolverTest SolverTestA;
+			KinematicDynamic_SLerpDrive(SolverTestA, Mass, Inertia, EJointForceMode::Force, Stiffness, Damping, AngAcc, Distance);
+
+			FVec3 OutAngles1 = FVec3(
+				2.0f * FMath::Asin(SolverTestA.Q1.X),
+				2.0f * FMath::Asin(SolverTestA.Q1.Y),
+				2.0f * FMath::Asin(SolverTestA.Q1.Z));
+
+			FReal EffectiveInertia1 = Inertia + Mass * Distance * Distance;
+
+			// Calculate the expected angle for the given torque
+			// Check for setup errors - if the torque leads to 180 degree rotation, it will keep spinning
+			FReal ExpectedAngleDeg = FMath::RadiansToDegrees(EffectiveInertia1 * AngAcc / Stiffness);
+			EXPECT_LT(ExpectedAngleDeg, 180);
+
+			FReal AngleDeg = FMath::RadiansToDegrees(OutAngles1.Y);
+			EXPECT_NEAR(AngleDeg, ExpectedAngleDeg, 0.5f) << "Distance: " << Distance << "; AngAcc: " << AngAcc;
 		}
 	}
 
