Chaos - ScaleInertia fixes and a unit test

diff --git a/Engine/Source/Programs/HeadlessChaos/Private/HeadlessChaosTestMassProperties.cpp b/Engine/Source/Programs/HeadlessChaos/Private/HeadlessChaosTestMassProperties.cpp
--- a/Engine/Source/Programs/HeadlessChaos/Private/HeadlessChaosTestMassProperties.cpp
+++ b/Engine/Source/Programs/HeadlessChaos/Private/HeadlessChaosTestMassProperties.cpp
@@ -1,5 +1,6 @@
 // Copyright Epic Games, Inc. All Rights Reserved.
 
+#include "Chaos/Box.h"
 #include "Chaos/Matrix.h"
 #include "Chaos/MassProperties.h"
 #include "Chaos/TriangleMesh.h"
@@ -238,6 +239,48 @@ namespace ChaosTest {
 			EXPECT_NEAR(ILocal2.Y, ILocal.Y, KINDA_SMALL_NUMBER);
 			EXPECT_NEAR(ILocal2.Z, ILocal.Z, KINDA_SMALL_NUMBER);
 		}
+	}
+
+	TEST(MassPropertyTests, TestScaleInertia)
+	{
+		FReal Mass = 12.0f;
+		FVec3 Dim = FVec3(2.0f, 5.0f, 10.0f);
+		FVec3 UnscaledInertia = TBox<FReal, 3>::GetInertiaTensor(Mass, Dim).GetDiagonal();
+		FReal UnscaledVolume = Dim.X * Dim.Y * Dim.Z;
 
+		FVec3 Scales[] = {
+			FVec3(1, 1, 1),
+			FVec3(5, 5, 5),
+			FVec3(5, 1, 1),
+			FVec3(1, 5, 1),
+			FVec3(1, 1, 5),
+			FVec3(5, 3, 2),
+			FVec3(0.1f, 0.1f, 0.1f),
+			FVec3(0.1f, 0.5f, 0.3f),
+			FVec3(0.5f, 2.3f, 1.5f),
+		};
+
+		bool bScaleMass = false;
+
+		for (int32 MassScale = 0; MassScale < 2; ++MassScale)
+		{
+			for (int32 Index = 0; Index < UE_ARRAY_COUNT(Scales); ++Index)
+			{
+				const FVec3& Scale = Scales[Index];
+
+				FReal ScaledVolume = Dim.X * Dim.Y * Dim.Z * Scale.X * Scale.Y * Scale.Z;
+				FReal ScaledMass = (bScaleMass) ? Mass * ScaledVolume / UnscaledVolume : Mass;
+				FVec3 ExpectedScaledInertia = TBox<FReal, 3>::GetInertiaTensor(ScaledMass, Dim * Scale).GetDiagonal();
+				
+				FVec3 ScaledInertia = Utilities::ScaleInertia(UnscaledInertia, Scale, bScaleMass);
+
+				// Accurate to 0.1%?
+				EXPECT_NEAR(ScaledInertia.X, ExpectedScaledInertia.X, ExpectedScaledInertia.X / 1000.0f);
+				EXPECT_NEAR(ScaledInertia.Y, ExpectedScaledInertia.Y, ExpectedScaledInertia.Y / 1000.0f);
+				EXPECT_NEAR(ScaledInertia.Z, ExpectedScaledInertia.Z, ExpectedScaledInertia.Z / 1000.0f);
+			}
+
+			bScaleMass = !bScaleMass;
+		}
 	}
 }
diff --git a/Engine/Source/Runtime/Engine/Private/PhysicsEngine/ImmediatePhysics/ImmediatePhysicsChaos/ImmediatePhysicsActorHandle_Chaos.cpp b/Engine/Source/Runtime/Engine/Private/PhysicsEngine/ImmediatePhysics/ImmediatePhysicsChaos/ImmediatePhysicsActorHandle_Chaos.cpp
--- a/Engine/Source/Runtime/Engine/Private/PhysicsEngine/ImmediatePhysics/ImmediatePhysicsChaos/ImmediatePhysicsActorHandle_Chaos.cpp
+++ b/Engine/Source/Runtime/Engine/Private/PhysicsEngine/ImmediatePhysics/ImmediatePhysicsChaos/ImmediatePhysicsActorHandle_Chaos.cpp
@@ -170,7 +170,7 @@ namespace ImmediatePhysics_Chaos
 			Density = BodyInstance->GetMassOverride() / MassProperties.Volume;
 		}
 		OutMass = Density * BodyInstance->MassScale * MassProperties.Volume;
-		OutInertia = Utilities::ScaleInertia(Density * TVector<float, 3>(MassProperties.InertiaTensor.M[0][0], MassProperties.InertiaTensor.M[1][1], MassProperties.InertiaTensor.M[2][2]), BodyInstance->InertiaTensorScale);
+		OutInertia = Utilities::ScaleInertia(Density * TVector<float, 3>(MassProperties.InertiaTensor.M[0][0], MassProperties.InertiaTensor.M[1][1], MassProperties.InertiaTensor.M[2][2]), BodyInstance->InertiaTensorScale, true);	// bScaleMass true to match legacy, but not correct
 		OutCoMTransform = FTransform(MassProperties.RotationOfMass, MassProperties.CenterOfMass + BodyInstance->COMNudge);
 #else
 		OutMass = BodyInstance->GetBodyMass();
diff --git a/Engine/Source/Runtime/Experimental/Chaos/Public/Chaos/Utilities.h b/Engine/Source/Runtime/Experimental/Chaos/Public/Chaos/Utilities.h
--- a/Engine/Source/Runtime/Experimental/Chaos/Public/Chaos/Utilities.h
+++ b/Engine/Source/Runtime/Experimental/Chaos/Public/Chaos/Utilities.h
@@ -288,14 +288,17 @@ namespace Chaos
 
 		/**
 		 * Given the local-space inertia for an unscaled object, return an inertia as if generated from a non-uniformly scaled shape with the specified scale.
+		 * If bScaleMass is true, it also takes into account the fact that the mass would have changed by the increase in volume.
 		 */
-		inline Chaos::FVec3 ScaleInertia(const Chaos::FVec3& Inertia, const Chaos::FVec3& Scale)
+		inline FVec3 ScaleInertia(const FVec3& Inertia, const FVec3& Scale, const bool bScaleMass)
 		{
-			using namespace Chaos;
-
-			FReal ScaleVol = Scale.X * Scale.Y * Scale.Z;
-			FVec3 ScaleSq = Scale * Scale;
-			return 0.5f * ScaleVol * (FVec3(Inertia.Y + Inertia.Z, Inertia.X + Inertia.Z, Inertia.X + Inertia.Y) * ScaleSq);
+			FVec3 XYZSq = (FVec3(0.5f * (Inertia.X + Inertia.Y + Inertia.Z)) - Inertia) * Scale * Scale;
+			FReal XX = XYZSq.Y + XYZSq.Z;
+			FReal YY = XYZSq.X + XYZSq.Z;
+			FReal ZZ = XYZSq.X + XYZSq.Y;
+			FVec3 ScaledInertia = FVec3(XX, YY, ZZ);
+			FReal MassScale = (bScaleMass) ? Scale.X * Scale.Y * Scale.Z : 1.0f;
+			return MassScale * ScaledInertia;
 		}
 
 	} // namespace Utilities
