Fix MTD not being rotated in overlap queries

diff --git a/Engine/Source/Programs/HeadlessChaos/Private/HeadlessChaosTestGJK.cpp b/Engine/Source/Programs/HeadlessChaos/Private/HeadlessChaosTestGJK.cpp
--- a/Engine/Source/Programs/HeadlessChaos/Private/HeadlessChaosTestGJK.cpp
+++ b/Engine/Source/Programs/HeadlessChaos/Private/HeadlessChaosTestGJK.cpp
@@ -1378,37 +1378,86 @@ namespace ChaosTest
 	template <typename T>
 	void GJKCapsuleConvexInitialOverlapSweep()
 	{
-		TParticles<T, 3> ConvexParticles;
-		ConvexParticles.AddParticles(8);
-
-		ConvexParticles.X(0) = {-256.000031, 12.0000601, 384.000061};
-		ConvexParticles.X(1) = {256.000031, 12.0000601, 384.000061};
-		ConvexParticles.X(2) = {256.000031, 12.0000601, 6.10351563e-05};
-		ConvexParticles.X(3) = {-256.000031, -11.9999399, 6.10351563e-05};
-		ConvexParticles.X(4) = {-256.000031, 12.0000601, 6.10351563e-05};
-		ConvexParticles.X(5) = {-256.000031, -11.9999399, 384.000061};
-		ConvexParticles.X(6) = {256.000031, -11.9999399, 6.10351563e-05};
-		ConvexParticles.X(7) = {256.000031, -11.9999399, 384.000061};
-
-		TUniquePtr<FConvex> UniqueConvex = MakeUnique<FConvex>(ConvexParticles);
-		TSerializablePtr<FConvex> AConv(UniqueConvex);
-		const TImplicitObjectScaled<FConvex> A(AConv, TVec3<T>(1.0, 1.0, 1.0));
-
-		const TVec3<T> Pt0(0.0, 0.0, -33.0);
-		TVec3<T> Pt1 = Pt0;
-		Pt1 += (TVec3<T>(0.0, 0.0, 1.0) * 66.0);
-
-		const TCapsule<T> B(Pt0, Pt1, 42.0);
-
-		const TRigidTransform<T, 3> BToATM(TVec3<T>(157.314758, -54.0000839, 76.1436157), TRotation<T, 3>::FromElements(0.0, 0.0, 0.704960823, 0.709246278));
-		const TVec3<T> LocalDir(-0.00641351938, -0.999979556, 0.0);
-		const T Length = 0.0886496082;
-		const TVec3<T> SearchDir(-3.06152344, 166.296631, -76.1436157);
-
-		T Time;
-		TVec3<T> Position, Normal;
-		EXPECT_TRUE(GJKRaycast2<T>(A, B, BToATM, LocalDir, Length, Time, Position, Normal, 0, true, SearchDir, 0));
-		EXPECT_FLOAT_EQ(Time, 0.0);
+		{
+			TParticles<T,3> ConvexParticles;
+			ConvexParticles.AddParticles(8);
+
+			ConvexParticles.X(0) ={-256.000031,12.0000601,384.000061};
+			ConvexParticles.X(1) ={256.000031,12.0000601,384.000061};
+			ConvexParticles.X(2) ={256.000031,12.0000601,6.10351563e-05};
+			ConvexParticles.X(3) ={-256.000031,-11.9999399,6.10351563e-05};
+			ConvexParticles.X(4) ={-256.000031,12.0000601,6.10351563e-05};
+			ConvexParticles.X(5) ={-256.000031,-11.9999399,384.000061};
+			ConvexParticles.X(6) ={256.000031,-11.9999399,6.10351563e-05};
+			ConvexParticles.X(7) ={256.000031,-11.9999399,384.000061};
+
+			TUniquePtr<FConvex> UniqueConvex = MakeUnique<FConvex>(ConvexParticles);
+			TSerializablePtr<FConvex> AConv(UniqueConvex);
+			const TImplicitObjectScaled<FConvex> A(AConv,TVec3<T>(1.0,1.0,1.0));
+
+			const TVec3<T> Pt0(0.0,0.0,-33.0);
+			TVec3<T> Pt1 = Pt0;
+			Pt1 += (TVec3<T>(0.0,0.0,1.0) * 66.0);
+
+			const TCapsule<T> B(Pt0,Pt1,42.0);
+
+			const TRigidTransform<T,3> BToATM(TVec3<T>(157.314758,-54.0000839,76.1436157),TRotation<T,3>::FromElements(0.0,0.0,0.704960823,0.709246278));
+			const TVec3<T> LocalDir(-0.00641351938,-0.999979556,0.0);
+			const T Length = 0.0886496082;
+			const TVec3<T> SearchDir(-3.06152344,166.296631,-76.1436157);
+
+			T Time;
+			TVec3<T> Position,Normal;
+			EXPECT_TRUE(GJKRaycast2<T>(A,B,BToATM,LocalDir,Length,Time,Position,Normal,0,true,SearchDir,0));
+			EXPECT_FLOAT_EQ(Time,0.0);
+		}
+
+		{
+			TParticles<T,3> ConvexParticles;
+			ConvexParticles.AddParticles(16);
+
+			ConvexParticles.X(0) ={-127.216454,203.240234,124.726524};
+			ConvexParticles.X(1) ={125.708847,203.240295,124.726524};
+			ConvexParticles.X(2) ={-120.419685,207.124924,-0.386817127};
+			ConvexParticles.X(3) ={-32.9052734,91.5147095,199.922119};
+			ConvexParticles.X(4) ={118.912071,91.3693237,155.363205};
+			ConvexParticles.X(5) ={31.3977623,91.5147705,199.922150};
+			ConvexParticles.X(6) ={115.392204,91.6678925,162.647476};
+			ConvexParticles.X(7) ={-120.419701,91.1026840,-0.386809498};
+			ConvexParticles.X(8) ={118.912086,207.124985,-0.386806667};
+			ConvexParticles.X(9) ={118.912086,91.1027603,-0.386806667};
+			ConvexParticles.X(10) ={-120.419685,91.3692703,155.363174};
+			ConvexParticles.X(11) ={-110.103012,199.020554,160.910324};
+			ConvexParticles.X(12) ={-116.899742,91.6678467,162.647491};
+			ConvexParticles.X(13) ={31.3977337,194.240265,194.534988};
+			ConvexParticles.X(14) ={-32.9052925,194.240204,194.534958};
+			ConvexParticles.X(15) ={108.595482,199.020599,160.910309};
+
+			auto Convex = MakeShared<FConvex, ESPMode::ThreadSafe>(ConvexParticles);
+			const auto& A = *Convex;
+			//const TImplicitObjectInstanced<FConvex> A(Convex);
+
+			const TVec3<T> Pt0(0.0,0.0,-45);
+			TVec3<T> Pt1 = Pt0;
+			Pt1 += (TVec3<T>(0.0,0.0,1.0) * 90);
+
+			const TCapsule<T> B(Pt0,Pt1,33.8499985);
+
+			const TRigidTransform<T,3> ATM(TVec3<T>(2624.00024, -383.998962, 4.00000000),TRotation<T,3>::FromElements(-5.07916162e-08, -3.39378659e-08, 0.555569768, 0.831469893));
+			const TRigidTransform<T,3> BTM(TVec3<T>(2461.92749, -205.484283, 106.071632),TRotation<T,3>::FromElements(0,0,0,1));
+			const TRigidTransform<T,3> BToATM(TVec3<T>(102.903252, 218.050415, 102.071655),TRotation<T,3>::FromElements(5.07916162e-08, 3.39378659e-08, -0.555569768, 0.831469893));
+
+			T Penetration;
+			TVec3<T> ClosestA,ClosestB,Normal;
+			const TVec3<T> Offset ={162.072754,-178.514679,-102.071632};
+			EXPECT_TRUE(GJKPenetration<T>(A,B,BToATM,Penetration,ClosestA,ClosestB,Normal,0,Offset,0));
+
+			const TRigidTransform<T,3> NewAToBTM (BToATM.GetTranslation() + (0.01 + Penetration) * Normal,BToATM.GetRotation());;
+
+			EXPECT_FALSE(GJKPenetration<T>(A,B,NewAToBTM,Penetration,ClosestA,ClosestB,Normal,0,Offset,0));
+
+		}
+		
 	}
 
 	template void SimplexLine<float>();
diff --git a/Engine/Source/Runtime/Experimental/Chaos/Public/Chaos/GeometryQueries.h b/Engine/Source/Runtime/Experimental/Chaos/Public/Chaos/GeometryQueries.h
--- a/Engine/Source/Runtime/Experimental/Chaos/Public/Chaos/GeometryQueries.h
+++ b/Engine/Source/Runtime/Experimental/Chaos/Public/Chaos/GeometryQueries.h
@@ -94,9 +94,17 @@ namespace Chaos
 			const FVec3 Offset = ATM.GetLocation() - BTM.GetLocation();
 			if (OutMTD)
 			{
-				FVec3 ClosestA;
-				FVec3 ClosestB;
-				return CastHelper(A, BToATM, [&](const auto& AConcrete, const auto& BToAFullTM) { return GJKPenetration<FReal>(AConcrete, B, BToAFullTM, OutMTD->Penetration, ClosestA, ClosestB, OutMTD->Normal, Thickness, Offset.SizeSquared() < 1e-4 ? FVec3(1, 0, 0) : Offset); });
+				return CastHelper(A, BToATM, [&](const auto& AConcrete, const auto& BToAFullTM)
+				{
+					FVec3 LocalA,LocalB,LocalNormal;
+					if(GJKPenetration<FReal>(AConcrete,B,BToAFullTM,OutMTD->Penetration,LocalA,LocalB,LocalNormal,Thickness,Offset.SizeSquared() < 1e-4 ? FVec3(1,0,0) : Offset))
+					{
+						OutMTD->Normal = ATM.TransformVectorNoScale(LocalNormal);
+						return true;
+					}
+
+					return false;
+				});
 			}
 			else
 			{
