HeadlessChaos : Solver Wake Events

- Unit tests for solver wake events.

diff --git a/Engine/Source/Programs/HeadlessChaos/Private/HeadlessChaosApp.cpp b/Engine/Source/Programs/HeadlessChaos/Private/HeadlessChaosApp.cpp
--- a/Engine/Source/Programs/HeadlessChaos/Private/HeadlessChaosApp.cpp
+++ b/Engine/Source/Programs/HeadlessChaos/Private/HeadlessChaosApp.cpp
@@ -221,6 +221,7 @@ TEST(ChaosSolver, ChaosSolverTests)
 {
 	ChaosTest::SingleParticleProxySingleThreadTest<float>();
 	ChaosTest::CommandListTest<float>();
+	ChaosTest::SingleParticleProxyWakeEventPropergationTest<float>();
 }
 TEST(ChaosSolver, DISABLED_ChaosSolverTests)
 {
diff --git a/Engine/Source/Programs/HeadlessChaos/Private/HeadlessChaosTestSolverProxies.cpp b/Engine/Source/Programs/HeadlessChaos/Private/HeadlessChaosTestSolverProxies.cpp
--- a/Engine/Source/Programs/HeadlessChaos/Private/HeadlessChaosTestSolverProxies.cpp
+++ b/Engine/Source/Programs/HeadlessChaos/Private/HeadlessChaosTestSolverProxies.cpp
@@ -130,4 +130,65 @@ namespace ChaosTest {
 	}
 
 	template void SingleParticleProxyTaskGraphTest<float>();
+
+
+	template<class T>
+	void SingleParticleProxyWakeEventPropergationTest()
+	{
+		auto Sphere = TSharedPtr<FImplicitObject, ESPMode::ThreadSafe>(new TSphere<float, 3>(TVector<float, 3>(0), 10));
+
+		FChaosSolversModule* Module = FChaosSolversModule::GetModule();
+		Module->ChangeThreadingMode(EChaosThreadingMode::SingleThread);
+
+		// Make a solver
+		Chaos::FPhysicsSolver* Solver = Module->CreateSolver(true);
+		Solver->SetEnabled(true);
+
+		// Make a particle
+
+		TUniquePtr<Chaos::TPBDRigidParticle<float, 3>> Particle = Chaos::TPBDRigidParticle<float, 3>::CreateParticle();
+		Particle->SetGeometry(Sphere);
+		Particle->SetX(TVector<float, 3>(0, 0, 1000));
+		Particle->SetV(TVector<float, 3>(0, 0, -10));
+		Solver->RegisterObject(Particle.Get());
+		Solver->AddDirtyProxy(Particle->Proxy);
+
+		TUniquePtr<Chaos::TPBDRigidParticle<float, 3>> Particle2 = Chaos::TPBDRigidParticle<float, 3>::CreateParticle();
+		Particle2->SetGeometry(Sphere);
+		Particle2->SetX(TVector<float, 3>(0, 0, 100));
+		Particle2->SetV(TVector<float, 3>(0, 0, 0));
+		Solver->RegisterObject(Particle2.Get());
+		Solver->AddDirtyProxy(Particle2->Proxy);
+		Particle2->SetObjectState(Chaos::EObjectStateType::Sleeping, false);
+
+		Solver->PushPhysicsState(Module->GetDispatcher());
+
+		FPhysicsSolverAdvanceTask AdvanceTask(Solver, 100.0f);
+		AdvanceTask.DoTask(ENamedThreads::GameThread, FGraphEventRef());
+
+		Solver->BufferPhysicsResults();
+		Solver->FlipBuffers();
+		Solver->UpdateGameThreadStructures();
+
+		// Make sure game thread data has changed
+		TVector<float, 3> V = Particle->V();
+		EXPECT_EQ(Particle->HasAwakeEvent(), false);
+		EXPECT_EQ(Particle->ObjectState(), Chaos::EObjectStateType::Dynamic);
+
+		EXPECT_EQ(Particle2->HasAwakeEvent(), true);
+		EXPECT_EQ(Particle2->ObjectState(), Chaos::EObjectStateType::Dynamic);
+
+		Particle2->ClearEvents();
+		EXPECT_EQ(Particle2->HasAwakeEvent(), false);
+
+		// Throw out the proxy
+		Solver->UnregisterObject(Particle.Get());
+
+		Module->DestroySolver(Solver);
+	}
+
+	template void SingleParticleProxyWakeEventPropergationTest<float>();
+
+
+
 }
diff --git a/Engine/Source/Programs/HeadlessChaos/Public/HeadlessChaosTestSolverProxies.h b/Engine/Source/Programs/HeadlessChaos/Public/HeadlessChaosTestSolverProxies.h
--- a/Engine/Source/Programs/HeadlessChaos/Public/HeadlessChaosTestSolverProxies.h
+++ b/Engine/Source/Programs/HeadlessChaos/Public/HeadlessChaosTestSolverProxies.h
@@ -10,4 +10,6 @@ namespace ChaosTest
 	template <typename T>
 	void SingleParticleProxyTaskGraphTest();
 
+	template <typename T>
+	void SingleParticleProxyWakeEventPropergationTest();
 }
