- Headless Chaos: Almost degenerate simplex test added
- Headless Chaos: Compile fix

diff --git a/Engine/Source/Programs/HeadlessChaos/Private/HeadlessChaosTestGJK.cpp b/Engine/Source/Programs/HeadlessChaos/Private/HeadlessChaosTestGJK.cpp
--- a/Engine/Source/Programs/HeadlessChaos/Private/HeadlessChaosTestGJK.cpp
+++ b/Engine/Source/Programs/HeadlessChaos/Private/HeadlessChaosTestGJK.cpp
@@ -401,6 +401,22 @@ namespace ChaosTest
 			EXPECT_EQ(Idxs[2], 3);
 			EXPECT_FLOAT_EQ(Barycentric[0] + Barycentric[1] + Barycentric[3], 1);
 		}
+
+		{
+			// Previous failing case observed with Voronoi region implementation - Not quite degenerate (totally degenerate cases work)
+			T Barycentric[4];
+			TVector<T, 3> Simplex[] = { { -15.9112930, -15.2787428, 1.33070087 },
+										{ 1.90487099, 2.25161266, 0.439208984 },
+										{ -15.8914719, -15.2915068, 1.34186459 },
+										{ 1.90874290, 2.24025059, 0.444719315 } };
+
+			FSimplex Idxs = { 0,1,2,3 };
+			const TVector<T, 3> ClosestPoint = TetrahedronSimplexFindOrigin(Simplex, Idxs, Barycentric);
+			EXPECT_EQ(Idxs.NumVerts, 3);
+			EXPECT_EQ(Idxs[0], 0);
+			EXPECT_EQ(Idxs[1], 1);
+			EXPECT_EQ(Idxs[2], 2);
+		}
 	}
 
 	//For each gjk test we should test:
@@ -956,7 +972,7 @@ namespace ChaosTest
 			//MTD
 			T Penetration;
 			TVec3<T> ClosestA, ClosestB;
-			EXPECT_TRUE(GJKPenetration<T>(A, B, TRigidTransform<T, 3>(TVector<T, 3>(2.5, 0, 2), TRotation<T, 3>::Identity), Penetration, ClosestA, ClosestB, Normal, 0, InitialDir));
+			EXPECT_TRUE((GJKPenetration<false, T>(A, B, TRigidTransform<T, 3>(TVector<T, 3>(2.5, 0, 2), TRotation<T, 3>::Identity), Penetration, ClosestA, ClosestB, Normal, 0, InitialDir)));
 			EXPECT_FLOAT_EQ(Penetration, 0.5);
 			EXPECT_VECTOR_NEAR(Normal, TVector3(-1, 0, 0).GetUnsafeNormal(), Eps);
 			EXPECT_NEAR(ClosestA[0], 3, Eps);	//could be any point on face, but should have x == 3
@@ -1238,7 +1254,7 @@ namespace ChaosTest
 			//MTD
 			T Penetration;
 			TVec3<T> ClosestA, ClosestB;
-			EXPECT_TRUE(GJKPenetration<T>(A, B, TRigidTransform<T, 3>(TVector<T, 3>(2.5, 0, 0), TRotation<T, 3>::Identity), Penetration, ClosestA, ClosestB, Normal, 0, InitialDir));
+			EXPECT_TRUE((GJKPenetration<false, T>(A, B, TRigidTransform<T, 3>(TVector<T, 3>(2.5, 0, 0), TRotation<T, 3>::Identity), Penetration, ClosestA, ClosestB, Normal, 0, InitialDir)));
 			EXPECT_FLOAT_EQ(Penetration, 1.5);
 			EXPECT_VECTOR_NEAR(Normal, TVec3<T>(-1, 0, 0), Eps);
 			EXPECT_NEAR(ClosestA[0], 3, Eps);	//could be any point on face, but should have x == 3
@@ -1252,7 +1268,7 @@ namespace ChaosTest
 			EXPECT_VECTOR_NEAR(Normal, TVec3<T>(-1, 0, 0), Eps);
 
 			//EPA
-			EXPECT_TRUE(GJKPenetration<T>(A, B, TRigidTransform<T, 3>(TVector<T, 3>(3, 0, 0), TRotation<T, 3>::Identity), Penetration, ClosestA, ClosestB, Normal, 0, InitialDir));
+			EXPECT_TRUE((GJKPenetration<false, T>(A, B, TRigidTransform<T, 3>(TVector<T, 3>(3, 0, 0), TRotation<T, 3>::Identity), Penetration, ClosestA, ClosestB, Normal, 0, InitialDir)));
 			EXPECT_NEAR(Penetration, 2, Eps);
 			EXPECT_VECTOR_NEAR(Normal, TVec3<T>(-1, 0, 0), Eps);
 			EXPECT_NEAR(ClosestA[0], 3, Eps);	//could be any point on face, but should have x == 3
@@ -1266,7 +1282,7 @@ namespace ChaosTest
 			EXPECT_VECTOR_NEAR(Normal, TVec3<T>(-1, 0, 0), Eps);
 
 			//EPA
-			EXPECT_TRUE(GJKPenetration<T>(A, B, TRigidTransform<T, 3>(TVector<T, 3>(3.25, 0, 0), TRotation<T, 3>::Identity), Penetration, ClosestA, ClosestB, Normal, 0, InitialDir));
+			EXPECT_TRUE((GJKPenetration<false, T>(A, B, TRigidTransform<T, 3>(TVector<T, 3>(3.25, 0, 0), TRotation<T, 3>::Identity), Penetration, ClosestA, ClosestB, Normal, 0, InitialDir)));
 			EXPECT_NEAR(Penetration, 2.25, Eps);
 			EXPECT_VECTOR_NEAR(Normal, TVec3<T>(-1, 0, 0), Eps);
 			EXPECT_NEAR(ClosestA[0], 3, Eps);	//could be any point on face, but should have x == 3
@@ -1280,7 +1296,7 @@ namespace ChaosTest
 			EXPECT_VECTOR_NEAR(Normal, TVec3<T>(0, 0, -1), Eps);
 
 			//MTD
-			EXPECT_TRUE(GJKPenetration<T>(A, B, TRigidTransform<T, 3>(TVector<T, 3>(3.25, 0, -2.875), TRotation<T, 3>::Identity), Penetration, ClosestA, ClosestB, Normal, 0, InitialDir));
+			EXPECT_TRUE((GJKPenetration<false, T>(A, B, TRigidTransform<T, 3>(TVector<T, 3>(3.25, 0, -2.875), TRotation<T, 3>::Identity), Penetration, ClosestA, ClosestB, Normal, 0, InitialDir)));
 			EXPECT_NEAR(Penetration, 0.125, Eps);
 			EXPECT_VECTOR_NEAR(Normal, TVec3<T>(0, 0, -1), Eps);
 			EXPECT_VECTOR_NEAR(ClosestA, TVec3<T>(3.25, 0, 0), Eps);
@@ -1450,11 +1466,11 @@ namespace ChaosTest
 			T Penetration;
 			TVec3<T> ClosestA,ClosestB,Normal;
 			const TVec3<T> Offset ={162.072754,-178.514679,-102.071632};
-			EXPECT_TRUE(GJKPenetration<T>(A,B,BToATM,Penetration,ClosestA,ClosestB,Normal,0,Offset,0));
+			EXPECT_TRUE((GJKPenetration<false, T>(A,B,BToATM,Penetration,ClosestA,ClosestB,Normal,0,Offset,0)));
 
 			const TRigidTransform<T,3> NewAToBTM (BToATM.GetTranslation() + (0.01 + Penetration) * Normal,BToATM.GetRotation());;
 
-			EXPECT_FALSE(GJKPenetration<T>(A,B,NewAToBTM,Penetration,ClosestA,ClosestB,Normal,0,Offset,0));
+			EXPECT_FALSE((GJKPenetration<false, T>(A,B,NewAToBTM,Penetration,ClosestA,ClosestB,Normal,0,Offset,0)));
 
 		}
 		
