Fix headless chaos compiler and missing OutTime

diff --git a/Engine/Source/Programs/HeadlessChaos/Private/HeadlessChaosTestGJK.cpp b/Engine/Source/Programs/HeadlessChaos/Private/HeadlessChaosTestGJK.cpp
--- a/Engine/Source/Programs/HeadlessChaos/Private/HeadlessChaosTestGJK.cpp
+++ b/Engine/Source/Programs/HeadlessChaos/Private/HeadlessChaosTestGJK.cpp
@@ -1489,10 +1489,50 @@ namespace ChaosTest
 
 			T Penetration;
 			TVec3<T> ClosestA,ClosestB,Normal;
-			EXPECT_TRUE(GJKPenetration<T>(A,B,BToATM, Penetration, ClosestA, ClosestB, Normal, 0.0,TVec3<T>(0,0,23.4092140)));
+			EXPECT_TRUE((GJKPenetration<false, T>(A,B,BToATM, Penetration, ClosestA, ClosestB, Normal, 0.0,TVec3<T>(0,0,23.4092140))));
 			EXPECT_FLOAT_EQ(Normal.Z,0);
 			EXPECT_FLOAT_EQ(Penetration,A.GetRadius() + B.GetRadius());
 		}
+
+#if 0
+		{
+			//capsule vs triangle as we make the sweep longer the world space point of impact should stay the same
+			TParticles<T,3> ConvexParticles;
+			ConvexParticles.AddParticles(3);
+
+			ConvexParticles.X(0) ={7400.00000, 12600.0000, 206.248123};
+			ConvexParticles.X(1) ={7500.00000, 12600.0000, 199.994904};
+			ConvexParticles.X(2) ={7500.00000, 12700.0000, 189.837433};
+			
+			TUniquePtr<FConvex> UniqueConvex = MakeUnique<FConvex>(ConvexParticles);
+			TSerializablePtr<FConvex> AConv(UniqueConvex);
+			const TImplicitObjectScaled<FConvex> A(AConv,TVec3<T>(1.0,1.0,1.0));
+
+			const TVec3<T> Pt0(0.0,0.0,-29.6999969);
+			TVec3<T> Pt1 = Pt0;
+			Pt1 += (TVec3<T>(0.0,0.0,1.0) * 59.3999939);
+
+			const TCapsule<T> B(Pt0,Pt1,42.0);
+
+			const TRigidTransform<T,3> BToATM(TVec3<T>(7475.74512, 12603.9082, 377.767120),TRotation<T,3>::FromElements(0,0,0,1));
+			const TVec3<T> LocalDir(0,0,-1);
+			const T Length = 49.9061584;
+			const TVec3<T> SearchDir(1,0,0);
+
+			T Time;
+			TVec3<T> Position,Normal;
+			EXPECT_TRUE(GJKRaycast2<T>(A,B,BToATM,LocalDir,Length,Time,Position,Normal,0,true,SearchDir,0));
+
+			const TRigidTransform<T,3> BToATM2(TVec3<T>(7475.74512, 12603.9082, 377.767120+100),TRotation<T,3>::FromElements(0,0,0,1));
+
+			T Time2;
+			TVec3<T> Position2,Normal2;
+			EXPECT_TRUE(GJKRaycast2<T>(A,B,BToATM2,LocalDir,Length+100,Time2,Position2,Normal2,0,true,SearchDir,0));
+			EXPECT_FLOAT_EQ(Time+100,Time2);
+			EXPECT_VECTOR_NEAR(Normal,Normal2,1e-4);
+			EXPECT_VECTOR_NEAR(Position,Position2,1e-4);
+		}
+#endif
 		
 	}
 
diff --git a/Engine/Source/Runtime/Experimental/Chaos/Public/Chaos/ImplicitObjectScaled.h b/Engine/Source/Runtime/Experimental/Chaos/Public/Chaos/ImplicitObjectScaled.h
--- a/Engine/Source/Runtime/Experimental/Chaos/Public/Chaos/ImplicitObjectScaled.h
+++ b/Engine/Source/Runtime/Experimental/Chaos/Public/Chaos/ImplicitObjectScaled.h
@@ -332,6 +332,7 @@ public:
 				{
 					OutPosition = MScale * UnscaledPosition;
 					OutNormal = (MInvScale * UnscaledNormal).GetSafeNormal();
+					OutTime = NewTime;
 					return true;
 				}
 			}
