Heightfield simple unit test

diff --git a/Engine/Source/Programs/HeadlessChaos/Private/HeadlessChaosTestHeightfield.cpp b/Engine/Source/Programs/HeadlessChaos/Private/HeadlessChaosTestHeightfield.cpp
new file mode 100644
--- /dev/null
+++ b/Engine/Source/Programs/HeadlessChaos/Private/HeadlessChaosTestHeightfield.cpp
@@ -0,0 +1,93 @@
+// Copyright 1998-2019 Epic Games, Inc. All Rights Reserved.
+
+#include "HeadlessChaosTestConstraints.h"
+
+#include "HeadlessChaos.h"
+#include "HeadlessChaosTestUtility.h"
+#include "Modules/ModuleManager.h"
+#include "Chaos/HeightField.h"
+
+namespace ChaosTest {
+
+	using namespace Chaos;
+
+	template<typename T>
+	void Raycast()
+	{
+		const uint32 Columns = 10;
+		const uint32 Rows = 10;
+		TArray<T> Heights;
+		Heights.AddZeroed(Rows * Columns);
+		const T CountToWorldScale = 1;
+
+		T Count = 0;
+		for(uint32 Row = 0; Row < Rows; ++Row)
+		{
+			for(uint32 Col = 0; Col < Columns; ++Col)
+			{
+				Heights[Row * Columns + Col] = CountToWorldScale * Count++;
+			}
+		}
+
+		auto TestHelper = [&](const TVec3<T>& Scale)
+		{
+			TArray<T> HeightsCopy = Heights;
+			THeightField<T> Heightfield(MoveTemp(HeightsCopy),TArray<uint8>(),Rows,Columns,Scale);
+			const auto& Bounds = Heightfield.BoundingBox();	//Current API forces us to do this to cache the bounds
+
+			//test straight down raycast
+			Count = 0;
+			T TOI;
+			TVec3<T> Position,Normal;
+			int32 FaceIdx;
+
+			//Compute expected normal
+			const TVec3<T> A(0,0,0);
+			const TVec3<T> B(Scale[0],0,CountToWorldScale * Scale[2]);
+			const TVec3<T> C(0,Scale[1],Columns*CountToWorldScale * Scale[2]);
+			const TVec3<T> ExpectedNormal = TVec3<T>::CrossProduct((B-A),(C-A)).GetUnsafeNormal();
+
+			int32 ExpectedFaceIdx = 0;
+			for(uint32 Row = 0; Row < Rows; ++Row)
+			{
+				for(uint32 Col = 0; Col < Columns; ++Col)
+				{
+					const TVec3<T> Start(Col*Scale[0],Row * Scale[1],1000*Scale[2]);
+					EXPECT_TRUE(Heightfield.Raycast(Start,TVec3<T>(0,0,-1),2000*Scale[2],0,TOI,Position,Normal,FaceIdx));
+					EXPECT_NEAR(TOI,(1000 - Heights[Row*Columns+Col])*Scale[2],1e-2);
+					EXPECT_VECTOR_NEAR(Position,TVec3<T>(Col*Scale[0],Row * Scale[1],Heights[Row*Columns+Col] * Scale[2]),1e-2);
+					EXPECT_VECTOR_NEAR(Normal,ExpectedNormal,1e-2);
+
+					//offset in from border ever so slightly to get a clear face
+					const bool bResult = Heightfield.Raycast(Start + TVec3<T>(0.2 * Scale[0],0.1 * Scale[1],0),TVec3<T>(0,0,-1),2000*Scale[2],0,TOI,Position,Normal,FaceIdx);
+					if(Col + 1 == Columns || Row + 1 == Rows)
+					{
+						EXPECT_FALSE(bResult);	//went past edge so no hit
+						//Went past column so do not increment expected face idx
+					} else
+					{
+						EXPECT_TRUE(bResult);
+						EXPECT_EQ(FaceIdx,ExpectedFaceIdx);	//each quad has two triangles, so for each column we pass two faces
+
+						ExpectedFaceIdx += 2;	//We hit the first triangle in the quad. Since we are going 1 quad at a time we skip 2
+					}
+				}
+			}
+		};
+
+		TestHelper(TVec3<T>(1));
+		TestHelper(TVec3<T>(1,1,3));
+		TestHelper(TVec3<T>(1,1,.3));
+		TestHelper(TVec3<T>(3,1,.3));
+		TestHelper(TVec3<T>(2,.1,.3));
+		
+	}
+
+
+	TEST(ChaosTests, Heightfield)
+	{
+		ChaosTest::Raycast<float>();
+		SUCCEED();
+	}
+
+}
\ No newline at end of file
