Added X and Y traversal tests for hf

diff --git a/Engine/Source/Programs/HeadlessChaos/Private/HeadlessChaosTestHeightfield.cpp b/Engine/Source/Programs/HeadlessChaos/Private/HeadlessChaosTestHeightfield.cpp
--- a/Engine/Source/Programs/HeadlessChaos/Private/HeadlessChaosTestHeightfield.cpp
+++ b/Engine/Source/Programs/HeadlessChaos/Private/HeadlessChaosTestHeightfield.cpp
@@ -14,22 +14,32 @@ namespace ChaosTest {
 	template<typename T>
 	void Raycast()
 	{
-		const uint32 Columns = 10;
-		const uint32 Rows = 10;
+		const int32 Columns = 10;
+		const int32 Rows = 10;
 		TArray<T> Heights;
 		Heights.AddZeroed(Rows * Columns);
 		const T CountToWorldScale = 1;
 
 		T Count = 0;
-		for(uint32 Row = 0; Row < Rows; ++Row)
+		for(int32 Row = 0; Row < Rows; ++Row)
 		{
-			for(uint32 Col = 0; Col < Columns; ++Col)
+			for(int32 Col = 0; Col < Columns; ++Col)
 			{
 				Heights[Row * Columns + Col] = CountToWorldScale * Count++;
 			}
 		}
 
-		auto TestHelper = [&](const TVec3<T>& Scale)
+		auto ComputeExpectedNormal = [&](const TVec3<T>& Scale)
+		{
+			//Compute expected normal
+			const TVec3<T> A(0,0,0);
+			const TVec3<T> B(Scale[0],0,CountToWorldScale * Scale[2]);
+			const TVec3<T> C(0,Scale[1],Columns*CountToWorldScale * Scale[2]);
+			const TVec3<T> ExpectedNormal = TVec3<T>::CrossProduct((B-A),(C-A)).GetUnsafeNormal();
+			return ExpectedNormal;
+		};
+
+		auto AlongZTest = [&](const TVec3<T>& Scale)
 		{
 			TArray<T> HeightsCopy = Heights;
 			THeightField<T> Heightfield(MoveTemp(HeightsCopy),TArray<uint8>(),Rows,Columns,Scale);
@@ -41,16 +51,12 @@ namespace ChaosTest {
 			TVec3<T> Position,Normal;
 			int32 FaceIdx;
 
-			//Compute expected normal
-			const TVec3<T> A(0,0,0);
-			const TVec3<T> B(Scale[0],0,CountToWorldScale * Scale[2]);
-			const TVec3<T> C(0,Scale[1],Columns*CountToWorldScale * Scale[2]);
-			const TVec3<T> ExpectedNormal = TVec3<T>::CrossProduct((B-A),(C-A)).GetUnsafeNormal();
+			const TVec3<T> ExpectedNormal = ComputeExpectedNormal(Scale);
 
 			int32 ExpectedFaceIdx = 0;
-			for(uint32 Row = 0; Row < Rows; ++Row)
+			for(int32 Row = 0; Row < Rows; ++Row)
 			{
-				for(uint32 Col = 0; Col < Columns; ++Col)
+				for(int32 Col = 0; Col < Columns; ++Col)
 				{
 					const TVec3<T> Start(Col*Scale[0],Row * Scale[1],1000*Scale[2]);
 					EXPECT_TRUE(Heightfield.Raycast(Start,TVec3<T>(0,0,-1),2000*Scale[2],0,TOI,Position,Normal,FaceIdx));
@@ -75,11 +81,96 @@ namespace ChaosTest {
 			}
 		};
 
-		TestHelper(TVec3<T>(1));
-		TestHelper(TVec3<T>(1,1,3));
-		TestHelper(TVec3<T>(1,1,.3));
-		TestHelper(TVec3<T>(3,1,.3));
-		TestHelper(TVec3<T>(2,.1,.3));
+		AlongZTest(TVec3<T>(1));
+		AlongZTest(TVec3<T>(1,1,3));
+		AlongZTest(TVec3<T>(1,1,.3));
+		AlongZTest(TVec3<T>(3,1,.3));
+		AlongZTest(TVec3<T>(2,.1,.3));
+
+		auto AlongXTest = [&](const TVec3<T>& Scale)
+		{
+			TArray<T> HeightsCopy = Heights;
+			THeightField<T> Heightfield(MoveTemp(HeightsCopy),TArray<uint8>(),Rows,Columns,Scale);
+			const auto& Bounds = Heightfield.BoundingBox();	//Current API forces us to do this to cache the bounds
+
+			//test along x axis
+			Count = 0;
+			T TOI;
+			TVec3<T> Position,Normal;
+			int32 FaceIdx;
+
+			const TVec3<T> ExpectedNormal = ComputeExpectedNormal(Scale);
+
+			//move from left to right and raycast down the x-axis. The Row idx indicates which cell we expect to hit
+			for(int32 Row = 0; Row < Rows; ++Row)
+			{
+				for(int32 Col = 0; Col < Columns; ++Col)
+				{
+					const TVec3<T> Start(-Scale[0], Row * Scale[1], Heights[Row*Columns + Col] * Scale[2] + 0.01 * Scale[2]);
+					const bool bResult = Heightfield.Raycast(Start,TVec3<T>(1,0,0),2000*Scale[0],0,TOI,Position,Normal,FaceIdx);
+					if(Col + 1 == Columns)
+					{
+						EXPECT_FALSE(bResult);
+						//No more columns so we shot over the final edge
+					}
+					else
+					{
+						EXPECT_TRUE(bResult);
+						EXPECT_NEAR(TOI,(Scale[0] * (1 + Col)),1e-1);
+						EXPECT_VECTOR_NEAR(Position,(Start + TVec3<T>{TOI,0,0}),1e-2);
+						EXPECT_VECTOR_NEAR(Normal,ExpectedNormal,1e-1);
+					}
+				}
+			}
+		};
+
+		AlongXTest(TVec3<T>(1));
+		AlongXTest(TVec3<T>(1,1,3));
+		AlongXTest(TVec3<T>(1,1,.3));
+		AlongXTest(TVec3<T>(3,1,.3));
+		AlongXTest(TVec3<T>(2,.1,.3));
+
+		auto AlongYTest = [&](const TVec3<T>& Scale)
+		{
+			TArray<T> HeightsCopy = Heights;
+			THeightField<T> Heightfield(MoveTemp(HeightsCopy),TArray<uint8>(),Rows,Columns,Scale);
+			const auto& Bounds = Heightfield.BoundingBox();	//Current API forces us to do this to cache the bounds
+
+			//test along y axis
+			Count = 0;
+			T TOI;
+			TVec3<T> Position,Normal;
+			int32 FaceIdx;
+
+			const TVec3<T> ExpectedNormal = ComputeExpectedNormal(Scale);
+
+			//move from left to right and raycast down the x-axis. The Row idx indicates which cell we expect to hit
+			for(int32 Row = 0; Row < Rows; ++Row)
+			{
+				for(int32 Col = 0; Col < Columns; ++Col)
+				{
+					const TVec3<T> Start(Col * Scale[0], -Scale[1],Heights[Row*Columns + Col] * Scale[2] + 0.01 * Scale[2]);
+					const bool bResult = Heightfield.Raycast(Start,TVec3<T>(0,1,0),2000*Scale[0],0,TOI,Position,Normal,FaceIdx);
+					if(Row + 1 == Rows)
+					{
+						EXPECT_FALSE(bResult);
+						//No more columns so we shot over the final edge
+					} else
+					{
+						EXPECT_TRUE(bResult);
+						EXPECT_NEAR(TOI,(Scale[1] * (1 + Row)),1e-1);
+						EXPECT_VECTOR_NEAR(Position,(Start + TVec3<T>{0,TOI,0}),1e-2);
+						EXPECT_VECTOR_NEAR(Normal,ExpectedNormal,1e-1);
+					}
+				}
+			}
+		};
+
+		AlongYTest(TVec3<T>(1));
+		AlongYTest(TVec3<T>(1,1,3));
+		AlongYTest(TVec3<T>(1,1,.3));
+		AlongYTest(TVec3<T>(3,1,.3));
+		AlongYTest(TVec3<T>(2,.1,.3));
 		
 	}
 
