Fix cis

diff --git a/Engine/Source/Programs/HeadlessChaos/Private/HeadlessChaosTestHeightfield.cpp b/Engine/Source/Programs/HeadlessChaos/Private/HeadlessChaosTestHeightfield.cpp
--- a/Engine/Source/Programs/HeadlessChaos/Private/HeadlessChaosTestHeightfield.cpp
+++ b/Engine/Source/Programs/HeadlessChaos/Private/HeadlessChaosTestHeightfield.cpp
@@ -1,5 +1,3 @@
-// Copyright 1998-2019 Epic Games, Inc. All Rights Reserved.
-
 #include "HeadlessChaosTestConstraints.h"
 
 #include "HeadlessChaos.h"
