Replace copyright header with dates removed as new style
#rnx

diff --git a/Engine/Source/Programs/HeadlessChaos/Private/HeadlessChaosTestHeightfield.cpp b/Engine/Source/Programs/HeadlessChaos/Private/HeadlessChaosTestHeightfield.cpp
--- a/Engine/Source/Programs/HeadlessChaos/Private/HeadlessChaosTestHeightfield.cpp
+++ b/Engine/Source/Programs/HeadlessChaos/Private/HeadlessChaosTestHeightfield.cpp
@@ -1,3 +1,4 @@
+// Copyright Epic Games, Inc. All Rights Reserved.
 #include "HeadlessChaosTestConstraints.h"
 
 #include "HeadlessChaos.h"
