Add false negative Sphere Sweep vs Triangle unit test for easy repro.
This repros behaviour causing traps to be placeable inside hills

diff --git a/Engine/Source/Programs/HeadlessChaos/Private/HeadlessChaosTestEPA.cpp b/Engine/Source/Programs/HeadlessChaos/Private/HeadlessChaosTestEPA.cpp
--- a/Engine/Source/Programs/HeadlessChaos/Private/HeadlessChaosTestEPA.cpp
+++ b/Engine/Source/Programs/HeadlessChaos/Private/HeadlessChaosTestEPA.cpp
@@ -10,6 +10,7 @@
 #include "Chaos/GJK.h"
 #include "Chaos/Convex.h"
 #include "Chaos/ImplicitObjectScaled.h"
+#include "Chaos/Triangle.h"
 
 namespace ChaosTest
 {
@@ -635,6 +636,49 @@ namespace ChaosTest
 	// Currently broken EPA edge cases. As they are fixed move them to EPARealFailures_Fixed above so that we can ensure they don't break again.
 	GTEST_TEST(EPATests, EPARealFailures_Broken)
 	{
+		using namespace Chaos;
+		// Sphere sweep against triangle, fails when it should hit. Raycast added as well for verification purposes.
+		{
+			const TTriangle<FReal> Triangle({ 0.000000000, 0.000000000, 0.000000000 }, { 128.000000, 0.000000000, -114.064575 }, { 128.000000, 128.000000, 2.35327148 });
+			const TSphere<FReal, 3> Sphere({ 0.0, 0.0, 0.0 }, 4);
+			const TRigidTransform<FReal, 3> Transform({ 174.592773, -161.781250, -68.0469971 }, FQuat::Identity);
+			const FVec3 Dir({ -0.406315684, 0.913382649, -0.0252906363 });
+			const FReal Length = 430.961548;
+			const FReal Thickness = 0.0f;
+			const bool bComputeMTD = true;
+
+			FReal OutTime = -1.0f;
+			FVec3 OutPosition;
+			FVec3 OutNormal;
+
+			bool bSweepResult = GJKRaycast2(Triangle, Sphere, Transform, Dir, Length, OutTime, OutPosition, OutNormal, Thickness, bComputeMTD);
+
+
+			// Do a raycast w/ same inputs instead of sweep against triangle to verify sweep should be a hit.
+			const FVec3 TriNormal = Triangle.GetNormal();
+			const TPlane<FReal, 3> TriPlane{ Triangle[0], TriNormal };
+			FVec3 RaycastPosition;
+			FVec3 RaycastNormal;
+			FReal Time;
+
+			int32 DummyFaceIndex;
+
+			bool bTriangleIntersects = false;
+			if (TriPlane.Raycast(Transform.GetTranslation(), Dir, Length, Thickness, Time, RaycastPosition, RaycastNormal, DummyFaceIndex))
+			{
+				FVec3 IntersectionPosition = RaycastPosition;
+				FVec3 IntersectionNormal = RaycastNormal;
+
+				const FVec3 ClosestPtOnTri = FindClosestPointOnTriangle(RaycastPosition, Triangle[0], Triangle[1], Triangle[2], RaycastPosition);	//We know Position is on the triangle plane
+				const FReal DistToTriangle2 = (RaycastPosition - ClosestPtOnTri).SizeSquared();
+				bTriangleIntersects = DistToTriangle2 <= SMALL_NUMBER;	//raycast gave us the intersection point so sphere radius is already accounted for
+			}
+
+			//EXPECT_EQ(bTriangleIntersects, bSweepResult); // uncomment to demonstrate failure.
+		}
 
 	}
+
+
+
 }
\ No newline at end of file
