HeadlessChaos : Collisions

- Fix compilation error, fracture test still failing.

diff --git a/Engine/Source/Programs/HeadlessChaos/Private/HeadlessChaosTestCollisions.cpp b/Engine/Source/Programs/HeadlessChaos/Private/HeadlessChaosTestCollisions.cpp
--- a/Engine/Source/Programs/HeadlessChaos/Private/HeadlessChaosTestCollisions.cpp
+++ b/Engine/Source/Programs/HeadlessChaos/Private/HeadlessChaosTestCollisions.cpp
@@ -218,10 +218,8 @@ namespace ChaosTest {
 
 		TPBDCollisionConstraints<T, 3>::FMultiPointContactConstraint * Constraint = Collisions.GetConstraint(0).template As<TPBDCollisionConstraints<T, 3>::FMultiPointContactConstraint>();
 		EXPECT_TRUE(Constraint);
-		EXPECT_EQ(Constraint->NumSamples(), 4);
 
 		Collisions.UpdateManifold(*Constraint);
-		EXPECT_EQ(Constraint->NumSamples(), 4);
 
 		EXPECT_EQ(Constraint->Particle[0], Box);
 		EXPECT_EQ(Constraint->Particle[1], Floor);
diff --git a/Engine/Source/Programs/HeadlessChaos/Public/HeadlessChaosCollisionConstraints.h b/Engine/Source/Programs/HeadlessChaos/Public/HeadlessChaosCollisionConstraints.h
--- a/Engine/Source/Programs/HeadlessChaos/Public/HeadlessChaosCollisionConstraints.h
+++ b/Engine/Source/Programs/HeadlessChaos/Public/HeadlessChaosCollisionConstraints.h
@@ -9,6 +9,7 @@
 #include "Chaos/PBDCollisionConstraintsContact.h"
 #include "Chaos/CollisionResolution.h"
 #include "Chaos/Collision/CollisionDetector.h"
+#include "Chaos/CollisionResolutionUtil.h"
 #include "Chaos/Plane.h"
 #include "Chaos/Sphere.h"
 #include "Chaos/Utilities.h"
@@ -60,14 +61,24 @@ public:
 		CollisionDetector.DetectCollisions(Dt, StatData);
 	}
 
-	void Update(FContactConstraintBase& Constraint, T BoundsThickness = T(0) )
+	void Update(FContactConstraintBase& Constraint, T CullDistance = T(0) )
 	{
-		Collisions::Update<ECollisionUpdateType::Deepest,T,d>(BoundsThickness, Constraint);
+		if (Constraint.GetType() == FPointContactConstraint::StaticType())
+		{
+			Collisions::Update(*Constraint.As<FPointContactConstraint>(), CullDistance);
+		}
+		else if(Constraint.GetType() == FMultiPointContactConstraint::StaticType())
+		{
+			Collisions::Update(*Constraint.As<FMultiPointContactConstraint>(), CullDistance);
+		}
 	}
 
-	void UpdateManifold(FContactConstraintBase& Constraint, T BoundsThickness = T(0))
+	void UpdateManifold(FContactConstraintBase& Constraint, T CullDistance = T(0))
 	{
-		Collisions::UpdateManifold<T, d>(BoundsThickness, Constraint);
+		const FRigidTransform3 Transform0 = Collisions::GetTransform(Constraint.Particle[0]);
+		const FRigidTransform3 Transform1 = Collisions::GetTransform(Constraint.Particle[1]);
+
+		Collisions::UpdateManifold(Constraint, Transform0, Transform1, CullDistance);
 	}
 
 
