Fix raycast traversal in heightfield. Added more unit tests that exposed previous issues. Fixes lots of weird issues like raycasts missing

diff --git a/Engine/Source/Programs/HeadlessChaos/Private/HeadlessChaosTestHeightfield.cpp b/Engine/Source/Programs/HeadlessChaos/Private/HeadlessChaosTestHeightfield.cpp
--- a/Engine/Source/Programs/HeadlessChaos/Private/HeadlessChaosTestHeightfield.cpp
+++ b/Engine/Source/Programs/HeadlessChaos/Private/HeadlessChaosTestHeightfield.cpp
@@ -15,161 +15,272 @@ namespace ChaosTest {
 	{
 		const int32 Columns = 10;
 		const int32 Rows = 10;
-		TArray<T> Heights;
-		Heights.AddZeroed(Rows * Columns);
 		const T CountToWorldScale = 1;
 
-		T Count = 0;
-		for(int32 Row = 0; Row < Rows; ++Row)
 		{
-			for(int32 Col = 0; Col < Columns; ++Col)
+
+			TArray<T> Heights;
+			Heights.AddZeroed(Rows * Columns);
+
+			T Count = 0;
+			for(int32 Row = 0; Row < Rows; ++Row)
 			{
-				Heights[Row * Columns + Col] = CountToWorldScale * Count++;
+				for(int32 Col = 0; Col < Columns; ++Col)
+				{
+					Heights[Row * Columns + Col] = CountToWorldScale * Count++;
+				}
 			}
-		}
 
-		auto ComputeExpectedNormal = [&](const TVec3<T>& Scale)
-		{
-			//Compute expected normal
-			const TVec3<T> A(0,0,0);
-			const TVec3<T> B(Scale[0],0,CountToWorldScale * Scale[2]);
-			const TVec3<T> C(0,Scale[1],Columns*CountToWorldScale * Scale[2]);
-			const TVec3<T> ExpectedNormal = TVec3<T>::CrossProduct((B-A),(C-A)).GetUnsafeNormal();
-			return ExpectedNormal;
-		};
-
-		auto AlongZTest = [&](const TVec3<T>& Scale)
-		{
-			TArray<T> HeightsCopy = Heights;
-			THeightField<T> Heightfield(MoveTemp(HeightsCopy),TArray<uint8>(),Rows,Columns,Scale);
-			const auto& Bounds = Heightfield.BoundingBox();	//Current API forces us to do this to cache the bounds
+			auto ComputeExpectedNormal = [&](const TVec3<T>& Scale)
+			{
+				//Compute expected normal
+				const TVec3<T> A(0,0,0);
+				const TVec3<T> B(Scale[0],0,CountToWorldScale * Scale[2]);
+				const TVec3<T> C(0,Scale[1],Columns*CountToWorldScale * Scale[2]);
+				const TVec3<T> ExpectedNormal = TVec3<T>::CrossProduct((B-A),(C-A)).GetUnsafeNormal();
+				return ExpectedNormal;
+			};
 
-			//test straight down raycast
-			Count = 0;
-			T TOI;
-			TVec3<T> Position,Normal;
-			int32 FaceIdx;
+			auto AlongZTest = [&](const TVec3<T>& Scale)
+			{
+				TArray<T> HeightsCopy = Heights;
+				THeightField<T> Heightfield(MoveTemp(HeightsCopy),TArray<uint8>(),Rows,Columns,Scale);
+				const auto& Bounds = Heightfield.BoundingBox();	//Current API forces us to do this to cache the bounds
 
-			const TVec3<T> ExpectedNormal = ComputeExpectedNormal(Scale);
+				//test straight down raycast
+				Count = 0;
+				T TOI;
+				TVec3<T> Position,Normal;
+				int32 FaceIdx;
 
-			int32 ExpectedFaceIdx = 0;
-			for(int32 Row = 0; Row < Rows; ++Row)
-			{
-				for(int32 Col = 0; Col < Columns; ++Col)
+				const TVec3<T> ExpectedNormal = ComputeExpectedNormal(Scale);
+
+				int32 ExpectedFaceIdx = 0;
+				for(int32 Row = 0; Row < Rows; ++Row)
 				{
-					const TVec3<T> Start(Col*Scale[0],Row * Scale[1],1000*Scale[2]);
-					EXPECT_TRUE(Heightfield.Raycast(Start,TVec3<T>(0,0,-1),2000*Scale[2],0,TOI,Position,Normal,FaceIdx));
-					EXPECT_NEAR(TOI,(1000 - Heights[Row*Columns+Col])*Scale[2],1e-2);
-					EXPECT_VECTOR_NEAR(Position,TVec3<T>(Col*Scale[0],Row * Scale[1],Heights[Row*Columns+Col] * Scale[2]),1e-2);
-					EXPECT_VECTOR_NEAR(Normal,ExpectedNormal,1e-2);
-
-					//offset in from border ever so slightly to get a clear face
-					const bool bResult = Heightfield.Raycast(Start + TVec3<T>(0.2 * Scale[0],0.1 * Scale[1],0),TVec3<T>(0,0,-1),2000*Scale[2],0,TOI,Position,Normal,FaceIdx);
-					if(Col + 1 == Columns || Row + 1 == Rows)
-					{
-						EXPECT_FALSE(bResult);	//went past edge so no hit
-						//Went past column so do not increment expected face idx
-					} else
+					for(int32 Col = 0; Col < Columns; ++Col)
 					{
-						EXPECT_TRUE(bResult);
-						EXPECT_EQ(FaceIdx,ExpectedFaceIdx);	//each quad has two triangles, so for each column we pass two faces
+						const TVec3<T> Start(Col*Scale[0],Row * Scale[1],1000*Scale[2]);
+						EXPECT_TRUE(Heightfield.Raycast(Start,TVec3<T>(0,0,-1),2000*Scale[2],0,TOI,Position,Normal,FaceIdx));
+						EXPECT_NEAR(TOI,(1000 - Heights[Row*Columns+Col])*Scale[2],1e-2);
+						EXPECT_VECTOR_NEAR(Position,TVec3<T>(Col*Scale[0],Row * Scale[1],Heights[Row*Columns+Col] * Scale[2]),1e-2);
+						EXPECT_VECTOR_NEAR(Normal,ExpectedNormal,1e-2);
 
-						ExpectedFaceIdx += 2;	//We hit the first triangle in the quad. Since we are going 1 quad at a time we skip 2
+						//offset in from border ever so slightly to get a clear face
+						const bool bResult = Heightfield.Raycast(Start + TVec3<T>(0.2 * Scale[0],0.1 * Scale[1],0),TVec3<T>(0,0,-1),2000*Scale[2],0,TOI,Position,Normal,FaceIdx);
+						if(Col + 1 == Columns || Row + 1 == Rows)
+						{
+							EXPECT_FALSE(bResult);	//went past edge so no hit
+							//Went past column so do not increment expected face idx
+						} else
+						{
+							EXPECT_TRUE(bResult);
+							EXPECT_EQ(FaceIdx,ExpectedFaceIdx);	//each quad has two triangles, so for each column we pass two faces
+
+							ExpectedFaceIdx += 2;	//We hit the first triangle in the quad. Since we are going 1 quad at a time we skip 2
+						}
 					}
 				}
-			}
-		};
+			};
 
-		AlongZTest(TVec3<T>(1));
-		AlongZTest(TVec3<T>(1,1,3));
-		AlongZTest(TVec3<T>(1,1,.3));
-		AlongZTest(TVec3<T>(3,1,.3));
-		AlongZTest(TVec3<T>(2,.1,.3));
+			AlongZTest(TVec3<T>(1));
+			AlongZTest(TVec3<T>(1,1,3));
+			AlongZTest(TVec3<T>(1,1,.3));
+			AlongZTest(TVec3<T>(3,1,.3));
+			AlongZTest(TVec3<T>(2,.1,.3));
 
-		auto AlongXTest = [&](const TVec3<T>& Scale)
-		{
-			TArray<T> HeightsCopy = Heights;
-			THeightField<T> Heightfield(MoveTemp(HeightsCopy),TArray<uint8>(),Rows,Columns,Scale);
-			const auto& Bounds = Heightfield.BoundingBox();	//Current API forces us to do this to cache the bounds
+			auto AlongXTest = [&](const TVec3<T>& Scale)
+			{
+				TArray<T> HeightsCopy = Heights;
+				THeightField<T> Heightfield(MoveTemp(HeightsCopy),TArray<uint8>(),Rows,Columns,Scale);
+				const auto& Bounds = Heightfield.BoundingBox();	//Current API forces us to do this to cache the bounds
 
-			//test along x axis
-			Count = 0;
-			T TOI;
-			TVec3<T> Position,Normal;
-			int32 FaceIdx;
+				//test along x axis
+				Count = 0;
+				T TOI;
+				TVec3<T> Position,Normal;
+				int32 FaceIdx;
 
-			const TVec3<T> ExpectedNormal = ComputeExpectedNormal(Scale);
+				const TVec3<T> ExpectedNormal = ComputeExpectedNormal(Scale);
 
-			//move from left to right and raycast down the x-axis. The Row idx indicates which cell we expect to hit
-			for(int32 Row = 0; Row < Rows; ++Row)
-			{
-				for(int32 Col = 0; Col < Columns; ++Col)
+				//move from left to right and raycast down the x-axis. The Row idx indicates which cell we expect to hit
+				for(int32 Row = 0; Row < Rows; ++Row)
 				{
-					const TVec3<T> Start(-Scale[0], Row * Scale[1], Heights[Row*Columns + Col] * Scale[2] + 0.01 * Scale[2]);
-					const bool bResult = Heightfield.Raycast(Start,TVec3<T>(1,0,0),2000*Scale[0],0,TOI,Position,Normal,FaceIdx);
-					if(Col + 1 == Columns)
+					for(int32 Col = 0; Col < Columns; ++Col)
 					{
-						EXPECT_FALSE(bResult);
-						//No more columns so we shot over the final edge
+						const TVec3<T> Start(-Scale[0],Row * Scale[1],Heights[Row*Columns + Col] * Scale[2] + 0.01 * Scale[2]);
+						const bool bResult = Heightfield.Raycast(Start,TVec3<T>(1,0,0),2000*Scale[0],0,TOI,Position,Normal,FaceIdx);
+						if(Col + 1 == Columns)
+						{
+							EXPECT_FALSE(bResult);
+							//No more columns so we shot over the final edge
+						} else
+						{
+							EXPECT_TRUE(bResult);
+							EXPECT_NEAR(TOI,(Scale[0] * (1 + Col)),1e-1);
+							EXPECT_VECTOR_NEAR(Position,(Start + TVec3<T>{TOI,0,0}),1e-2);
+							EXPECT_VECTOR_NEAR(Normal,ExpectedNormal,1e-1);
+						}
 					}
-					else
+				}
+			};
+
+			AlongXTest(TVec3<T>(1));
+			AlongXTest(TVec3<T>(1,1,3));
+			AlongXTest(TVec3<T>(1,1,.3));
+			AlongXTest(TVec3<T>(3,1,.3));
+			AlongXTest(TVec3<T>(2,.1,.3));
+
+			auto AlongYTest = [&](const TVec3<T>& Scale)
+			{
+				TArray<T> HeightsCopy = Heights;
+				THeightField<T> Heightfield(MoveTemp(HeightsCopy),TArray<uint8>(),Rows,Columns,Scale);
+				const auto& Bounds = Heightfield.BoundingBox();	//Current API forces us to do this to cache the bounds
+
+				//test along y axis
+				Count = 0;
+				T TOI;
+				TVec3<T> Position,Normal;
+				int32 FaceIdx;
+
+				const TVec3<T> ExpectedNormal = ComputeExpectedNormal(Scale);
+
+				for(int32 Row = 0; Row < Rows; ++Row)
+				{
+					for(int32 Col = 0; Col < Columns; ++Col)
 					{
-						EXPECT_TRUE(bResult);
-						EXPECT_NEAR(TOI,(Scale[0] * (1 + Col)),1e-1);
-						EXPECT_VECTOR_NEAR(Position,(Start + TVec3<T>{TOI,0,0}),1e-2);
-						EXPECT_VECTOR_NEAR(Normal,ExpectedNormal,1e-1);
+						const TVec3<T> Start(Col * Scale[0],-Scale[1],Heights[Row*Columns + Col] * Scale[2] + 0.01 * Scale[2]);
+						const bool bResult = Heightfield.Raycast(Start,TVec3<T>(0,1,0),2000*Scale[0],0,TOI,Position,Normal,FaceIdx);
+						if(Row + 1 == Rows)
+						{
+							EXPECT_FALSE(bResult);
+							//No more columns so we shot over the final edge
+						} else
+						{
+							EXPECT_TRUE(bResult);
+							EXPECT_NEAR(TOI,(Scale[1] * (1 + Row)),1e-1);
+							EXPECT_VECTOR_NEAR(Position,(Start + TVec3<T>{0,TOI,0}),1e-2);
+							EXPECT_VECTOR_NEAR(Normal,ExpectedNormal,1e-1);
+						}
 					}
 				}
-			}
-		};
+			};
 
-		AlongXTest(TVec3<T>(1));
-		AlongXTest(TVec3<T>(1,1,3));
-		AlongXTest(TVec3<T>(1,1,.3));
-		AlongXTest(TVec3<T>(3,1,.3));
-		AlongXTest(TVec3<T>(2,.1,.3));
+			AlongYTest(TVec3<T>(1));
+			AlongYTest(TVec3<T>(1,1,3));
+			AlongYTest(TVec3<T>(1,1,.3));
+			AlongYTest(TVec3<T>(3,1,.3));
+			AlongYTest(TVec3<T>(2,.1,.3));
+		}
 
-		auto AlongYTest = [&](const TVec3<T>& Scale)
 		{
-			TArray<T> HeightsCopy = Heights;
-			THeightField<T> Heightfield(MoveTemp(HeightsCopy),TArray<uint8>(),Rows,Columns,Scale);
-			const auto& Bounds = Heightfield.BoundingBox();	//Current API forces us to do this to cache the bounds
-
-			//test along y axis
-			Count = 0;
-			T TOI;
-			TVec3<T> Position,Normal;
-			int32 FaceIdx;
 
-			const TVec3<T> ExpectedNormal = ComputeExpectedNormal(Scale);
-
-			//move from left to right and raycast down the x-axis. The Row idx indicates which cell we expect to hit
+			//For diagonal test simply increase height on the y axis
+			TArray<T> Heights2;
+			Heights2.AddZeroed(Rows*Columns);
 			for(int32 Row = 0; Row < Rows; ++Row)
 			{
 				for(int32 Col = 0; Col < Columns; ++Col)
 				{
-					const TVec3<T> Start(Col * Scale[0], -Scale[1],Heights[Row*Columns + Col] * Scale[2] + 0.01 * Scale[2]);
-					const bool bResult = Heightfield.Raycast(Start,TVec3<T>(0,1,0),2000*Scale[0],0,TOI,Position,Normal,FaceIdx);
-					if(Row + 1 == Rows)
+					Heights2[Row * Columns + Col] = CountToWorldScale * Row;
+				}
+			}
+
+			auto ComputeExpectedNormal2 = [&](const TVec3<T>& Scale)
+			{
+				//Compute expected normal
+				const TVec3<T> A(0,0,0);
+				const TVec3<T> B(Scale[0],0,0);
+				const TVec3<T> C(0,Scale[1],CountToWorldScale * Scale[2]);
+				const TVec3<T> ExpectedNormal = TVec3<T>::CrossProduct((B-A),(C-A)).GetUnsafeNormal();
+				return ExpectedNormal;
+			};
+
+			auto AlongXYTest = [&](const TVec3<T>& Scale)
+			{
+				TArray<T> HeightsCopy = Heights2;
+				THeightField<T> Heightfield(MoveTemp(HeightsCopy),TArray<uint8>(),Rows,Columns,Scale);
+				const auto& Bounds = Heightfield.BoundingBox();	//Current API forces us to do this to cache the bounds
+
+				//test along x-y axis
+				T TOI;
+				TVec3<T> Position,Normal;
+				int32 FaceIdx;
+
+				const TVec3<T> ExpectedNormal = ComputeExpectedNormal2(Scale);
+
+				for(int32 Row = 0; Row < Rows; ++Row)
+				{
+					for(int32 Col = 0; Col < Columns; ++Col)
+					{
+						const TVec3<T> Start(Col * Scale[0],0,Heights2[Row*Columns + Col] * Scale[2] + 0.01 * Scale[2]);
+						const TVec3<T> Dir = TVec3<T>(1,1,0).GetUnsafeNormal();
+						const bool bResult = Heightfield.Raycast(Start,Dir,2000*Scale[0],0,TOI,Position,Normal,FaceIdx);
+
+						//As we increase the row, fewer columns will hit because the ray will exit the heightfield
+						const bool bShouldHit = Col + Row + 1 < Columns;
+						EXPECT_EQ(bShouldHit,bResult);
+						if(bResult)
+						{
+							EXPECT_VECTOR_NEAR(Normal,ExpectedNormal,1e-1);
+						}
+					}
+				}
+			};
+
+			AlongXYTest(TVec3<T>(1));
+
+
+			auto ToCellsTest = [&](const TVec3<T>& Scale)
+			{
+				//Pick cells and shoot ray at them
+				//This should always succeed because 0,0 is the lowest and n,n is the heighest
+				TArray<T> HeightsCopy = Heights2;
+				THeightField<T> Heightfield(MoveTemp(HeightsCopy),TArray<uint8>(),Rows,Columns,Scale);
+				const auto& Bounds = Heightfield.BoundingBox();	//Current API forces us to do this to cache the bounds
+
+				T TOI;
+				TVec3<T> Position,Normal;
+				int32 FaceIdx;
+
+				const TVec3<T> ExpectedNormal = ComputeExpectedNormal2(Scale);
+
+				const TVec3<T> Start(0,0,Heights2.Last() * Scale[2]);
+				for(int32 TargetIdx = 0; TargetIdx < Heights2.Num(); ++TargetIdx)
+				{
+					const int32 Col = TargetIdx % Columns;
+					const int32 Row = TargetIdx / Columns;
+					const TVec3<T> EndUnscaled(Col,Row,Heights2[TargetIdx]);
+					TVec3<T> End = EndUnscaled * Scale;
+					if(Col + 1 == Columns)
 					{
-						EXPECT_FALSE(bResult);
-						//No more columns so we shot over the final edge
-					} else
+						//pull back slightly to avoid precision issues at edge
+						End[0] -= 0.1 * Scale[0];
+					}
+
+					if(Row + 1 == Rows)
 					{
-						EXPECT_TRUE(bResult);
-						EXPECT_NEAR(TOI,(Scale[1] * (1 + Row)),1e-1);
-						EXPECT_VECTOR_NEAR(Position,(Start + TVec3<T>{0,TOI,0}),1e-2);
-						EXPECT_VECTOR_NEAR(Normal,ExpectedNormal,1e-1);
+						//pulling back row affects Z so just skip
+						continue;
 					}
+
+					const TVec3<T> Dir = (End - Start).GetUnsafeNormal();
+
+					const bool bResult = Heightfield.Raycast(Start,Dir,2000,0,TOI,Position,Normal,FaceIdx);
+					EXPECT_TRUE(bResult);
+					EXPECT_VECTOR_NEAR(Normal,ExpectedNormal,1e-1);
+					EXPECT_VECTOR_NEAR(Position,End,1e-1);
 				}
-			}
-		};
+			};
+
+			ToCellsTest(TVec3<T>(1));
+			ToCellsTest(TVec3<T>(1,1,10));
+			ToCellsTest(TVec3<T>(1,1,.1));
+			ToCellsTest(TVec3<T>(3,1,.1));
+			ToCellsTest(TVec3<T>(.3,1,.1));
 
-		AlongYTest(TVec3<T>(1));
-		AlongYTest(TVec3<T>(1,1,3));
-		AlongYTest(TVec3<T>(1,1,.3));
-		AlongYTest(TVec3<T>(3,1,.3));
-		AlongYTest(TVec3<T>(2,.1,.3));
+
+		}
 		
 	}
 
diff --git a/Engine/Source/Runtime/Experimental/Chaos/Private/Chaos/HeightField.cpp b/Engine/Source/Runtime/Experimental/Chaos/Private/Chaos/HeightField.cpp
--- a/Engine/Source/Runtime/Experimental/Chaos/Private/Chaos/HeightField.cpp
+++ b/Engine/Source/Runtime/Experimental/Chaos/Private/Chaos/HeightField.cpp
@@ -704,13 +704,18 @@ namespace Chaos
 	template <typename T>
 	bool THeightField<T>::GridCast(const TVector<T, 3>& StartPoint, const TVector<T, 3>& Dir, const T Length, THeightfieldRaycastVisitor<T>& Visitor) const
 	{
+		//Is this check needed?
+		if(Length < 1e-4)
+		{
+			return false;
+		}
+
 		T CurrentLength = Length;
 		TVector<T, 2> ClippedFlatRayStart;
 		TVector<T, 2> ClippedFlatRayEnd;
 
 		// Data for fast box cast
 		TVector<T, 3> Min, Max, HitPoint;
-		float ToI;
 		bool bParallel[3];
 		TVector<T, 3> InvDir;
 
@@ -721,135 +726,98 @@ namespace Chaos
 			InvDir[Axis] = bParallel[Axis] ? 0 : 1 / Dir[Axis];
 		}
 
+		T TOI;
+		const FBounds2D FlatBounds = GetFlatBounds();
+		TAABB<T,3> Bounds(
+			TVec3<T>(FlatBounds.Min[0],FlatBounds.Min[1],GeomData.GetMinHeight() * GeomData.Scale[2]),
+			TVec3<T>(FlatBounds.Max[0],FlatBounds.Max[1],GeomData.GetMaxHeight() * GeomData.Scale[2])
+			);
+		TVector<T,3> NextStart;
 
-
-		if(GetFlatBounds().ClipLine(StartPoint, StartPoint + Dir * Length, ClippedFlatRayStart, ClippedFlatRayEnd))
+		if(Bounds.RaycastFast(StartPoint, Dir, InvDir, bParallel, Length, InvCurrentLength, TOI, NextStart))
 		{
-			// The line is now valid and is entirely enclosed by the bounds (and thus, the grid)
-			if(FMath::Abs((ClippedFlatRayEnd - ClippedFlatRayStart).SizeSquared()) < SMALL_NUMBER)
-			{
-				// This is a cast down the Z axis, handle in a simpler way as this is the common case
-				const int32 QueryX = (int32)(ClippedFlatRayStart[0] / GeomData.Scale[0]);
-				const int32 QueryY = (int32)(ClippedFlatRayStart[1] / GeomData.Scale[1]);
-				const TVector<int32, 2> Cell = FlatGrid.ClampIndex({QueryX, QueryY});
+			const TVector<T,2> Scale2D(GeomData.Scale[0],GeomData.Scale[1]);
+			TVector<int32,2> CellIdx = FlatGrid.Cell(TVector<int32,2>(NextStart[0] / Scale2D[0], NextStart[1] / Scale2D[1]));
 
-				if (GetCellBounds3DScaled(Cell, Min, Max))
+			// Boundaries might push us one cell over
+			CellIdx = FlatGrid.ClampIndex(CellIdx);
+			const T ZDx = Bounds.Extents()[2];
+			const T ZMidPoint = Bounds.Min()[2] + ZDx * 0.5;
+			const TVector<T,3> ScaledDx(FlatGrid.Dx()[0] * Scale2D[0],FlatGrid.Dx()[1] * Scale2D[1],ZDx);
+			const TVector<T,2> ScaledDx2D(ScaledDx[0],ScaledDx[1]);
+			const TVector<T,2> ScaledMin = FlatGrid.MinCorner() * Scale2D;
+
+			//START
+			do
+			{
+				if(GetCellBounds3DScaled(CellIdx,Min,Max))
 				{
-					if (TAABB<T, 3>(Min,Max).RaycastFast(StartPoint, Dir, InvDir, bParallel, CurrentLength, InvCurrentLength, ToI, HitPoint))
+					// Check cell bounds
+					//todo: can do it without raycast
+					if(TAABB<T,3>(Min,Max).RaycastFast(StartPoint,Dir,InvDir,bParallel,CurrentLength,InvCurrentLength,TOI,HitPoint))
 					{
-						Visitor.VisitRaycast(Cell[1] * (GeomData.NumCols - 1) + Cell[0], CurrentLength);
+						// Visit the selected cell
+						bool bContinue = Visitor.VisitRaycast(CellIdx[1] * (GeomData.NumCols - 1) + CellIdx[0],CurrentLength);
+						if(!bContinue)
+						{
+							return false;
+						}
 					}
 				}
-			}
-			else
-			{
-				// Rasterize the line over the grid
-				const TVector<T, 2> Scale2D(GeomData.Scale[0], GeomData.Scale[1]);
-				TVector<int32, 2> StartCell = FlatGrid.Cell(ClippedFlatRayStart/Scale2D);
-				TVector<int32, 2> EndCell = FlatGrid.Cell(ClippedFlatRayEnd/Scale2D);
 
-				// Boundaries might push us one cell over
-				StartCell = FlatGrid.ClampIndex(StartCell);
-				EndCell = FlatGrid.ClampIndex(EndCell);
 
-				int32 DeltaX = FMath::Abs(EndCell[0] - StartCell[0]);
-				int32 DeltaY = -FMath::Abs(EndCell[1] - StartCell[1]);
-				int32 DirX = StartCell[0] < EndCell[0] ? 1 : -1;
-				int32 DirY = StartCell[1] < EndCell[1] ? 1 : -1;
-				int32 Error =  DeltaX + DeltaY;
+				//find next cell
 
+				//We want to know which plane we used to cross into next cell
+				const TVector<T,2> ScaledCellCenter2D = ScaledMin + TVector<T,2>(CellIdx[0] + 0.5,CellIdx[1] + 0.5) * ScaledDx2D;
+				const TVector<T,3> ScaledCellCenter(ScaledCellCenter2D[0], ScaledCellCenter2D[1], ZMidPoint);
 
-				if(StartCell == EndCell)
+				T Times[3];
+				T BestTime = CurrentLength;
+				bool bTerminate = true;
+				for(int Axis = 0; Axis < 3; ++Axis)
 				{
-					// Visit the selected cell, in this case no need to rasterize because we never
-					// leave the initial cell
-					if (GetCellBounds3DScaled(StartCell, Min, Max))
+					if(!bParallel[Axis])
 					{
-						// Check cell bounds
-						if (TAABB<T, 3>(Min,Max).RaycastFast(StartPoint, Dir, InvDir, bParallel, CurrentLength, InvCurrentLength, ToI, HitPoint))
+						const T CrossPoint = Dir[Axis] > 0 ? ScaledCellCenter[Axis] + ScaledDx[Axis] / 2 : ScaledCellCenter[Axis] - ScaledDx[Axis] / 2;
+						const T Distance = CrossPoint - NextStart[Axis];	//note: CellCenter already has /2, we probably want to use the corner instead
+						const T Time = Distance * InvDir[Axis];
+						Times[Axis] = Time;
+						if(Time < BestTime)
 						{
-							Visitor.VisitRaycast(StartCell[1] * (GeomData.NumCols - 1) + StartCell[0], CurrentLength);
+							bTerminate = false;	//found at least one plane to pass through
+							BestTime = Time;
 						}
+					} else
+					{
+						Times[Axis] = TNumericLimits<T>::Max();
 					}
 				}
-				else
+
+				if(bTerminate)
 				{
+					return false;
+				}
 
-					if (GetCellBounds3DScaled(StartCell, Min, Max))
-					{
-						// Check cell bounds
-						if (TAABB<T, 3>(Min,Max).RaycastFast(StartPoint, Dir, InvDir, bParallel, CurrentLength, InvCurrentLength, ToI, HitPoint))
-						{
-							Visitor.VisitRaycast(StartCell[1] * (GeomData.NumCols - 1) + StartCell[0], CurrentLength);
-						}
-					}
+				const TVector<int32,2> PrevIdx = CellIdx;
 
-					do 
+				for(int Axis = 0; Axis < 2; ++Axis)
+				{
+					CellIdx[Axis] += (Times[Axis] <= BestTime) ? (Dir[Axis] > 0 ? 1 : -1) : 0;
+					if(CellIdx[Axis] < 0 || CellIdx[Axis] >= FlatGrid.Counts()[Axis])
 					{
-						const int32 DoubleError = Error * 2;
-
-						if(DoubleError >= DeltaY)
-						{
-							Error += DeltaY;
-							StartCell[0] += DirX;
-						}
-
-						if(DoubleError <= DeltaX)
-						{
-							Error += DeltaX;
-							StartCell[1] += DirY;
-						}
-
-						if (GetCellBounds3DScaled(StartCell, Min, Max))
-						{
-							// Check cell bounds
-							if (TAABB<T, 3>(Min,Max).RaycastFast(StartPoint, Dir, InvDir, bParallel, CurrentLength, InvCurrentLength, ToI, HitPoint))
-							{
-								// Visit the selected cell
-								bool bContinue = Visitor.VisitRaycast(StartCell[1] * (GeomData.NumCols - 1) + StartCell[0], CurrentLength);
-								if (!bContinue)
-								{
-									return true;
-								}
-							}
-						}
-
-						// When traversing diagonally, there is a possibility of cutting corners on adjacent cell w/ current algorithm.
-						// Visit both adjacent cells to cover this issue.
-						// To visit minimal amount of cells, need to choose cell based on whichever cell boundary (vertical/horizontal) is
-						// next intersected using real value position within cells. (MaxW has a shelf in progress, holding off to not introduce bugs).
-						if (DoubleError >= DeltaY && DoubleError <= DeltaX)
-						{
-							TVector<int32, 2> TestCell = StartCell + TVector<int32, 2>(0, -DirY); // subtract because StartCell was incremented already
-							if (GetCellBounds3DScaled(TestCell, Min, Max))
-							{
-								if (TAABB<T, 3>(Min,Max).RaycastFast(StartPoint, Dir, InvDir, bParallel, CurrentLength, InvCurrentLength, ToI, HitPoint))
-								{
-									bool bContinue = Visitor.VisitRaycast(TestCell[1] * (GeomData.NumCols - 1) + TestCell[0], CurrentLength);
-									if (!bContinue)
-									{
-										return true;
-									}
-								}
-							}
-
-							TestCell = StartCell + TVector<int32, 2>(-DirX, 0);  // subtract because StartCell was incremented already
-							if (GetCellBounds3DScaled(TestCell, Min, Max))
-							{
-								if (TAABB<T, 3>(Min,Max).RaycastFast(StartPoint, Dir, InvDir, bParallel, CurrentLength, InvCurrentLength, ToI, HitPoint))
-								{
-									bool bContinue = Visitor.VisitRaycast(TestCell[1] * (GeomData.NumCols - 1) + TestCell[0], CurrentLength);
-									if (!bContinue)
-									{
-										return true;
-									}
-								}
-							}
-						}
+						return false;
+					}
+				}
 
-					} while(StartCell != EndCell);
+				if(PrevIdx == CellIdx)
+				{
+					//crossed on z plane which means no longer in heightfield bounds
+					return false;
 				}
-			}
+
+				NextStart = NextStart + Dir * BestTime;
+			} while(true);
 		}
 
 		return false;
diff --git a/Engine/Source/Runtime/Experimental/Chaos/Public/Chaos/HeightField.h b/Engine/Source/Runtime/Experimental/Chaos/Public/Chaos/HeightField.h
--- a/Engine/Source/Runtime/Experimental/Chaos/Public/Chaos/HeightField.h
+++ b/Engine/Source/Runtime/Experimental/Chaos/Public/Chaos/HeightField.h
@@ -258,6 +258,11 @@ namespace Chaos
 				return static_cast<typename FDataType::RealType>(MinValue);
 			}
 
+			FORCEINLINE T GetMaxHeight() const
+			{
+				return static_cast<typename FDataType::RealType>(MaxValue);
+			}
+
 			void Serialize(FChaosArchive& Ar)
 			{
 				int32 TempRealSize = RealSize;
