Fix triangle support function. Fixes sweep vs heightfield and sweep vs trimesh

diff --git a/Engine/Source/Programs/HeadlessChaos/Private/HeadlessChaosTestEPA.cpp b/Engine/Source/Programs/HeadlessChaos/Private/HeadlessChaosTestEPA.cpp
--- a/Engine/Source/Programs/HeadlessChaos/Private/HeadlessChaosTestEPA.cpp
+++ b/Engine/Source/Programs/HeadlessChaos/Private/HeadlessChaosTestEPA.cpp
@@ -630,13 +630,7 @@ namespace ChaosTest
 			bool bResult = GJKRaycast2(ScaledConvex, Sphere, BToATM, LocalDir, Length, OutTime, LocalPosition, LocalNormal, Thickness, bComputeMTD, Offset, Thickness);
 
 		}
-
-	}
-
-	// Currently broken EPA edge cases. As they are fixed move them to EPARealFailures_Fixed above so that we can ensure they don't break again.
-	GTEST_TEST(EPATests, EPARealFailures_Broken)
-	{
-		using namespace Chaos;
+		
 		// Sphere sweep against triangle, fails when it should hit. Raycast added as well for verification purposes.
 		{
 			const TTriangle<FReal> Triangle({ 0.000000000, 0.000000000, 0.000000000 }, { 128.000000, 0.000000000, -114.064575 }, { 128.000000, 128.000000, 2.35327148 });
@@ -674,11 +668,16 @@ namespace ChaosTest
 				bTriangleIntersects = DistToTriangle2 <= SMALL_NUMBER;	//raycast gave us the intersection point so sphere radius is already accounted for
 			}
 
-			//EXPECT_EQ(bTriangleIntersects, bSweepResult); // uncomment to demonstrate failure.
+			EXPECT_EQ(bTriangleIntersects, bSweepResult); // uncomment to demonstrate failure.
 		}
 
 	}
 
+	// Currently broken EPA edge cases. As they are fixed move them to EPARealFailures_Fixed above so that we can ensure they don't break again.
+	GTEST_TEST(EPATests, EPARealFailures_Broken)
+	{
+	}
+
 
 
 }
\ No newline at end of file
diff --git a/Engine/Source/Programs/HeadlessChaos/Private/HeadlessChaosTestGJK.cpp b/Engine/Source/Programs/HeadlessChaos/Private/HeadlessChaosTestGJK.cpp
--- a/Engine/Source/Programs/HeadlessChaos/Private/HeadlessChaosTestGJK.cpp
+++ b/Engine/Source/Programs/HeadlessChaos/Private/HeadlessChaosTestGJK.cpp
@@ -8,6 +8,7 @@
 #include "Chaos/Capsule.h"
 #include "Chaos/Convex.h"
 #include "Chaos/ImplicitObjectScaled.h"
+#include "Chaos/Triangle.h"
 
 namespace ChaosTest
 {
@@ -1494,7 +1495,6 @@ namespace ChaosTest
 			EXPECT_FLOAT_EQ(Penetration,A.GetRadius() + B.GetRadius());
 		}
 
-#if 0
 		{
 			//capsule vs triangle as we make the sweep longer the world space point of impact should stay the same
 			TParticles<T,3> ConvexParticles;
@@ -1506,7 +1506,9 @@ namespace ChaosTest
 			
 			TUniquePtr<FConvex> UniqueConvex = MakeUnique<FConvex>(ConvexParticles);
 			TSerializablePtr<FConvex> AConv(UniqueConvex);
-			const TImplicitObjectScaled<FConvex> A(AConv,TVec3<T>(1.0,1.0,1.0));
+			const TImplicitObjectScaled<FConvex> AConvScaled(AConv,TVec3<T>(1.0,1.0,1.0));
+
+			TTriangle<T> A(ConvexParticles.X(0),ConvexParticles.X(1),ConvexParticles.X(2));
 
 			const TVec3<T> Pt0(0.0,0.0,-29.6999969);
 			TVec3<T> Pt1 = Pt0;
@@ -1514,25 +1516,25 @@ namespace ChaosTest
 
 			const TCapsule<T> B(Pt0,Pt1,42.0);
 
-			const TRigidTransform<T,3> BToATM(TVec3<T>(7475.74512, 12603.9082, 377.767120),TRotation<T,3>::FromElements(0,0,0,1));
-			const TVec3<T> LocalDir(0,0,-1);
+			const TRigidTransform<T,3> BToATM(TVec3<T>(7475.74512, 12603.9082, 277.767120),TRotation<T,3>::FromElements(0,0,0,1));
+			const TVec3<T> LocalDir(0,0,-0.999999940);
 			const T Length = 49.9061584;
 			const TVec3<T> SearchDir(1,0,0);
 
 			T Time;
 			TVec3<T> Position,Normal;
-			EXPECT_TRUE(GJKRaycast2<T>(A,B,BToATM,LocalDir,Length,Time,Position,Normal,0,true,SearchDir,0));
+			EXPECT_TRUE(GJKRaycast2<T>(AConvScaled,B,BToATM,LocalDir,Length,Time,Position,Normal,0,true,SearchDir,0));
 
-			const TRigidTransform<T,3> BToATM2(TVec3<T>(7475.74512, 12603.9082, 377.767120+100),TRotation<T,3>::FromElements(0,0,0,1));
+			const TRigidTransform<T,3> BToATM2(TVec3<T>(7475.74512, 12603.9082, 277.767120+100),TRotation<T,3>::FromElements(0,0,0,1));
 
 			T Time2;
 			TVec3<T> Position2,Normal2;
+			EXPECT_TRUE(GJKRaycast2<T>(AConvScaled,B,BToATM2,LocalDir,Length+100,Time2,Position2,Normal2,0,true,SearchDir,0));
 			EXPECT_TRUE(GJKRaycast2<T>(A,B,BToATM2,LocalDir,Length+100,Time2,Position2,Normal2,0,true,SearchDir,0));
 			EXPECT_FLOAT_EQ(Time+100,Time2);
 			EXPECT_VECTOR_NEAR(Normal,Normal2,1e-4);
-			EXPECT_VECTOR_NEAR(Position,Position2,1e-4);
+			EXPECT_VECTOR_NEAR(Position,Position2,1e-3);
 		}
-#endif
 		
 	}
 
diff --git a/Engine/Source/Runtime/Experimental/Chaos/Public/Chaos/Triangle.h b/Engine/Source/Runtime/Experimental/Chaos/Public/Chaos/Triangle.h
--- a/Engine/Source/Runtime/Experimental/Chaos/Public/Chaos/Triangle.h
+++ b/Engine/Source/Runtime/Experimental/Chaos/Public/Chaos/Triangle.h
@@ -53,7 +53,7 @@ namespace Chaos
 			const float DotB = TVector<T, 3>::DotProduct(B, Direction);
 			const float DotC = TVector<T, 3>::DotProduct(C, Direction);
 
-			if(DotA > DotB && DotA > DotC)
+			if(DotA >= DotB && DotA >= DotC)
 			{
 				if(Thickness != 0)
 				{
@@ -61,7 +61,7 @@ namespace Chaos
 				}
 				return A;
 			}
-			else if(DotB > DotA && DotB > DotC)
+			else if(DotB >= DotA && DotB >= DotC)
 			{
 				if(Thickness != 0)
 				{
@@ -83,11 +83,11 @@ namespace Chaos
 			const float DotB = TVector<T, 3>::DotProduct(B, Direction);
 			const float DotC = TVector<T, 3>::DotProduct(C, Direction);
 
-			if (DotA > DotB && DotA > DotC)
+			if (DotA >= DotB && DotA >= DotC)
 			{
 				return A;
 			}
-			else if (DotB > DotA && DotB > DotC)
+			else if (DotB >= DotA && DotB >= DotC)
 			{
 				return B;
 			}
