GJK/EPA unit tests reproing Inflation > 0 assert.

diff --git a/Engine/Source/Programs/HeadlessChaos/Private/HeadlessChaosTestEPA.cpp b/Engine/Source/Programs/HeadlessChaos/Private/HeadlessChaosTestEPA.cpp
--- a/Engine/Source/Programs/HeadlessChaos/Private/HeadlessChaosTestEPA.cpp
+++ b/Engine/Source/Programs/HeadlessChaos/Private/HeadlessChaosTestEPA.cpp
@@ -676,6 +676,75 @@ namespace ChaosTest
 	// Currently broken EPA edge cases. As they are fixed move them to EPARealFailures_Fixed above so that we can ensure they don't break again.
 	GTEST_TEST(EPATests, EPARealFailures_Broken)
 	{
+		using namespace Chaos;
+		// Asserts Inflation > 0, InGJKPreDist2 = 0.57, very wrong.
+		// Scaled Convex vs Box
+		{
+			TArray<TPlaneConcrete<FReal, 3>> ConvexPlanes(
+			{
+				{{0.000000000, -512.000000, -32.0000000},{0.000000000, 0.000000000, -1.00000000}},
+				{{512.000000, 0.000000000, -32.0000000},{1.00000000, 0.000000000, 0.000000000}},
+				{{512.000000, -512.000000, 0.000000000},{0.000000000, -1.00000000, -0.000000000}},
+				{{512.000000, 0.000000000, -32.0000000},{-0.000000000, 0.000000000, -1.00000000}},
+				{{0.000000000, -512.000000, -32.0000000},{-1.00000000, 0.000000000, 0.000000000}},
+				{{0.000000000, 0.000000000, 0.000000000},{0.000000000, 1.00000000, 0.000000000}},
+				{{0.000000000, -512.000000, -32.0000000},{0.000000000, -1.00000000, 0.000000000}},
+				{{512.000000, -512.000000, 0.000000000},{0.000000000, 0.000000000, 1.00000000}},
+				{{0.000000000, 0.000000000, 0.000000000},{-1.00000000, -0.000000000, 0.000000000}},
+				{{512.000000, -512.000000, 0.000000000},{1.00000000, -0.000000000, 0.000000000}},
+				{{512.000000, 0.000000000, -32.0000000},{0.000000000, 1.00000000, -0.000000000}},
+				{{0.000000000, 0.000000000, 0.000000000},{-0.000000000, 0.000000000, 1.00000000}}	
+			});
+
+			TParticles<FReal, 3> SurfaceParticles(
+			{
+				{0.000000000, -512.000000, -32.0000000},
+				{512.000000, 0.000000000, -32.0000000},
+				{512.000000, -512.000000, -32.0000000},
+				{512.000000, -512.000000, 0.000000000},
+				{0.000000000, 0.000000000, -32.0000000},
+				{0.000000000, 0.000000000, 0.000000000},
+				{0.000000000, -512.000000, 0.000000000},
+				{512.000000, 0.000000000, 0.000000000}
+			});
+
+			FVec3 ConvexScale = { 25, 25, 1 };
+			TUniquePtr<FConvex> Convex = MakeUnique<FConvex>(MoveTemp(ConvexPlanes), MoveTemp(SurfaceParticles));
+			TImplicitObjectScaled<FConvex> ScaledConvex(MakeSerializable(Convex), ConvexScale, 0.0f);
+
+			TBox<FReal, 3> Box({ -50.0000000, -60.0000000, -30.0000000 }, { 50.0000000, 60.0000000, 30.0000000 });
+
+			const TRigidTransform<FReal, 3> BToATM({ 4404.39404, -5311.81934, 44.1764526 }, FQuat(0.100362606, 0.0230407044, -0.818859160, 0.564682245), FVec3(1.0f));
+			const FVec3 LocalDir = { -0.342119515, -0.920166731, -0.190387309 };
+			const FReal Length = 53.3335228;
+			const FReal Thickness = 0.0f;
+			const bool bComputeMTD = true;
+			const FVec3 Offset = FVec3(-5311.83203, -4404.37891, -44.1764526);
+
+			FReal OutTime = -1;
+			FVec3 LocalPosition(-1);
+			FVec3 LocalNormal(-1);
+
+			bool bResult = GJKRaycast2(ScaledConvex, Box, BToATM, LocalDir, Length, OutTime, LocalPosition, LocalNormal, Thickness, bComputeMTD, Offset, Thickness);
+		}
+
+		// Asserts Inflation > 0, InGJKPreDist2 is barely over 1e-6
+		// Triangle v Box
+		{
+			TTriangle<FReal> Triangle(FVec3(0.000000000, 0.000000000, 0.000000000), FVec3( 128.000000, 0.000000000, 35.9375000 ), FVec3(128.000000, 128.000000, 134.381042 ));
+			TBox<FReal, 3> Box(FVec3(-50.0000000, -60.0000000, -30.0000000), FVec3 (50.0000000, 60.0000000, 30.0000000));
+
+			const TRigidTransform<FReal, 3> Transform({ 127.898438, 35.0742188, 109.781067 }, FQuat(0.374886870, -0.0289460570, 0.313643545, 0.871922970), FVec3(1.0));
+			const FVec3 Dir = { 0.801564395, 0.525258720, 0.285653293 };
+			const FReal Length = 26.7055893;
+			const FReal Thickness = 0.0f;
+			const bool bComputeMTD = true;
+
+			FReal OutTime = -1;
+			FVec3 LocalPosition(-1);
+			FVec3 LocalNormal(-1);
+			GJKRaycast2(Triangle, Box, Transform, Dir, Length, OutTime, LocalPosition, LocalNormal, Thickness, bComputeMTD);
+		}
 	}
 
 
