Chaos : Fix
-Fix CIS

diff --git a/Engine/Source/Programs/HeadlessChaos/Public/HeadlessChaosCollisionConstraints.h b/Engine/Source/Programs/HeadlessChaos/Public/HeadlessChaosCollisionConstraints.h
--- a/Engine/Source/Programs/HeadlessChaos/Public/HeadlessChaosCollisionConstraints.h
+++ b/Engine/Source/Programs/HeadlessChaos/Public/HeadlessChaosCollisionConstraints.h
@@ -78,7 +78,8 @@ public:
 		const FRigidTransform3 Transform0 = Collisions::GetTransform(Constraint.Particle[0]);
 		const FRigidTransform3 Transform1 = Collisions::GetTransform(Constraint.Particle[1]);
 
-		Collisions::UpdateManifold(Constraint, Transform0, Transform1, CullDistance);
+		if( Constraint.GetType()== FCollisionConstraintBase::FType::MultiPoint)
+			Collisions::UpdateManifold(*Constraint.As<FMultiPointContactConstraint>(), CullDistance);
 	}
 
 
