Chaos : Filtering

- Update tests to reflect changes in shape requirements.

diff --git a/Engine/Source/Programs/HeadlessChaos/Private/HeadlessChaosTestSim.cpp b/Engine/Source/Programs/HeadlessChaos/Private/HeadlessChaosTestSim.cpp
--- a/Engine/Source/Programs/HeadlessChaos/Private/HeadlessChaosTestSim.cpp
+++ b/Engine/Source/Programs/HeadlessChaos/Private/HeadlessChaosTestSim.cpp
@@ -1,6 +1,8 @@
 // Copyright Epic Games, Inc. All Rights Reserved.
 
 #include "HeadlessChaos.h"
+#include "HeadlessChaosTestUtility.h"
+
 #include "Chaos/PBDRigidsEvolutionGBF.h"
 #include "Chaos/Box.h"
 #include "Chaos/Sphere.h"
@@ -27,6 +29,9 @@ namespace ChaosTest {
 		Dynamic->I() = FMatrix33(100000.0f, 100000.0f, 100000.0f);
 		Dynamic->InvI() = FMatrix33(1.0f / 100000.0f, 1.0f / 100000.0f, 1.0f / 100000.0f);
 
+		::ChaosTest::SetParticleSimDataToCollide({ Static,Dynamic });
+
+
 		const FReal Dt = 1 / 60.f;
 		for (int i = 0; i < 100; ++i)
 		{
@@ -54,6 +59,9 @@ namespace ChaosTest {
 		Dynamic->I() = FMatrix33(100000.0f, 100000.0f, 100000.0f);
 		Dynamic->InvI() = FMatrix33(1.0f / 100000.0f, 1.0f / 100000.0f, 1.0f / 100000.0f);
 
+		::ChaosTest::SetParticleSimDataToCollide({ Static,Dynamic });
+
+
 		const FReal Dt = 1 / 60.f;
 		for (int i = 0; i < 100; ++i)
 		{
@@ -89,6 +97,8 @@ namespace ChaosTest {
 		Dynamic->I() = FMatrix33(1, 1, 1);
 		Dynamic->InvI() = FMatrix33(1, 1, 1);
 
+		::ChaosTest::SetParticleSimDataToCollide({ Static,Dynamic });
+
 		for (int i = 0; i < 100; ++i)
 		{
 			Evolution.AdvanceOneTimeStep(1 / 60.f);
@@ -131,6 +141,8 @@ namespace ChaosTest {
 		Dynamic2->I() = FMatrix33(100000.0f, 100000.0f, 100000.0f);
 		Dynamic2->InvI() = FMatrix33(1.0f / 100000.0f, 1.0f / 100000.0f, 1.0f / 100000.0f);
 
+		::ChaosTest::SetParticleSimDataToCollide({ Static,Dynamic1,Dynamic2 });
+
 		bool Dynamic1WentToSleep = false;
 		bool Dynamic1HasWokeAgain = false;
 		for (int i = 0; i < 1000; ++i)
diff --git a/Engine/Source/Programs/HeadlessChaos/Private/HeadlessChaosTestSolverProxies.cpp b/Engine/Source/Programs/HeadlessChaos/Private/HeadlessChaosTestSolverProxies.cpp
--- a/Engine/Source/Programs/HeadlessChaos/Private/HeadlessChaosTestSolverProxies.cpp
+++ b/Engine/Source/Programs/HeadlessChaos/Private/HeadlessChaosTestSolverProxies.cpp
@@ -2,6 +2,8 @@
 
 #include "HeadlessChaosTestSolverProxies.h"
 #include "HeadlessChaos.h"
+#include "HeadlessChaosTestUtility.h"
+
 
 #include "Chaos/ParticleHandle.h"
 #include "Chaos/ErrorReporter.h"
@@ -42,6 +44,8 @@ namespace ChaosTest {
 		Particle->SetV(TVector<float, 3>(0, 0, 10));
 		Solver->AddDirtyProxy(Particle->Proxy);
 
+		::ChaosTest::SetParticleSimDataToCollide({ Particle.Get() });
+
 		Solver->PushPhysicsState(Module->GetDispatcher());
 
 		FPhysicsSolverAdvanceTask AdvanceTask(Solver, 100.0f);
@@ -161,6 +165,8 @@ namespace ChaosTest {
 		Solver->AddDirtyProxy(Particle2->Proxy);
 		Particle2->SetObjectState(Chaos::EObjectStateType::Sleeping, false);
 
+		::ChaosTest::SetParticleSimDataToCollide({Particle.Get(),Particle2.Get() });
+
 		Solver->PushPhysicsState(Module->GetDispatcher());
 
 		FPhysicsSolverAdvanceTask AdvanceTask(Solver, 100.0f);
diff --git a/Engine/Source/Programs/HeadlessChaos/Private/HeadlessChaosTestUtility.cpp b/Engine/Source/Programs/HeadlessChaos/Private/HeadlessChaosTestUtility.cpp
--- a/Engine/Source/Programs/HeadlessChaos/Private/HeadlessChaosTestUtility.cpp
+++ b/Engine/Source/Programs/HeadlessChaos/Private/HeadlessChaosTestUtility.cpp
@@ -268,6 +268,9 @@ namespace ChaosTest {
 			OutElements->Add(TVector<int32, 3>(2, 0, 1)); // Front
 			OutElements->Add(TVector<int32, 3>(0, 2, 3));
 		}
+
+		::ChaosTest::SetParticleSimDataToCollide({ Particle});
+
 	}
 
 	
@@ -513,6 +516,9 @@ namespace ChaosTest {
 			OutElements->Add(TVector<int32, 3>(2, 0, 1)); // Front
 			OutElements->Add(TVector<int32, 3>(0, 2, 3));
 		}
+
+		::ChaosTest::SetParticleSimDataToCollide({ Particle });
+
 	}
 
 	template<class T>
@@ -560,6 +566,8 @@ namespace ChaosTest {
 		Particle->W() = TVector<T, 3>(0.f, 0.f, 0.f);
 		Particle->SetDynamicGeometry( MakeUnique<TPlane<T, 3>>(TVector<T, 3>(0.f, 0.f, 0.f), TVector<T, 3>(0.f, 0.f, 1.f)));
 
+		::ChaosTest::SetParticleSimDataToCollide({ Particle });
+
 		return Particle;
 	}
 	template TKinematicGeometryParticleHandle<float, 3>* AppendStaticAnalyticFloor(TPBDRigidsEvolutionGBF<float, 3>& Evolution);
@@ -576,6 +584,8 @@ namespace ChaosTest {
 		Particle->W() = TVector<T, 3>(0.f, 0.f, 0.f);
 		Particle->SetDynamicGeometry(MakeUnique<TPlane<T, 3>>(TVector<T, 3>(0.f, 0.f, 0.f), TVector<T, 3>(0.f, 0.f, 1.f)));
 
+		::ChaosTest::SetParticleSimDataToCollide({ Particle });
+
 		return Particle;
 	}
 	template TKinematicGeometryParticleHandle<float, 3>* AppendStaticAnalyticFloor(TPBDRigidsSOAs<float, 3>& Evolution);
@@ -606,6 +616,8 @@ namespace ChaosTest {
 
 		Particle->SetDynamicGeometry(MakeUnique<FConvex>(Cube));
 
+		::ChaosTest::SetParticleSimDataToCollide({ Particle });
+
 		return Particle;
 	}
 	template TKinematicGeometryParticleHandle<float, 3>* AppendStaticConvexFloor(TPBDRigidsSOAs<float, 3>& Evolution);
@@ -666,6 +678,13 @@ namespace ChaosTest {
 		InParticles.InvI() = PMatrix<T, 3, 3>(1.f, 0.f, 0.f, 0.f, 1.f, 0.f, 0.f, 0.f, 1.f);
 		InParticles.SetDynamicGeometry(MakeUnique<FConvex>(Cube));
 		InParticles.SetObjectStateLowLevel(EObjectStateType::Dynamic);
+
+		::ChaosTest::SetParticleSimDataToCollide({ &InParticles });
+		//for (const TUniquePtr<Chaos::TPerShapeData<Chaos::FReal, 3>>& Shape : InParticles.ShapesArray())
+		//{
+		//	Shape->SimData.Word3 = 1;
+		//	Shape->SimData.Word1 = 1;
+		//}
 	}
 	template void AppendDynamicParticleConvexBox(TPBDRigidParticleHandle<float, 3> &, const TVector<float, 3>& Scale);
 
@@ -754,5 +773,27 @@ namespace ChaosTest {
 		return FVec3(FVector::UpVector);
 	}
 
+	void SetParticleSimDataToCollide(TArray< Chaos::TGeometryParticle<float, 3>* > ParticleArray)
+	{
+		for (Chaos::TGeometryParticle<float, 3>* Particle : ParticleArray)
+		{
+			for (const TUniquePtr<Chaos::TPerShapeData<Chaos::FReal, 3>>& Shape :Particle->ShapesArray())
+			{
+				Shape->SimData.Word3 = 1;
+				Shape->SimData.Word1 = 1;
+			}
+		}
+	}
+	void SetParticleSimDataToCollide(TArray< Chaos::TGeometryParticleHandle<float, 3>* > ParticleArray)
+	{
+		for (Chaos::TGeometryParticleHandle<float, 3>* Particle : ParticleArray)
+		{
+			for (const TUniquePtr<Chaos::TPerShapeData<Chaos::FReal, 3>>& Shape : Particle->ShapesArray())
+			{
+				Shape->SimData.Word3 = 1;
+				Shape->SimData.Word1 = 1;
+			}
+		}
+	}
 
 }
diff --git a/Engine/Source/Programs/HeadlessChaos/Public/HeadlessChaosTestUtility.h b/Engine/Source/Programs/HeadlessChaos/Public/HeadlessChaosTestUtility.h
--- a/Engine/Source/Programs/HeadlessChaos/Public/HeadlessChaosTestUtility.h
+++ b/Engine/Source/Programs/HeadlessChaos/Public/HeadlessChaosTestUtility.h
@@ -164,4 +164,9 @@ namespace ChaosTest {
 		return FRotation3::FromAxisAngle(FVec3(1, 0, 0), X) * FRotation3::FromAxisAngle(FVec3(0, 1, 0), Y) * FRotation3::FromAxisAngle(FVec3(0, 0, 1), Z);
 	}
 
+	/**/
+	void SetParticleSimDataToCollide(TArray< Chaos::TGeometryParticle<float, 3>* > ParticleArray);
+	void SetParticleSimDataToCollide(TArray< Chaos::TGeometryParticleHandle<float, 3>* > ParticleArray);
+
+	
 }
