Chaos - Simple-ish failing EPA unit test

diff --git a/Engine/Source/Programs/HeadlessChaos/Private/HeadlessChaosTestEPA.cpp b/Engine/Source/Programs/HeadlessChaos/Private/HeadlessChaosTestEPA.cpp
--- a/Engine/Source/Programs/HeadlessChaos/Private/HeadlessChaosTestEPA.cpp
+++ b/Engine/Source/Programs/HeadlessChaos/Private/HeadlessChaosTestEPA.cpp
@@ -752,30 +752,30 @@ namespace ChaosTest
 
 		TArray<TPlaneConcrete<FReal,3>> GroundConvexPlanes(
 		{
-		{{0.000000000,-512.000000,-32.0000000},{0.000000000,0.000000000,-1.00000000}},
-		{{512.000000,0.000000000,-32.0000000},{1.00000000,0.000000000,0.000000000}},
-		{{512.000000,-512.000000,0.000000000},{0.000000000,-1.00000000,-0.000000000}},
-		{{512.000000,0.000000000,-32.0000000},{-0.000000000,0.000000000,-1.00000000}},
-		{{0.000000000,-512.000000,-32.0000000},{-1.00000000,0.000000000,0.000000000}},
-		{{0.000000000,0.000000000,0.000000000},{0.000000000,1.00000000,0.000000000}},
-		{{0.000000000,-512.000000,-32.0000000},{0.000000000,-1.00000000,0.000000000}},
-		{{512.000000,-512.000000,0.000000000},{0.000000000,0.000000000,1.00000000}},
-		{{0.000000000,0.000000000,0.000000000},{-1.00000000,-0.000000000,0.000000000}},
-		{{512.000000,-512.000000,0.000000000},{1.00000000,-0.000000000,0.000000000}},
-		{{512.000000,0.000000000,-32.0000000},{0.000000000,1.00000000,-0.000000000}},
-		{{0.000000000,0.000000000,0.000000000},{-0.000000000,0.000000000,1.00000000}}
+			{{0.000000000,-512.000000,-32.0000000},{0.000000000,0.000000000,-1.00000000}},
+			{{512.000000,0.000000000,-32.0000000},{1.00000000,0.000000000,0.000000000}},
+			{{512.000000,-512.000000,0.000000000},{0.000000000,-1.00000000,-0.000000000}},
+			{{512.000000,0.000000000,-32.0000000},{-0.000000000,0.000000000,-1.00000000}},
+			{{0.000000000,-512.000000,-32.0000000},{-1.00000000,0.000000000,0.000000000}},
+			{{0.000000000,0.000000000,0.000000000},{0.000000000,1.00000000,0.000000000}},
+			{{0.000000000,-512.000000,-32.0000000},{0.000000000,-1.00000000,0.000000000}},
+			{{512.000000,-512.000000,0.000000000},{0.000000000,0.000000000,1.00000000}},
+			{{0.000000000,0.000000000,0.000000000},{-1.00000000,-0.000000000,0.000000000}},
+			{{512.000000,-512.000000,0.000000000},{1.00000000,-0.000000000,0.000000000}},
+			{{512.000000,0.000000000,-32.0000000},{0.000000000,1.00000000,-0.000000000}},
+			{{0.000000000,0.000000000,0.000000000},{-0.000000000,0.000000000,1.00000000}}
 		});
 
 		TParticles<FReal,3> GroundSurfaceParticles(
 		{
 			{0.000000000,-512.000000,-32.0000000},
-		{512.000000,0.000000000,-32.0000000},
-		{512.000000,-512.000000,-32.0000000},
-		{512.000000,-512.000000,0.000000000},
-		{0.000000000,0.000000000,-32.0000000},
-		{0.000000000,0.000000000,0.000000000},
-		{0.000000000,-512.000000,0.000000000},
-		{512.000000,0.000000000,0.000000000}
+			{512.000000,0.000000000,-32.0000000},
+			{512.000000,-512.000000,-32.0000000},
+			{512.000000,-512.000000,0.000000000},
+			{0.000000000,0.000000000,-32.0000000},
+			{0.000000000,0.000000000,0.000000000},
+			{0.000000000,-512.000000,0.000000000},
+			{512.000000,0.000000000,0.000000000}
 		});
 
 		TParticles<FReal,3> BoatSurfaceParticles(
@@ -842,7 +842,7 @@ namespace ChaosTest
 			{-142.015427, -9.19016743, -11.2502632},
 			{-27.1481628, -111.977615, 35.8436165},
 			{-23.9894104, -116.828667, 43.9551315}
-			});
+		});
 
 
 
@@ -968,7 +968,7 @@ namespace ChaosTest
 			{{-150.005661, -100.080147, 51.9458313},{-0.153074399, -0.805065334, -0.573095083}},
 			{{154.527252, 7.14775467, 133.457825},{-0.0549643822, -0.115145117, 0.991826892}},
 			{{-134.336945, 83.6453476, 19.9467926},{-0.852820277, 0.468897045, -0.229854882}},		
-});
+		});
 
 
 		FVec3 GroundConvexScale = { 25,25,1 };
@@ -993,11 +993,35 @@ namespace ChaosTest
 
 		FVec3 WorldLocation = BoatTransform.TransformPosition(ClosestA);
 		FVec3 WorldNormal = BoatTransform.TransformVectorNoScale(Normal);
+	}
 
+	//
+	// Two boxes, one large and flat, with another smaller box on top.
+	// All normals should be either Up or Down depending on which body is first (in this case, they are down)
+	// Problem:
+	//		Occasionally we get a non-vertical normal from EPA
+	// This test reproduces an in-game example. 
+	// NOTE: These are shapes in a skeletal mesh, hence the rotation (mesh is on its side so its Y is "up")
+	//
+	GTEST_TEST(EPATests, DISABLED_EPARealFailures_BoxBox)
+	{
+		TBox<FReal, 3> A(FVec3(-250.000000, 250.000000, -750.000000), FVec3(1250.00000, 350.000000, 750.000000));
+		TBox<FReal, 3> B(FVec3(50.0000000, -50.0000000, -50.0000000), FVec3(150.000000, 50.0000000, 50.0000000));
+		FRigidTransform3 ATM = FRigidTransform3(FVec3(0.000000000, 0.000000000, 0.000000000), FRotation3::FromElements(-0.707106709, 0.000000000, 0.000000000, 0.707106829));
+		FRigidTransform3 BTM = FRigidTransform3(FVec3(804.534912, 17.9698277, -199.983994), FRotation3::FromElements(0.706765056, -2.78512689e-05, 0.000485646888, -0.707448125));
+		FVec3 InitialDir = FVec3(1.00000000, 0.000000000, 0.000000000);
 
-		
-	}
+		const FRigidTransform3 BToATM = BTM.GetRelativeTransform(ATM);
+		FReal Penetration;
+		FVec3 ClosestA, ClosestB, NormalA;
+		int32 NumIterations = 0;
+		GJKPenetration<true>(A, B, BToATM, Penetration, ClosestA, ClosestB, NormalA, 0.0f, InitialDir, 0.0f, &NumIterations);
 
+		FVec3 Location = ATM.TransformPosition(ClosestA);
+		FVec3 Normal = -ATM.TransformVectorNoScale(NormalA);
+		FReal Phi = -Penetration;
 
+		EXPECT_GT(FMath::Abs(Normal.Z), 0.8f);
+	}
 
 }
\ No newline at end of file
