Fix headlesschaos build errors.

diff --git a/Engine/Source/Programs/HeadlessChaos/Private/HeadlessChaosApp.cpp b/Engine/Source/Programs/HeadlessChaos/Private/HeadlessChaosApp.cpp
--- a/Engine/Source/Programs/HeadlessChaos/Private/HeadlessChaosApp.cpp
+++ b/Engine/Source/Programs/HeadlessChaos/Private/HeadlessChaosApp.cpp
@@ -43,7 +43,7 @@ DEFINE_LOG_CATEGORY(LogHeadlessChaos);
 
 TEST(ImplicitTests, Implicit) {
 	ChaosTest::ImplicitPlane<float>();
-	//ChaosTest::ImplicitCube<float>();
+	ChaosTest::ImplicitCube<float>();
 	ChaosTest::ImplicitSphere<float>();
 	ChaosTest::ImplicitCylinder<float>();
 	ChaosTest::ImplicitTaperedCylinder<float>();
diff --git a/Engine/Source/Programs/HeadlessChaos/Private/HeadlessChaosTestImplicits.cpp b/Engine/Source/Programs/HeadlessChaos/Private/HeadlessChaosTestImplicits.cpp
--- a/Engine/Source/Programs/HeadlessChaos/Private/HeadlessChaosTestImplicits.cpp
+++ b/Engine/Source/Programs/HeadlessChaos/Private/HeadlessChaosTestImplicits.cpp
@@ -21,6 +21,7 @@
 #include "Chaos/Levelset.h"
 #include "Chaos/UniformGrid.h"
 #include "Chaos/Utilities.h"
+#include "Chaos/CastingUtilities.h"
 #include "Chaos/Convex.h"
 #include "Math/RandomStream.h"
 #include "Chaos/ErrorReporter.h"
@@ -115,16 +116,15 @@ namespace ChaosTest {
 		TestFindClosestIntersection(Subject, TVector<T, 3>(0, 0, 2), TVector<T, 3>(0, 0, 1), Caller);
 
 		// closest point near origin (-)
-		TestFindClosestIntersection(Subject, TVector<T, 3>(0, 0, 1 / 2.), TVector<T, 3>(0, 0, 1), Caller);
+		TestFindClosestIntersection<T>(Subject, TVector<T, 3>(0, 0, 1 / 2.), TVector<T, 3>(0, 0, 1), Caller);
 	}
 
 
 	/* Takes an ImplicitObject of unit size (circumscribed inside a 2x2 cube centered on the origin).
 	   Tests the .Support() function. */
-	template<class T>
-	void UnitImplicitObjectSupportPhis(FImplicitObject &Subject, FString Caller)
+	template<typename GeometryType, typename T>
+	void UnitImplicitObjectSupportPhis(GeometryType& Subject, FString Caller)
 	{
-#if 0
 		typedef TVector<T, 3> TVector3;
 		FString Error = FString("Called by ") + Caller + FString(".");
 
@@ -142,7 +142,6 @@ namespace ChaosTest {
 		EXPECT_VECTOR_NEAR_ERR(Subject.Support(TVector3(0, -1, 0), T(1)), (TVector3(0, -2, 0)), KINDA_SMALL_NUMBER, Error);
 		EXPECT_VECTOR_NEAR_ERR(Subject.Support(TVector3(1, 0, 0), T(1)), (TVector3(2, 0, 0)), KINDA_SMALL_NUMBER, Error);
 		EXPECT_VECTOR_NEAR_ERR(Subject.Support(TVector3(-1, 0, 0), T(1)), (TVector3(-2, 0, 0)), KINDA_SMALL_NUMBER, Error);
-#endif
 	}
 
 
@@ -174,14 +173,14 @@ namespace ChaosTest {
 		{// closest point near origin
 			TPlane<T, 3> Subject(TVector3(0), TVector3(0, 0, 1));
 			TVector3 InputPoint = TVector3(1, 1, 1);
-			TestFindClosestIntersection(Subject, InputPoint, TVector<T, 3>(1, 1, 0), Caller);
+			TestFindClosestIntersection<T>(Subject, InputPoint, TVector<T, 3>(1, 1, 0), Caller);
 			EXPECT_VECTOR_NEAR_DEFAULT(Subject.FindClosestPoint(InputPoint), (TVector3(1, 1, 0)));
 		}
 
 		{// closest point single axis off origin (+)
 			TVector3 InputPoint = TVector3(0, 0, 2);
 			TPlane<T, 3> Subject = TPlane<T, 3>(TVector3(0, 0, 1), TVector3(0, 0, 1));
-			TestFindClosestIntersection(Subject, InputPoint, TVector<T, 3>(0, 0, 1), Caller);
+			TestFindClosestIntersection<T>(Subject, InputPoint, TVector<T, 3>(0, 0, 1), Caller);
 			EXPECT_VECTOR_NEAR(Subject.FindClosestPoint(InputPoint), FVector(0, 0, 1), 0.001);
 			EXPECT_VECTOR_NEAR(Subject.FindClosestPoint(TVector3(0, 1, 2)), FVector(0,1,1), 0.001);
 		}
@@ -189,7 +188,7 @@ namespace ChaosTest {
 		{// closest point off origin (+)
 			TVector3 InputPoint = TVector3(11,11,11);
 			TPlane<T, 3> Subject = TPlane<T, 3>(TVector3(10, 10, 10), TVector3(1, 1, 1).GetSafeNormal());
-			TestFindClosestIntersection(Subject, InputPoint, TVector<T, 3>(10, 10, 10), Caller);
+			TestFindClosestIntersection<T>(Subject, InputPoint, TVector<T, 3>(10, 10, 10), Caller);
 			TVector3 NearestPoint = Subject.FindClosestPoint(InputPoint); // wrong (9.26...)
 			EXPECT_VECTOR_NEAR(Subject.FindClosestPoint(InputPoint), FVector(10, 10, 10), 0.001);
 		}
@@ -197,7 +196,7 @@ namespace ChaosTest {
 		{// closest point off origin (-)
 			TVector3 InputPoint = TVector3(9,9,9);
 			TPlane<T, 3>Subject = TPlane<T, 3>(TVector3(10, 10, 10), TVector3(1, 1, 1).GetSafeNormal());
-			TestFindClosestIntersection(Subject, InputPoint, TVector<T, 3>(10, 10, 10), Caller);
+			TestFindClosestIntersection<T>(Subject, InputPoint, TVector<T, 3>(10, 10, 10), Caller);
 			TVector3 NearestPoint = Subject.FindClosestPoint(InputPoint); // (10.73...)
 			EXPECT_VECTOR_NEAR(Subject.FindClosestPoint(InputPoint), FVector(10, 10, 10), 0.001);
 		}
@@ -208,7 +207,6 @@ namespace ChaosTest {
 	template<class T>
 	void ImplicitCube()
 	{
-		EXPECT_TRUE(false);
 
 		typedef TVector<T, 3> TVector3;
 		FString Caller("ImplicitCube()");
@@ -317,7 +315,7 @@ namespace ChaosTest {
 		{// closest point off origin (+)
 			TBox<T, 3> Subject2(TVector3(2), TVector3(4));
 			TVector3 InputPoint(5, 5, 5);
-			TestFindClosestIntersection(Subject2, InputPoint, TVector<T, 3>(4, 4, 4), Caller);
+			TestFindClosestIntersection<T>(Subject2, InputPoint, TVector<T, 3>(4, 4, 4), Caller);
 
 			EXPECT_VECTOR_NEAR(Subject2.FindClosestPoint(InputPoint), FVector(4,4,4), 0.001);
 			TVector3 test2 = Subject2.FindClosestPoint(TVector3(3.5, 3.5, 3.5));
@@ -357,7 +355,7 @@ namespace ChaosTest {
 		UnitImplicitObjectNormalsInternal<T>(Subject, Caller);
 		UnitImplicitObjectNormalsExternal<T>(Subject, Caller);
 		UnitImplicitObjectIntersections<T>(Subject, Caller);
-		UnitImplicitObjectSupportPhis<T>(Subject, Caller);
+		UnitImplicitObjectSupportPhis<TSphere<T,3>, T>(Subject, Caller);
 
 		// intersection
 		EXPECT_TRUE(Subject.Intersects(TSphere<T, 3>(TVector<T,3>(0.f), 2.f)));
@@ -380,7 +378,7 @@ namespace ChaosTest {
 		{// closest point off origin (+)
 			TSphere<T, 3> Subject2(TVector3(2), 2);
 			TVector3 InputPoint(2, 2, 5);
-			TestFindClosestIntersection(Subject2, InputPoint, TVector<T, 3>(2, 2, 4), Caller);
+			TestFindClosestIntersection<T>(Subject2, InputPoint, TVector<T, 3>(2, 2, 4), Caller);
 			EXPECT_VECTOR_NEAR(Subject2.FindClosestPoint(InputPoint), FVector(2, 2, 4), 0.001);
 			EXPECT_VECTOR_NEAR(Subject2.FindClosestPoint(TVector3(2, 2, 3.5)), FVector(2, 2, 4), 0.001);
 		}
@@ -557,13 +555,13 @@ namespace ChaosTest {
 		{// closest point off origin (+)
 			TTaperedCylinder<T> Subject2(TVector<T, 3>(2, 2, 4), TVector<T, 3>(2, 2, 0), 2, 2);
 			TVector<T, 3> InputPoint(2, 2, 5);
-			TestFindClosestIntersection(Subject2, InputPoint, TVector<T, 3>(2, 2, 4), Caller);
+			TestFindClosestIntersection<T>(Subject2, InputPoint, TVector<T, 3>(2, 2, 4), Caller);
 		}
 
 		{// closest point off origin (-)
 			TTaperedCylinder<T> Subject2(TVector<T, 3>(2, 2, 4), TVector<T, 3>(2, 2, 0), 2, 2);
 			TVector<T, 3> InputPoint(2, 3, 2);
-			TestFindClosestIntersection(Subject2, InputPoint, TVector<T, 3>(2, 4, 2), Caller);
+			TestFindClosestIntersection<T>(Subject2, InputPoint, TVector<T, 3>(2, 4, 2), Caller);
 		}
 	}
 	template void ImplicitTaperedCylinder<float>();
@@ -580,11 +578,11 @@ namespace ChaosTest {
 
 		UnitImplicitObjectNormalsInternal<T>(SubjectUnit, Caller);
 		UnitImplicitObjectNormalsExternal<T>(SubjectUnit, Caller);
-		UnitImplicitObjectSupportPhis<T>(SubjectUnit, Caller);
+		UnitImplicitObjectSupportPhis<TCapsule<T>, T>(SubjectUnit, Caller);
 
 #if RUN_KNOWN_BROKEN_TESTS
 		// FindClosestIntersection broken with cylinder size 0
-		UnitImplicitObjectIntersections<T>(SubjectUnit, Caller);
+		UnitImplicitObjectIntersections<T>(SubjectUnit, Caller); 
 #endif
 
 		TCapsule<T> Subject(TVector3(0, 0, 1), TVector3(0, 0, -1), 1);
@@ -620,7 +618,7 @@ namespace ChaosTest {
 	{
 		typedef TVector<T, 3> TVector3;
 		TUniquePtr<TBox<T, 3>> UnitCube = MakeUnique<TBox<T, 3>>(TVector<T, 3>(-1), TVector<T, 3>(1));
-		TImplicitObjectScaled<TBox<T, 3>> UnitUnscaled(MakeSerializable(UnitCube), TVector<T, 3>(1));
+		TImplicitObjectScaled<TBox<T,3>> UnitUnscaled(MakeSerializable(UnitCube), TVector<T, 3>(1));
 		UnitImplicitObjectNormalsInternal<T>(UnitUnscaled, FString("ImplicitTransformed()"));
 		UnitImplicitObjectNormalsExternal<T>(UnitUnscaled, FString("ImplicitTransformed()"));
 		UnitImplicitObjectIntersections<T>(UnitUnscaled, FString("ImplicitTransformed()"));
@@ -689,17 +687,19 @@ namespace ChaosTest {
 	template <class T>
 	void ImplicitTransformed()
 	{
-#if 0
 		typedef TVector<T, 3> TVector3;
+
+
+		TRigidTransform<T, 3> Identity(FVec3(0), FQuat::Identity);
 		
 		TUniquePtr<TBox<T, 3>> UnitCube = MakeUnique<TBox<T, 3>>(TVector<T, 3>(-1), TVector<T, 3>(1));
-		TImplicitObjectTransformed<T, 3> UnitUnrotated(MakeSerializable(UnitCube), TRigidTransform<T, 3>(TVector<T, 3>(0), TRotation<T, 3>::FromAxisAngle(TVector<T, 3>(0), 0)));
+		TImplicitObjectTransformed<T, 3> UnitUnrotated(MakeSerializable(UnitCube), TRigidTransform<T, 3>(TVector<T, 3>(0), FQuat(0,0,0, 0)));
 		UnitImplicitObjectNormalsInternal<T>(UnitUnrotated, FString("ImplicitTransformed()"));
 		UnitImplicitObjectNormalsExternal<T>(UnitUnrotated, FString("ImplicitTransformed()"));
 		UnitImplicitObjectIntersections<T>(UnitUnrotated, FString("ImplicitTransformed()"));
 		
 		// Rotate 45 degrees around z axis @ origin.
-		TImplicitObjectTransformed<T, 3> UnitRotated(MakeSerializable(UnitCube), TRigidTransform<T, 3>(TVector<T, 3>(0), TRotation<T, 3>::FromAxisAngle(TVector<T, 3>(0, 0, sin(.3927)), cos(.3927))));
+		TImplicitObjectTransformed<T, 3> UnitRotated(MakeSerializable(UnitCube), TRigidTransform<T, 3>(TVector<T, 3>(0), FQuat(0, 0, sin(.3927), cos(.3927))));
 		
 		{// unit rotated normals
 			TVector<T, 3> Normal;
@@ -714,13 +714,13 @@ namespace ChaosTest {
 		}
 
 		TUniquePtr<TBox<T, 3>> Cube = MakeUnique<TBox<T, 3>>(TVector<T, 3>(-2, -5, -5), TVector<T,3>(8, 5, 5));
-		TImplicitObjectTransformed<T, 3> Untransformed(MakeSerializable(Cube), TRigidTransform<T, 3>(TVector<T, 3>(0), TRotation<T, 3>::FromAxisAngle(TVector<T, 3>(0), 0)));
-		TImplicitObjectTransformed<T, 3> Translated(MakeSerializable(Cube), TRigidTransform<T, 3>(TVector<T, 3>(4, 0, 0), TRotation<T, 3>::FromAxisAngle(TVector<T, 3>(0), 0)));
+		TImplicitObjectTransformed<T, 3> Untransformed(MakeSerializable(Cube), TRigidTransform<T, 3>(TVector<T, 3>(0), FQuat::Identity));
+		TImplicitObjectTransformed<T, 3> Translated(MakeSerializable(Cube), TRigidTransform<T, 3>(TVector<T, 3>(4, 0, 0), FQuat::Identity));
 		
 		// Rotate 90 degrees around z axis @ origin. 
 		float rad_45 = FMath::DegreesToRadians(45);
-		TImplicitObjectTransformed<T, 3> Rotated(MakeSerializable(Cube), TRigidTransform<T, 3>(TVector<T, 3>(0), TRotation<T, 3>::FromElements(TVector<T, 3>(0, 0, sin(rad_45)), cos(rad_45))));
-		TImplicitObjectTransformed<T, 3> Transformed(MakeSerializable(Cube), TRigidTransform<T, 3>(TVector<T, 3>(4, 0, 0), TRotation<T, 3>::FromElements(TVector<T, 3>(0, 0, sin(rad_45)), cos(rad_45))));
+		TImplicitObjectTransformed<T, 3> Rotated(MakeSerializable(Cube), TRigidTransform<T, 3>(TVector<T, 3>(0), FQuat(0, 0, sin(rad_45), cos(rad_45))));
+		TImplicitObjectTransformed<T, 3> Transformed(MakeSerializable(Cube), TRigidTransform<T, 3>(TVector<T, 3>(4, 0, 0), FQuat(0, 0, sin(rad_45), cos(rad_45))));
 
 		{// phi
 			const TVector<T, 3> NearEdge(7.5, 0, 0);
@@ -750,21 +750,43 @@ namespace ChaosTest {
 		}
 		
 		{//support
+
 			const TVector<T, 3> DirX(1, 0, 0);
-			TVector<T, 3> SupportPt = Untransformed.Support(DirX, 1);
+			TVector<T, 3> SupportPt = Utilities::CastHelper(Untransformed, Identity, [&](const auto& Concrete, const auto& FullTM)
+			{
+				FVec3 SupportLocal = Concrete.Support(FullTM.InverseTransformVectorNoScale(DirX), 1);
+				return FullTM.TransformPosition(SupportLocal);
+			});
 			EXPECT_VECTOR_NEAR_DEFAULT(SupportPt, TVector3(9, 5, 5));
 
-			SupportPt = Translated.Support(DirX, 1);
+			SupportPt = Utilities::CastHelper(Translated, Identity, [&](const auto& Concrete, const auto& FullTM)
+			{
+				FVec3 SupportLocal = Concrete.Support(FullTM.InverseTransformVectorNoScale(DirX), 1);
+				return FullTM.TransformPosition(SupportLocal);
+			});
 			EXPECT_VECTOR_NEAR_DEFAULT(SupportPt, TVector3(13, 5, 5));
 
 			const TVector<T, 3> DirZ(0, 0, -1);
-			SupportPt = Translated.Support(DirZ, 1);
+			SupportPt = Utilities::CastHelper(Translated, Identity, [&](const auto& Concrete, const auto& FullTM)
+			{
+				FVec3 SupportLocal = Concrete.Support(FullTM.InverseTransformVectorNoScale(DirZ), 1);
+				return FullTM.TransformPosition(SupportLocal);
+			});
 			EXPECT_VECTOR_NEAR_DEFAULT(SupportPt, TVector3(12, 5, -6));
 
-			SupportPt = Rotated.Support(DirZ, 1);
+			SupportPt = Utilities::CastHelper(Rotated, Identity, [&](const auto& Concrete, const auto& FullTM)
+			{
+				FVec3 SupportLocal = Concrete.Support(FullTM.InverseTransformVectorNoScale(DirZ), 1);
+				return FullTM.TransformPosition(SupportLocal);
+			});
 			EXPECT_VECTOR_NEAR_DEFAULT(SupportPt, TVector3(-5, 8, -6)); // @todo why -5?
 
-			SupportPt = Transformed.Support(DirZ, 1);
+			SupportPt = Utilities::CastHelper(Transformed, Identity, [&](const auto& Concrete, const auto& FullTM)
+			{
+				FVec3 SupportLocal = Concrete.Support(FullTM.InverseTransformVectorNoScale(DirZ), 1);
+				FVec3 TransformedPt = FullTM.TransformPosition(SupportLocal);
+				return TransformedPt;
+			});
 			EXPECT_VECTOR_NEAR_DEFAULT(SupportPt, TVector3(-1, 8, -6));
 		}
 
@@ -782,7 +804,6 @@ namespace ChaosTest {
 			Result = Translated.FindClosestIntersection(TVector3(7, 0, 4.5), TVector3(7, 0, 5.5), KINDA_SMALL_NUMBER);
 			EXPECT_VECTOR_NEAR(Result.First, TVector3(7, 0, 5), 0.001);
 		}
-#endif
 	}
 	template void ImplicitTransformed<float>();
 
@@ -859,16 +880,15 @@ namespace ChaosTest {
 	template<class T>
 	void ImplicitUnion()
 	{
-#if 0
 		typedef TVector<T, 3> TVector3;
 		FString Caller("ImplicitUnion()");
-		TUniquePtr<TImplicitObjectUnion<T, 3>> MUnionedObjects;
+		TUniquePtr<FImplicitObjectUnion> MUnionedObjects;
 
 		{// unit cylinder - sanity check
 			TArray<TUniquePtr<FImplicitObject>> Objects;
 			Objects.Add(MakeUnique<TCylinder<T>>(TVector<T, 3>(0, 0, 1), TVector3(0), 1));
 			Objects.Add(MakeUnique<TCylinder<T>>(TVector<T, 3>(0, 0, -1), TVector3(0), 1));
-			MUnionedObjects.Reset(new Chaos::TImplicitObjectUnion<float, 3>(std::move(Objects)));
+			MUnionedObjects.Reset(new Chaos::FImplicitObjectUnion(std::move(Objects)));
 
 			// Can't use the default internal unit tests because they expect different behavior internally where the two cylinders are joined together. 
 			EXPECT_VECTOR_NEAR(MUnionedObjects->Normal(TVector3(0, 0, 2 / 3.)), (TVector3(0, 0, 1)), KINDA_SMALL_NUMBER);
@@ -893,7 +913,7 @@ namespace ChaosTest {
 		TArray<TUniquePtr<FImplicitObject>> Objects;
 		Objects.Add(MakeUnique<TCylinder<T>>(TVector<T, 3>(0, 0, -2), TVector3(0, 0, 2), 1));
 		Objects.Add(MakeUnique<TCylinder<T>>(TVector<T, 3>(0, -2, 0), TVector3(0, 2, 0), 1));
-		MUnionedObjects.Reset(new Chaos::TImplicitObjectUnion<float, 3>(std::move(Objects)));
+		MUnionedObjects.Reset(new Chaos::FImplicitObjectUnion(std::move(Objects)));
 
 		{// closest point near origin (+)
 			EXPECT_NEAR(MUnionedObjects->SignedDistance(TVector3(0, 0, 9 / 4.)), 1 / 4., KINDA_SMALL_NUMBER);
@@ -918,7 +938,7 @@ namespace ChaosTest {
 		TArray<TUniquePtr<FImplicitObject>> Objects2;
 		Objects2.Add(MakeUnique<TCylinder<T>>(TVector<T, 3>(4, 4, 2), TVector3(4, 4, 6), 1));
 		Objects2.Add(MakeUnique<TCylinder<T>>(TVector<T, 3>(4, 2, 4), TVector3(4, 6, 4), 1));
-		MUnionedObjects.Reset(new Chaos::TImplicitObjectUnion<float, 3>(std::move(Objects2)));
+		MUnionedObjects.Reset(new Chaos::FImplicitObjectUnion(std::move(Objects2)));
 
 		{// closest point off origin (+)
 			EXPECT_NEAR(MUnionedObjects->SignedDistance(TVector3(4, 4, 4 + 9 / 4.)), 1 / 4., KINDA_SMALL_NUMBER);
@@ -946,7 +966,7 @@ namespace ChaosTest {
 			TArray<TUniquePtr<FImplicitObject>> Unions;
 			Unions.Add(MakeUnique<TCapsule<T>>(TVector<T, 3>(0, 0, 0), TVector<T, 3>(0, 0, -2), 1));
 			Unions.Add(MakeUnique<TCapsule<T>>(TVector<T, 3>(0, 0, 0), TVector<T, 3>(0, 0, 2), 1));
-			MUnionedObjects.Reset(new Chaos::TImplicitObjectUnion<float, 3>(std::move(Unions)));
+			MUnionedObjects.Reset(new Chaos::FImplicitObjectUnion(std::move(Unions)));
 
 			EXPECT_VECTOR_NEAR(MUnionedObjects->Normal(TVector3(0, 0, 7 / 3.)), (TVector3(0, 0, 1)), KINDA_SMALL_NUMBER);
 			EXPECT_VECTOR_NEAR(MUnionedObjects->Normal(TVector3(0, 0, -7 / 3.)), (TVector3(0, 0, -1)), KINDA_SMALL_NUMBER);
@@ -968,9 +988,9 @@ namespace ChaosTest {
 			Objects1.Add(MakeUnique<TTaperedCylinder<T>>(TVector<T, 3>(0, 0, 1), TVector<T, 3>(0, 0, -1), 1, 1));
 
 			TArray<TUniquePtr<FImplicitObject>> Unions;
-			Unions.Emplace(new TImplicitObjectUnion<T, 3>(MoveTemp(Objects1)));
-			TUniquePtr<TImplicitObjectUnion<T, 3>> UnionedUnions;
-			UnionedUnions.Reset(new Chaos::TImplicitObjectUnion<float, 3>(std::move(Unions)));
+			Unions.Emplace(new FImplicitObjectUnion(MoveTemp(Objects1)));
+			TUniquePtr<FImplicitObjectUnion> UnionedUnions;
+			UnionedUnions.Reset(new Chaos::FImplicitObjectUnion(std::move(Unions)));
 
 			UnitImplicitObjectNormalsExternal<T>(*UnionedUnions, FString("ImplicitUnion() - nested union unit cylinder 1"));
 			UnitImplicitObjectNormalsInternal<T>(*UnionedUnions, FString("ImplicitUnion() - nested union unit cylinder 1"));
@@ -986,10 +1006,10 @@ namespace ChaosTest {
 			ObjectsB.Add(MakeUnique<TTaperedCylinder<T>>(TVector<T, 3>(0, 0, 1), TVector<T, 3>(0, 0, -1), 1, 1));
 
 			TArray<TUniquePtr<FImplicitObject>> Unions;
-			Unions.Emplace(new TImplicitObjectUnion<T, 3>(MoveTemp(ObjectsA)));
-			Unions.Emplace(new TImplicitObjectUnion<T, 3>(MoveTemp(ObjectsB)));
-			TUniquePtr<TImplicitObjectUnion<T, 3>> UnionedUnions;
-			UnionedUnions.Reset(new Chaos::TImplicitObjectUnion<float, 3>(std::move(Unions)));
+			Unions.Emplace(new FImplicitObjectUnion(MoveTemp(ObjectsA)));
+			Unions.Emplace(new FImplicitObjectUnion(MoveTemp(ObjectsB)));
+			TUniquePtr<FImplicitObjectUnion> UnionedUnions;
+			UnionedUnions.Reset(new Chaos::FImplicitObjectUnion(std::move(Unions)));
 
 			UnitImplicitObjectNormalsExternal<T>(*UnionedUnions, FString("ImplicitUnion() - nested union unit sphere 1"));
 			UnitImplicitObjectNormalsInternal<T>(*UnionedUnions, FString("ImplicitUnion() - nested union unit sphere 1"));
@@ -1002,10 +1022,10 @@ namespace ChaosTest {
 			ObjectsA.Add(MakeUnique<TCylinder<T>>(TVector<T, 3>(0, 0, 0), TVector<T, 3>(0, 0, -1), 1));
 			ObjectsB.Add(MakeUnique<TCylinder<T>>(TVector<T, 3>(0, 0, 0), TVector<T, 3>(0, 0, 1), 1));
 			TArray<TUniquePtr<FImplicitObject>> Unions;
-			Unions.Emplace(new TImplicitObjectUnion<T, 3>(MoveTemp(ObjectsA)));
-			Unions.Emplace(new TImplicitObjectUnion<T, 3>(MoveTemp(ObjectsB)));
-			TUniquePtr<TImplicitObjectUnion<T, 3>> UnionedUnions;
-			UnionedUnions.Reset(new Chaos::TImplicitObjectUnion<float, 3>(std::move(Unions)));
+			Unions.Emplace(new FImplicitObjectUnion(MoveTemp(ObjectsA)));
+			Unions.Emplace(new FImplicitObjectUnion(MoveTemp(ObjectsB)));
+			TUniquePtr<FImplicitObjectUnion> UnionedUnions;
+			UnionedUnions.Reset(new Chaos::FImplicitObjectUnion(std::move(Unions)));
 
 			UnitImplicitObjectNormalsExternal<T>(*UnionedUnions, FString("ImplicitUnion() - nested union unit cylinder 2"));
 
@@ -1022,7 +1042,7 @@ namespace ChaosTest {
 			// Distance is 0 at the joined faces.
 			EXPECT_NEAR(UnionedUnions->SignedDistance(TVector3(0, 0, 0)), 0., KINDA_SMALL_NUMBER);
 		}
-#endif
+
 	}
 	template void ImplicitUnion<float>();
 
