Add unit test demonstrating how slight difference in rotation causes vastly different normal in a initially overlapping sweep.

diff --git a/Engine/Source/Programs/HeadlessChaos/Private/HeadlessChaosTestEPA.cpp b/Engine/Source/Programs/HeadlessChaos/Private/HeadlessChaosTestEPA.cpp
--- a/Engine/Source/Programs/HeadlessChaos/Private/HeadlessChaosTestEPA.cpp
+++ b/Engine/Source/Programs/HeadlessChaos/Private/HeadlessChaosTestEPA.cpp
@@ -10,6 +10,7 @@
 #include "Chaos/GJK.h"
 #include "Chaos/Convex.h"
 #include "Chaos/ImplicitObjectScaled.h"
+#include "Chaos/TriangleMeshImplicitObject.h"
 #include "Chaos/Triangle.h"
 
 namespace ChaosTest
@@ -1024,4 +1025,68 @@ namespace ChaosTest
 		EXPECT_GT(FMath::Abs(Normal.Z), 0.8f);
 	}
 
+
+		//
+		// Performs the same initially overlapping sweep twice, with slightly different rotations, gives different normals.
+		// this is convex box slightly penetrating surface of triangle mesh. Trimesh normals point up.
+		//
+		GTEST_TEST(EPATests, EPARealFailures_ConvexTrimeshRotationalDifferencesBreakNormal)
+		{
+			using namespace Chaos;
+			TParticles<FReal, 3> ConvexBoxSurfaceParticles(
+				{
+				{50.0999985, -50.1124992, -50.1250000},
+				{-50.0999985, 50.1250000, -50.1250000},
+				{50.0999985, 50.1250000, -50.0999985},
+				{50.0999985, 50.1250000, 50.1250000},
+				{-50.0999985, -50.1124992, -50.0999985},
+				{-50.0999985, -50.1124992, 50.1250000},
+				{50.0999985, -50.1124992, 50.1250000},
+				{-50.0999985, 50.1250000, 50.1250000},
+				});
+
+			FConvex ConvexBox(MoveTemp(ConvexBoxSurfaceParticles));
+
+			TParticles<FReal, 3> TrimeshParticles(
+				{
+					{50.0000000, 50.0000000, -8.04061356e-15},
+					{50.0000000, -50.0000000, 8.04061356e-15},
+					{-50.0000000, 50.0000000, -8.04061356e-15},
+					{-50.0000000, -50.0000000, 8.04061356e-15}
+				});
+
+			TArray<TVector<int32, 3>> Indices;
+			Indices.Emplace(1, 0, 2);
+			Indices.Emplace(1, 2, 3);
+
+			TArray<uint16> Materials;
+			Materials.Emplace(0);
+			Materials.Emplace(0);
+			TUniquePtr<FTriangleMeshImplicitObject> TriangleMesh = MakeUnique<FTriangleMeshImplicitObject>(MoveTemp(TrimeshParticles), MoveTemp(Indices), MoveTemp(Materials));
+			TImplicitObjectScaled<FTriangleMeshImplicitObject> ScaledTriangleMesh = TImplicitObjectScaled<FTriangleMeshImplicitObject>(MakeSerializable(TriangleMesh), FVec3(11.5, 11.5, 11.5));
+
+			FQuat Rotation0(0.00488796039, 0.00569311855, -0.000786740216, 0.999971569);
+			FQuat Rotation1(0.0117356628, -0.0108017093, -0.000888462295, 0.999872327);
+
+			FVec3 Translation(309.365723, -69.4132690, 51.2289352);
+
+			TRigidTransform<FReal, 3> Transform0(Translation, Rotation0);
+			TRigidTransform<FReal, 3> Transform1(Translation, Rotation1);
+
+			FVec3 Dir(-0.00339674903, 5.76980747e-05, -0.999994159);
+			FReal Length = 1.83530724;
+
+			FReal OutTime = -1;
+			FVec3 Normal(0.0f);
+			FVec3 Position(0.0f);
+			int32 FaceIndex = -1;
+			bool bResult = ScaledTriangleMesh.LowLevelSweepGeom(ConvexBox, Transform0, Dir, Length, OutTime, Position, Normal, FaceIndex, 0.0f, true);
+
+			bResult = ScaledTriangleMesh.LowLevelSweepGeom(ConvexBox, Transform1, Dir, Length, OutTime, Position, Normal, FaceIndex, 0.0f, true);
+			
+			// Observe that normals are in opposite direction, while rotations are very similar.
+
+		}
+
+
 }
\ No newline at end of file
