Repro Player/Rockwall collision issue in unit test.

diff --git a/Engine/Source/Programs/HeadlessChaos/Private/HeadlessChaosTestEPA.cpp b/Engine/Source/Programs/HeadlessChaos/Private/HeadlessChaosTestEPA.cpp
--- a/Engine/Source/Programs/HeadlessChaos/Private/HeadlessChaosTestEPA.cpp
+++ b/Engine/Source/Programs/HeadlessChaos/Private/HeadlessChaosTestEPA.cpp
@@ -1088,5 +1088,53 @@ namespace ChaosTest
 
 		}
 
+		//
+		// Player can clip through RockWall trimesh, this repros a failure, MTD seems wrong, pushes 
+		//
+		GTEST_TEST(EPATests, EPARealFailures_CapsuleVsTrimeshRockWallWrongNormal)
+		{
+			using namespace Chaos;
+			TParticles<FReal, 3> TrimeshParticles({
+				{-30.2797413, 89.3729248, 12.6687469},
+				{-48.5687141, 81.4716797, 12.5528765},
+				{-9.00777531, 44.7844124, 12.8065948}
+
+			});
+
+			TRigidTransform<FReal, 3> TriMeshTransform(FVec3(0, 0, 936), FQuat(-0.642856, 0.402588, -0.648765, -0.061344), FVec3(1));
+			TRigidTransform<FReal, 3> CapsuleTransform(FVec3(-560.400146, -738.941528, 75.013588), FQuat(0.000000, 0.000000, 0.976842, 0.213964), FVec3(1));
+
+			FVec3 ExpectedNormal = FVec3::CrossProduct(TrimeshParticles.X(1) - TrimeshParticles.X(0), TrimeshParticles.X(2) - TrimeshParticles.X(0));
+			ExpectedNormal.Normalize();
+			ExpectedNormal = TriMeshTransform.TransformVector(ExpectedNormal);
+
+			TArray<TVector<int32, 3>> Indices;
+			// 63,65,50
+			Indices.Emplace(0, 1, 2);
+
+			TArray<uint16> Materials;
+			Materials.Emplace(0);
+
+			TUniquePtr<FTriangleMeshImplicitObject> TriangleMesh = MakeUnique<FTriangleMeshImplicitObject>(MoveTemp(TrimeshParticles), MoveTemp(Indices), MoveTemp(Materials));
+			TImplicitObjectScaled<FTriangleMeshImplicitObject> ScaledTriangleMesh = TImplicitObjectScaled<FTriangleMeshImplicitObject>(MakeSerializable(TriangleMesh), FVec3(10.109713, 18.734829, 9.246257));
+
+			TCapsule<FReal> Capsule(TVec3<FReal>(0, 0, -33), TVec3<FReal>(0, 0, 33), 42);
+			TRigidTransform<FReal, 3> BToATM = CapsuleTransform.GetRelativeTransform(TriMeshTransform);
+
+
+			const FVec3 Dir(-0.834976, 0.550286, 0.000000);
+			const FReal Length = 19.977188;
 
+			FReal OutTime = -1.0f;
+			FVec3 Normal(0.0f);
+			FVec3 Position(0.0f);
+			int32 FaceIndex = -1;
+
+
+			bool bResult = ScaledTriangleMesh.LowLevelSweepGeom(Capsule, BToATM, Dir, Length, OutTime, Position, Normal, FaceIndex, 0.0f, true);
+
+			Normal = TriMeshTransform.TransformVectorNoScale(Normal);
+			Position = TriMeshTransform.TransformPositionNoScale(Position);
+			
+		}
 }
\ No newline at end of file
