Added failure case

diff --git a/Engine/Source/Programs/HeadlessChaos/Private/HeadlessChaosTestEPA.cpp b/Engine/Source/Programs/HeadlessChaos/Private/HeadlessChaosTestEPA.cpp
--- a/Engine/Source/Programs/HeadlessChaos/Private/HeadlessChaosTestEPA.cpp
+++ b/Engine/Source/Programs/HeadlessChaos/Private/HeadlessChaosTestEPA.cpp
@@ -1135,6 +1135,8 @@ namespace ChaosTest
 
 			Normal = TriMeshTransform.TransformVectorNoScale(Normal);
 			Position = TriMeshTransform.TransformPositionNoScale(Position);
+
+			EXPECT_GT(FVec3::DotProduct(ExpectedNormal, Normal), 0);
 			
 		}
 }
\ No newline at end of file
