Update unit test for GJKRaycast2 ClosestB computation being wrong (capsule sweep vs RockWall trimesh).

diff --git a/Engine/Source/Programs/HeadlessChaos/Private/HeadlessChaosTestEPA.cpp b/Engine/Source/Programs/HeadlessChaos/Private/HeadlessChaosTestEPA.cpp
--- a/Engine/Source/Programs/HeadlessChaos/Private/HeadlessChaosTestEPA.cpp
+++ b/Engine/Source/Programs/HeadlessChaos/Private/HeadlessChaosTestEPA.cpp
@@ -1089,54 +1089,56 @@ namespace ChaosTest
 		}
 
 		//
-		// Player can clip through RockWall trimesh, this repros a failure, MTD seems wrong, pushes 
+		// Player can clip through RockWall trimesh, this repros a failure, MTD seems wrong.
+		// This is failing GJKRaycast2 call.
+		// Fixed: (11457046) ClosestB computation was transformed wrong messing up normal.
 		//
-		GTEST_TEST(EPATests, EPARealFailures_CapsuleVsTrimeshRockWallWrongNormal)
+		GTEST_TEST(EPATests, EPARealFailures_CapsuleVsTrimeshRockWallWrongNormalGJKRaycast2)
 		{
 			using namespace Chaos;
-			TParticles<FReal, 3> TrimeshParticles({
-				{-30.2797413, 89.3729248, 12.6687469},
-				{-48.5687141, 81.4716797, 12.5528765},
-				{-9.00777531, 44.7844124, 12.8065948}
-
-			});
+			// Triangle w/ world scale
+			TTriangle<FReal> Triangle({
+				{-306.119476, 1674.38647, 117.138489},
+				{-491.015747, 1526.35803, 116.067123},
+				{-91.0660172, 839.028320, 118.413063}
+				});
 
-			TRigidTransform<FReal, 3> TriMeshTransform(FVec3(0, 0, 936), FQuat(-0.642856, 0.402588, -0.648765, -0.061344), FVec3(1));
-			TRigidTransform<FReal, 3> CapsuleTransform(FVec3(-560.400146, -738.941528, 75.013588), FQuat(0.000000, 0.000000, 0.976842, 0.213964), FVec3(1));
 
-			FVec3 ExpectedNormal = FVec3::CrossProduct(TrimeshParticles.X(1) - TrimeshParticles.X(0), TrimeshParticles.X(2) - TrimeshParticles.X(0));
+			FVec3 ExpectedNormal = FVec3::CrossProduct(Triangle[1] - Triangle[0], Triangle[2] - Triangle[0]);
 			ExpectedNormal.Normalize();
-			ExpectedNormal = TriMeshTransform.TransformVector(ExpectedNormal);
 
-			TArray<TVector<int32, 3>> Indices;
-			// 63,65,50
-			Indices.Emplace(0, 1, 2);
+			TRigidTransform<FReal, 3> StartTM(FVec3(-344.031799, 1210.37158, 134.252747), FQuat(-0.255716801, -0.714108050, 0.0788889676, -0.646866322), FVec3(1));
 
-			TArray<uint16> Materials;
-			Materials.Emplace(0);
+			// Wrapping in 1,1,1 scale is unnecessary, but this is technically what is happening when sweeping against scaled trimesh.
+			TUniquePtr<TCapsule<FReal>> Capsule = MakeUnique<TCapsule<FReal>>(TVec3<FReal>(0, 0, -33), TVec3<FReal>(0, 0, 33), 42);
+			TImplicitObjectScaled<TCapsule<FReal>> ScaledCapsule = TImplicitObjectScaled<TCapsule<FReal>>(MakeSerializable(Capsule), FVec3(1));
 
-			TUniquePtr<FTriangleMeshImplicitObject> TriangleMesh = MakeUnique<FTriangleMeshImplicitObject>(MoveTemp(TrimeshParticles), MoveTemp(Indices), MoveTemp(Materials));
-			TImplicitObjectScaled<FTriangleMeshImplicitObject> ScaledTriangleMesh = TImplicitObjectScaled<FTriangleMeshImplicitObject>(MakeSerializable(TriangleMesh), FVec3(10.109713, 18.734829, 9.246257));
 
-			TCapsule<FReal> Capsule(TVec3<FReal>(0, 0, -33), TVec3<FReal>(0, 0, 33), 42);
-			TRigidTransform<FReal, 3> BToATM = CapsuleTransform.GetRelativeTransform(TriMeshTransform);
-
-
-			const FVec3 Dir(-0.834976, 0.550286, 0.000000);
-			const FReal Length = 19.977188;
+			const FVec3 Dir(-0.102473199, 0.130887285, -0.986087084);
+			const FReal LengthScale = 9.31486130;
+			const FReal  CurrentLength = 2.14465737;
+			const FReal Length = LengthScale * CurrentLength;
+			const bool bComputeMTD = true;
+			const FReal Thickness = 0;
 
 			FReal OutTime = -1.0f;
 			FVec3 Normal(0.0f);
 			FVec3 Position(0.0f);
 			int32 FaceIndex = -1;
 
+			// This is local to trimesh, world scale.
+			bool bResult = GJKRaycast2<FReal>(Triangle, ScaledCapsule, StartTM, Dir, Length, OutTime, Position, Normal, Thickness, bComputeMTD);
 
-			bool bResult = ScaledTriangleMesh.LowLevelSweepGeom(Capsule, BToATM, Dir, Length, OutTime, Position, Normal, FaceIndex, 0.0f, true);
+			// Compare results against GJKPenetration, sweep is initial overlap, so this should be the same.
+			FVec3 Normal2, ClosestA, ClosestB;
+			FReal OutTime2;
+			bool bResult2 = GJKPenetration(Triangle, ScaledCapsule, StartTM, OutTime2, ClosestA, ClosestB, Normal2);
 
-			Normal = TriMeshTransform.TransformVectorNoScale(Normal);
-			Position = TriMeshTransform.TransformPositionNoScale(Position);
 
-			EXPECT_GT(FVec3::DotProduct(ExpectedNormal, Normal), 0);
-			
+			EXPECT_VECTOR_NEAR(Normal, Normal2, KINDA_SMALL_NUMBER);
+			EXPECT_NEAR(OutTime, -OutTime2, KINDA_SMALL_NUMBER);
+
+			const FVec3 ClosestBShouldBe{ -287.344025, 1211.66296, 101.851364 };
+			EXPECT_VECTOR_NEAR(ClosestB, ClosestBShouldBe, KINDA_SMALL_NUMBER);
 		}
 }
\ No newline at end of file
