- Test case: Add one failure case for EPA
- GJK: Comment about using -V instread of B - A for increased accuracy

diff --git a/Engine/Source/Programs/HeadlessChaos/Private/HeadlessChaosTestEPA.cpp b/Engine/Source/Programs/HeadlessChaos/Private/HeadlessChaosTestEPA.cpp
--- a/Engine/Source/Programs/HeadlessChaos/Private/HeadlessChaosTestEPA.cpp
+++ b/Engine/Source/Programs/HeadlessChaos/Private/HeadlessChaosTestEPA.cpp
@@ -1025,6 +1025,29 @@ namespace ChaosTest
 		EXPECT_GT(FMath::Abs(Normal.Z), 0.8f);
 	}
 
+	GTEST_TEST(EPATests, DISABLED_EPA_EdgeCases)
+	{
+		TBox<FReal, 3> A(FVec3(0.0f, 0.0f, 0.0f), FVec3(100.0f, 100.0f, 100.0f));
+		TBox<FReal, 3> B(FVec3(100.0f + 0.00001f, 0.0f, 0.0f), FVec3(200.0f, 100.0f, 100.0f)); // This small epsilon can actually be zero as well.
+		FRigidTransform3 ATM = FRigidTransform3(FVec3(0.0f, 0.0f, 0.0f), FRotation3::FromElements(0.0f, 0.0f, 0.0f, 1.0f));
+		FRigidTransform3 BTM = FRigidTransform3(FVec3(0.0f, 0.0f, 0.0f), FRotation3::FromElements(0.0f, 0.0f, 0.0f, 1.0f));
+		FVec3 InitialDir = FVec3(1.0f, 0.0f, 0.0f);
+
+		const FRigidTransform3 BToATM = BTM.GetRelativeTransform(ATM);
+		FReal Penetration;
+		FVec3 ClosestA, ClosestB, NormalA;
+		int32 NumIterations = 0;
+		GJKPenetration<true>(A, B, BToATM, Penetration, ClosestA, ClosestB, NormalA, 0.0f, InitialDir, 0.0f, &NumIterations);
+
+		FVec3 Location = ATM.TransformPosition(ClosestA);
+		FVec3 Normal = ATM.TransformVectorNoScale(NormalA);
+		FReal Phi = -Penetration;
+
+		EXPECT_NEAR(Normal.X, 1.0f, 1e-3f);
+		EXPECT_NEAR(Normal.Y, 0.0f, 1e-3f);
+		EXPECT_NEAR(Normal.Z, 0.0f, 1e-3f);
+	}
+
 
 		//
 		// Performs the same initially overlapping sweep twice, with slightly different rotations, gives different normals.
diff --git a/Engine/Source/Runtime/Experimental/Chaos/Public/Chaos/GJK.h b/Engine/Source/Runtime/Experimental/Chaos/Public/Chaos/GJK.h
--- a/Engine/Source/Runtime/Experimental/Chaos/Public/Chaos/GJK.h
+++ b/Engine/Source/Runtime/Experimental/Chaos/Public/Chaos/GJK.h
@@ -187,7 +187,7 @@ namespace Chaos
 
 			
 			const T PreDist = FMath::Sqrt(PrevDist2);
-			OutNormal = (ClosestBInA - ClosestA).GetUnsafeNormal();	//question: should we just use PreDist2?
+			OutNormal = (ClosestBInA - ClosestA).GetUnsafeNormal();	//Note1: should we just use PreDist2? //Note2: we can just use -V.GetUnsafeNormal() here instead if it improves accuracy
 			T Penetration = ThicknessA + ThicknessB - PreDist;
 			if (!bNegativePenetrationSupport)
 			{
