GJKpenetration edge case: in some rare cases this function can report penetration when there is none

diff --git a/Engine/Source/Programs/HeadlessChaos/Private/HeadlessChaosTestGJK.cpp b/Engine/Source/Programs/HeadlessChaos/Private/HeadlessChaosTestGJK.cpp
--- a/Engine/Source/Programs/HeadlessChaos/Private/HeadlessChaosTestGJK.cpp
+++ b/Engine/Source/Programs/HeadlessChaos/Private/HeadlessChaosTestGJK.cpp
@@ -1535,6 +1535,36 @@ namespace ChaosTest
 			EXPECT_VECTOR_NEAR(Normal,Normal2,1e-4);
 			EXPECT_VECTOR_NEAR(Position,Position2,1e-3);
 		}
+
+		
+		{
+			// For this test we are clearly not penetrating
+			// but we had an actual bug (edge condition) that showed we are
+
+			const TVec3<T> Pt0(0.0, 0.0, 0.0);
+			TVec3<T> Pt1(100.0,0,0);
+			TVec3<T> Pt2(0, 1000000.0, 0);
+			
+
+			const TCapsule<T> A(Pt1, Pt2, 1.0);
+			const TSphere<T, 3> B(Pt0, 1.0);
+
+			const TRigidTransform<T, 3> BToATM(TVec3<T>(0, 0, 0), TRotation<T, 3>::FromElements(0.0, 0.0, 0, 1)); // Unit transform
+			const TVec3<T> InitDir(0.1, 0.0, 0.0);
+
+			T Penetration;
+			TVec3<T> ClosestA, ClosestB, Normal;			
+
+			// First demonstrate the distance between the shapes are more than 90cm.
+			bool IsValid = GJKPenetration<true, T>(A, B, BToATM, Penetration, ClosestA, ClosestB, Normal, 0, InitDir, 0);
+			EXPECT_TRUE(IsValid);
+			EXPECT_TRUE(Penetration < -90.0f);
+
+			// Since there is no penetration (by more than 90cm) this function should return false when negative penetration is not supported
+			bool IsPenetrating = GJKPenetration<false, T>(A, B, BToATM, Penetration, ClosestA, ClosestB, Normal, 0, InitDir, 0);
+			EXPECT_FALSE(IsPenetrating);
+			
+		}
 		
 	}
 
diff --git a/Engine/Source/Runtime/Experimental/Chaos/Public/Chaos/GJK.h b/Engine/Source/Runtime/Experimental/Chaos/Public/Chaos/GJK.h
--- a/Engine/Source/Runtime/Experimental/Chaos/Public/Chaos/GJK.h
+++ b/Engine/Source/Runtime/Experimental/Chaos/Public/Chaos/GJK.h
@@ -191,7 +191,7 @@ namespace Chaos
 			T Penetration = ThicknessA + ThicknessB - PreDist;
 			if (!bNegativePenetrationSupport)
 			{
-				Penetration = FMath::Clamp<T>(Penetration, 0, TNumericLimits<T>::Max());
+				if (Penetration < 0.0f) return false;
 			}
 			OutPenetration = Penetration;
 			OutClosestA = ClosestA + OutNormal * ThicknessA;
