Fix headless chaos compile.

diff --git a/Engine/Source/Programs/HeadlessChaos/Private/HeadlessChaosTestBroadphase.cpp b/Engine/Source/Programs/HeadlessChaos/Private/HeadlessChaosTestBroadphase.cpp
--- a/Engine/Source/Programs/HeadlessChaos/Private/HeadlessChaosTestBroadphase.cpp
+++ b/Engine/Source/Programs/HeadlessChaos/Private/HeadlessChaosTestBroadphase.cpp
@@ -654,7 +654,8 @@ namespace ChaosTest
 		Particle->GTGeometryParticle() = TGeometryParticle<FReal, 3>::CreateParticle().Release();
 		Particle->X() = TVector<FReal, 3>(0, 0, 0);
 
-		FPBDRigidsEvolutionGBF Evolution(Particles);
+		THandleArray<FChaosPhysicsMaterial> PhysicalMaterials;
+		FPBDRigidsEvolutionGBF Evolution(Particles, PhysicalMaterials);
 
 		// Flush spatial acceleration structures to put particle into structure.
 		Evolution.FlushSpatialAcceleration();
diff --git a/Engine/Source/Programs/HeadlessChaos/Private/HeadlessChaosTestClustering.cpp b/Engine/Source/Programs/HeadlessChaos/Private/HeadlessChaosTestClustering.cpp
--- a/Engine/Source/Programs/HeadlessChaos/Private/HeadlessChaosTestClustering.cpp
+++ b/Engine/Source/Programs/HeadlessChaos/Private/HeadlessChaosTestClustering.cpp
@@ -26,7 +26,8 @@ namespace ChaosTest {
 	{
 #if TODO_REIMPLEMENT_RIGID_CLUSTERING
 		TPBDRigidsSOAs<T, 3> Particles;
-		FPBDRigidsEvolutionGBF Evolution(Particles);
+		THandleArray<FChaosPhysicsMaterial> PhysicalMaterials;
+		FPBDRigidsEvolutionGBF Evolution(Particles, PhysicalMaterials);
 		auto& ClusteredParticles = Particles.GetClusteredParticles();
 		
 		uint32 FirstId = ClusteredParticles.Size();
@@ -70,7 +71,8 @@ namespace ChaosTest {
 #if TODO_REIMPLEMENT_RIGID_CLUSTERING
 
 		TPBDRigidsSOAs<T, 3> Particles;
-		FPBDRigidsEvolutionGBF Evolution(Particles);
+		THandleArray<FChaosPhysicsMaterial> PhysicalMaterials;
+		FPBDRigidsEvolutionGBF Evolution(Particles, PhysicalMaterials);
 		auto& ClusteredParticles = Particles.GetClusteredParticles();
 
 		//create a long row of boxes - the depth 0 cluster is the entire row, the depth 1 clusters 4 boxes each, the depth 2 clusters are 1 box each
@@ -187,7 +189,8 @@ namespace ChaosTest {
 	{
 #if TODO_REIMPLEMENT_RIGID_CLUSTERING
 		TPBDRigidsSOAs<T, 3> Particles;
-		FPBDRigidsEvolutionGBF Evolution(Particles);
+		THandleArray<FChaosPhysicsMaterial> PhysicalMaterials;
+		FPBDRigidsEvolutionGBF Evolution(Particles, PhysicalMaterials);
 		auto& ClusteredParticles = Particles.GetClusteredParticles();
 
 		//create a long row of boxes - the depth 0 cluster is the entire row, the depth 1 clusters 4 boxes each, the depth 2 clusters are 1 box each
diff --git a/Engine/Source/Programs/HeadlessChaos/Private/HeadlessChaosTestConstraintHandle.cpp b/Engine/Source/Programs/HeadlessChaos/Private/HeadlessChaosTestConstraintHandle.cpp
--- a/Engine/Source/Programs/HeadlessChaos/Private/HeadlessChaosTestConstraintHandle.cpp
+++ b/Engine/Source/Programs/HeadlessChaos/Private/HeadlessChaosTestConstraintHandle.cpp
@@ -110,7 +110,8 @@ namespace ChaosTest {
 	void JointConstraintHandles()
 	{
 		TPBDRigidsSOAs<FReal, 3> Particles;
-		TEvolution Evolution(Particles);
+		THandleArray<FChaosPhysicsMaterial> PhysicalMaterials;
+		TEvolution Evolution(Particles, PhysicalMaterials);
 
 		TArray<TPBDRigidParticleHandle<FReal, 3>*> ParticleHandles = Evolution.CreateDynamicParticles(5);
 
@@ -130,7 +131,8 @@ namespace ChaosTest {
 	void PositionConstraintHandles()
 	{
 		TPBDRigidsSOAs<FReal, 3> Particles;
-		TEvolution Evolution(Particles);
+		THandleArray<FChaosPhysicsMaterial> PhysicalMaterials;
+		TEvolution Evolution(Particles, PhysicalMaterials);
 		TArray<TPBDRigidParticleHandle<FReal, 3>*> ParticleHandles = Evolution.CreateDynamicParticles(5);
 
 		TPBDPositionConstraints<FReal, 3> Constraints;
@@ -149,7 +151,8 @@ namespace ChaosTest {
 	void RigidSpringConstraintHandles()
 	{
 		TPBDRigidsSOAs<FReal, 3> Particles;
-		TEvolution Evolution(Particles);
+		THandleArray<FChaosPhysicsMaterial> PhysicalMaterials;
+		TEvolution Evolution(Particles, PhysicalMaterials);
 		TArray<TPBDRigidParticleHandle<FReal, 3>*> ParticleHandles = Evolution.CreateDynamicParticles(5);
 
 		FPBDRigidSpringConstraints Constraints;
@@ -169,7 +172,8 @@ namespace ChaosTest {
 	void RigidDynamicSpringConstraintHandles()
 	{
 		TPBDRigidsSOAs<FReal, 3> Particles;
-		TEvolution Evolution(Particles);
+		THandleArray<FChaosPhysicsMaterial> PhysicalMaterials;
+		TEvolution Evolution(Particles, PhysicalMaterials);
 		TArray<TPBDRigidParticleHandle<FReal, 3>*> ParticleHandles = Evolution.CreateDynamicParticles(5);
 
 		TPBDRigidDynamicSpringConstraints<FReal, 3> Constraints;
diff --git a/Engine/Source/Programs/HeadlessChaos/Private/HeadlessChaosTestConstraints.cpp b/Engine/Source/Programs/HeadlessChaos/Private/HeadlessChaosTestConstraints.cpp
--- a/Engine/Source/Programs/HeadlessChaos/Private/HeadlessChaosTestConstraints.cpp
+++ b/Engine/Source/Programs/HeadlessChaos/Private/HeadlessChaosTestConstraints.cpp
@@ -30,7 +30,8 @@ namespace ChaosTest {
 	{
 		{
 			TPBDRigidsSOAs<FReal, 3> Particles;
-			TEvolution Evolution(Particles);
+			THandleArray<FChaosPhysicsMaterial> PhysicalMaterials;
+			TEvolution Evolution(Particles, PhysicalMaterials);
 			TArray<TPBDRigidParticleHandle<FReal, 3>*> Dynamics = Evolution.CreateDynamicParticles(1);
 			TArray<FVec3> Positions = { FVec3(0) };
 			TPBDPositionConstraints<FReal, 3> PositionConstraints(MoveTemp(Positions), MoveTemp(Dynamics), 1.f);
@@ -43,7 +44,8 @@ namespace ChaosTest {
 		}
 		{
 			TPBDRigidsSOAs<FReal, 3> Particles;
-			FPBDRigidsEvolutionGBF Evolution(Particles);
+			THandleArray<FChaosPhysicsMaterial> PhysicalMaterials;
+			FPBDRigidsEvolutionGBF Evolution(Particles, PhysicalMaterials);
 			TArray<TPBDRigidParticleHandle<FReal, 3>*> Dynamics = Evolution.CreateDynamicParticles(1);
 			Evolution.GetGravityForces().SetEnabled(*Dynamics[0], false);
 
@@ -78,7 +80,8 @@ namespace ChaosTest {
 	{
 		const int32 Iterations = 10;
 		TPBDRigidsSOAs<FReal, 3> Particles;
-		TEvolution Evolution(Particles, Iterations);
+		THandleArray<FChaosPhysicsMaterial> PhysicalMaterials;
+		TEvolution Evolution(Particles, PhysicalMaterials, Iterations);
 		TArray<TPBDRigidParticleHandle<FReal, 3>*> Dynamics = Evolution.CreateDynamicParticles(2);
 		TArray<FVec3> PositionConstraintPositions = { FVec3(0, 0, 0) };
 
diff --git a/Engine/Source/Programs/HeadlessChaos/Private/HeadlessChaosTestForces.cpp b/Engine/Source/Programs/HeadlessChaos/Private/HeadlessChaosTestForces.cpp
--- a/Engine/Source/Programs/HeadlessChaos/Private/HeadlessChaosTestForces.cpp
+++ b/Engine/Source/Programs/HeadlessChaos/Private/HeadlessChaosTestForces.cpp
@@ -17,7 +17,8 @@ namespace ChaosTest {
 	void Gravity()
 	{
 		TPBDRigidsSOAs<T, 3> Particles;
-		TEvolution Evolution(Particles);
+		THandleArray<FChaosPhysicsMaterial> PhysicalMaterials;
+		TEvolution Evolution(Particles, PhysicalMaterials);
 		
 		TArray<TPBDRigidParticleHandle<T,3>*> Dynamics = Evolution.CreateDynamicParticles(1);
 		Evolution.AdvanceOneTimeStep(0.1);
diff --git a/Engine/Source/Programs/HeadlessChaos/Private/HeadlessChaosTestGraph.cpp b/Engine/Source/Programs/HeadlessChaos/Private/HeadlessChaosTestGraph.cpp
--- a/Engine/Source/Programs/HeadlessChaos/Private/HeadlessChaosTestGraph.cpp
+++ b/Engine/Source/Programs/HeadlessChaos/Private/HeadlessChaosTestGraph.cpp
@@ -581,6 +581,7 @@ namespace ChaosTest {
 		TArray<TPBDRigidParticleHandle<FReal, 3>*> DynParticles = SOAs.CreateDynamicParticles(NumParticles);
 		TArray<TGeometryParticleHandle<FReal, 3>*> Particles;
 		TArrayCollectionArray<TSerializablePtr<FChaosPhysicsMaterial>> PhysicsMaterials;
+		THandleArray<FChaosPhysicsMaterial> PhysicalMaterials;
  		SOAs.GetParticleHandles().AddArray(&PhysicsMaterials);
 
 		for (int32 Idx = 0; Idx < NumParticles; ++Idx)
@@ -627,7 +628,7 @@ namespace ChaosTest {
 
 			for (int32 IslandIndex = 0; IslandIndex < Graph.NumIslands(); ++IslandIndex)
 			{
-				bool bSleeped = Graph.SleepInactive(IslandIndex, PhysicsMaterials);
+				bool bSleeped = Graph.SleepInactive(IslandIndex, PhysicsMaterials, PhysicalMaterials);
 
 				if (bSleeped)
 				{
diff --git a/Engine/Source/Programs/HeadlessChaos/Private/HeadlessChaosTestJoints.cpp b/Engine/Source/Programs/HeadlessChaos/Private/HeadlessChaosTestJoints.cpp
--- a/Engine/Source/Programs/HeadlessChaos/Private/HeadlessChaosTestJoints.cpp
+++ b/Engine/Source/Programs/HeadlessChaos/Private/HeadlessChaosTestJoints.cpp
@@ -476,8 +476,8 @@ namespace ChaosTest {
 
 		Box2->X() = FVec3((FReal)500, (FReal)0, (FReal)1000);
 		Box2->P() = Box2->X();
-
-		FPBDRigidsEvolutionGBF Evolution(Particles);
+		THandleArray<FChaosPhysicsMaterial> PhysicalMaterials;
+		FPBDRigidsEvolutionGBF Evolution(Particles, PhysicalMaterials);
 		TVector<TGeometryParticleHandle<FReal, 3>*,2> ConstrainedParticles = TVector<TGeometryParticleHandle<FReal,3>*, 2>(StaticBox, Box2);
 		TVector<FVec3, 2> Points = { FVec3((FReal)100, (FReal)0, (FReal)1000), FVec3((FReal)400, (FReal)0, (FReal)1000) };
 
@@ -518,7 +518,8 @@ namespace ChaosTest {
 			Box2.X() = FVec3((FReal)500, (FReal)0, (FReal)1000);
 			Box2.P() = Box2.X();
 
-			FPBDRigidsEvolutionGBF Evolution(Particles);
+			THandleArray<FChaosPhysicsMaterial> PhysicalMaterials;
+			FPBDRigidsEvolutionGBF Evolution(Particles, PhysicalMaterials);
 			TArray<TVector<TGeometryParticleHandle<FReal,3>*, 2>> Constraints = { TVector<TGeometryParticleHandle<FReal,3>*, 2>(&StaticBox, &Box2) };
 
 			Evolution.SetPhysicsMaterial(&StaticBox, MakeSerializable(PhysicalMaterial));
@@ -547,7 +548,8 @@ namespace ChaosTest {
 			Box2.X() = FVec3((FReal)500, (FReal)0, (FReal)1000);
 			Box2.P() = Box2.X();
 
-			FPBDRigidsEvolutionGBF Evolution(Particles);
+			THandleArray<FChaosPhysicsMaterial> PhysicalMaterials;
+			FPBDRigidsEvolutionGBF Evolution(Particles, PhysicalMaterials);
 			TArray<TVector<TGeometryParticleHandle<FReal,3>*, 2>> Constraints = { TVector<TGeometryParticleHandle<FReal,3>*, 2>(&StaticBox, &Box2) };
 
 			Evolution.SetPhysicsMaterial(&StaticBox, MakeSerializable(PhysicalMaterial));
diff --git a/Engine/Source/Programs/HeadlessChaos/Private/HeadlessChaosTestSerialization.cpp b/Engine/Source/Programs/HeadlessChaos/Private/HeadlessChaosTestSerialization.cpp
--- a/Engine/Source/Programs/HeadlessChaos/Private/HeadlessChaosTestSerialization.cpp
+++ b/Engine/Source/Programs/HeadlessChaos/Private/HeadlessChaosTestSerialization.cpp
@@ -424,7 +424,9 @@ namespace ChaosTest
 			{
 				Chaos::FChaosArchive ChaosAr(*File);
 				TPBDRigidsSOAs<float, 3> Particles;
-				FPBDRigidsEvolutionGBF Evolution(Particles);
+
+				THandleArray<FChaosPhysicsMaterial> PhysicalMaterials;
+				FPBDRigidsEvolutionGBF Evolution(Particles, PhysicalMaterials);
 
 				Evolution.Serialize(ChaosAr);
 				Evolution.AdvanceOneTimeStep(1 / 60.f);
diff --git a/Engine/Source/Programs/HeadlessChaos/Private/HeadlessChaosTestSim.cpp b/Engine/Source/Programs/HeadlessChaos/Private/HeadlessChaosTestSim.cpp
--- a/Engine/Source/Programs/HeadlessChaos/Private/HeadlessChaosTestSim.cpp
+++ b/Engine/Source/Programs/HeadlessChaos/Private/HeadlessChaosTestSim.cpp
@@ -16,7 +16,8 @@ namespace ChaosTest {
 	GTEST_TEST(SimTests, SphereSphereSimTest)
 	{
 		TPBDRigidsSOAs<FReal, 3> Particles;
-		FPBDRigidsEvolutionGBF Evolution(Particles);
+		THandleArray<FChaosPhysicsMaterial> PhysicalMaterials;
+		FPBDRigidsEvolutionGBF Evolution(Particles, PhysicalMaterials);
 		auto Static = Evolution.CreateStaticParticles(1)[0];
 		auto Dynamic = Evolution.CreateDynamicParticles(1)[0];
 
@@ -45,7 +46,8 @@ namespace ChaosTest {
 	GTEST_TEST(SimTests, BoxBoxSimTest)
 	{
 		TPBDRigidsSOAs<FReal, 3> Particles;
-		FPBDRigidsEvolutionGBF Evolution(Particles);
+		THandleArray<FChaosPhysicsMaterial> PhysicalMaterials;
+		FPBDRigidsEvolutionGBF Evolution(Particles, PhysicalMaterials);
 		auto Static = Evolution.CreateStaticParticles(1)[0];
 		auto Dynamic = Evolution.CreateDynamicParticles(1)[0];
 
@@ -83,7 +85,8 @@ namespace ChaosTest {
 	GTEST_TEST(SimTests, DISABLED_VeryLowInertiaSimTest)
 	{
 		TPBDRigidsSOAs<FReal, 3> Particles;
-		FPBDRigidsEvolutionGBF Evolution(Particles);
+		THandleArray<FChaosPhysicsMaterial> PhysicalMaterials;
+		FPBDRigidsEvolutionGBF Evolution(Particles, PhysicalMaterials);
 		auto Static = Evolution.CreateStaticParticles(1)[0];
 		auto Dynamic = Evolution.CreateDynamicParticles(1)[0];
 
@@ -111,7 +114,8 @@ namespace ChaosTest {
 	GTEST_TEST(SimTests, SleepAndWakeSimTest)
 	{
 		TPBDRigidsSOAs<FReal, 3> Particles;
-		FPBDRigidsEvolutionGBF Evolution(Particles);
+		THandleArray<FChaosPhysicsMaterial> PhysicalMaterials;
+		FPBDRigidsEvolutionGBF Evolution(Particles, PhysicalMaterials);
 		auto Static = Evolution.CreateStaticParticles(1)[0];
 		auto Dynamic1 = Evolution.CreateDynamicParticles(1)[0];
 		auto Dynamic2 = Evolution.CreateDynamicParticles(1)[0];
diff --git a/Engine/Source/Programs/HeadlessChaos/Public/HeadlessChaosTestConstraints.h b/Engine/Source/Programs/HeadlessChaos/Public/HeadlessChaosTestConstraints.h
--- a/Engine/Source/Programs/HeadlessChaos/Public/HeadlessChaosTestConstraints.h
+++ b/Engine/Source/Programs/HeadlessChaos/Public/HeadlessChaosTestConstraints.h
@@ -20,7 +20,7 @@ namespace ChaosTest
 	public:
 
 		FConstraintsTest(const int32 NumIterations, const FReal Gravity)
-			: Evolution(SOAs, NumIterations)
+			: Evolution(SOAs, PhysicalMaterials, NumIterations)
 		{
 			PhysicalMaterial = MakeUnique<FChaosPhysicsMaterial>();
 			PhysicalMaterial->Friction = 0;
@@ -63,6 +63,7 @@ namespace ChaosTest
 		TPBDRigidsSOAs<FReal, 3> SOAs;
 		FPBDRigidsEvolutionGBF Evolution;
 		TUniquePtr<FChaosPhysicsMaterial> PhysicalMaterial;
+		THandleArray<FChaosPhysicsMaterial> PhysicalMaterials;
 
 		TGeometryParticleHandle<FReal, 3>* GetParticle(const int32 Idx)
 		{
