Add stress test to HeadlessChaos for AABBTree Dirty elements query.

diff --git a/Engine/Source/Programs/HeadlessChaos/HeadlessChaos.Target.cs b/Engine/Source/Programs/HeadlessChaos/HeadlessChaos.Target.cs
--- a/Engine/Source/Programs/HeadlessChaos/HeadlessChaos.Target.cs
+++ b/Engine/Source/Programs/HeadlessChaos/HeadlessChaos.Target.cs
@@ -22,6 +22,8 @@ public class HeadlessChaosTarget : TargetRules
 
         bHasExports = false;
 
+        bUseLoggingInShipping = true;
+
         // UnrealHeaderTool is a console application, not a Windows app (sets entry point to main(), instead of WinMain())
         bIsBuildingConsoleApplication = true;
 
diff --git a/Engine/Source/Programs/HeadlessChaos/Private/HeadlessChaosApp.cpp b/Engine/Source/Programs/HeadlessChaos/Private/HeadlessChaosApp.cpp
--- a/Engine/Source/Programs/HeadlessChaos/Private/HeadlessChaosApp.cpp
+++ b/Engine/Source/Programs/HeadlessChaos/Private/HeadlessChaosApp.cpp
@@ -120,6 +120,7 @@ TEST(SerializationTests, Serialization) {
 TEST(BroadphaseTests, Broadphase) {
 	ChaosTest::BPPerfTest<float>();
 	ChaosTest::TestPendingSpatialDataHandlePointerConflict();
+	//ChaosTest::SpatialAccelerationDirtyAndGlobalQueryStrestTest<float>();
 	SUCCEED();
 }
 
diff --git a/Engine/Source/Programs/HeadlessChaos/Private/HeadlessChaosTestBroadphase.cpp b/Engine/Source/Programs/HeadlessChaos/Private/HeadlessChaosTestBroadphase.cpp
--- a/Engine/Source/Programs/HeadlessChaos/Private/HeadlessChaosTestBroadphase.cpp
+++ b/Engine/Source/Programs/HeadlessChaos/Private/HeadlessChaosTestBroadphase.cpp
@@ -9,6 +9,7 @@
 #include "Chaos/PBDRigidsSOAs.h"
 #include "Chaos/PBDRigidsEvolutionGBF.h"
 #include "Chaos/AABBTree.h"
+#include "ChaosLog.h"
 #include "PBDRigidsSolver.h"
 #include "Chaos/SpatialAccelerationCollection.h"
 
@@ -179,13 +180,59 @@ namespace ChaosTest
 	};
 
 	template <typename T>
-	auto BuildBoxes(TUniquePtr<TBox<T,3>>& Box)
+	struct TStressTestVisitor : ISpatialVisitor<TAccelerationStructureHandle<T, 3>, T>
 	{
-		Box = MakeUnique<TBox<T, 3>>(TVector<T, 3>(0, 0, 0), TVector<T, 3>(100, 100, 100));
+		using FPayload = TAccelerationStructureHandle<T, 3>;
+
+		TStressTestVisitor() {}
+
+		enum class SQType
+		{
+			Raycast,
+			Sweep,
+			Overlap
+		};
+
+		bool VisitRaycast(const TSpatialVisitorData<FPayload>& Data, FQueryFastData& CurData)
+		{
+			return true;
+		}
+
+		bool VisitSweep(const TSpatialVisitorData<FPayload>& Data, FQueryFastData& CurData)
+		{
+			return true;
+		}
+
+		bool VisitOverlap(const TSpatialVisitorData<FPayload>& Data)
+		{
+			return true;
+		}
+
+		virtual bool Overlap(const TSpatialVisitorData<FPayload>& Instance) override
+		{
+			return VisitOverlap(Instance);
+		}
+
+		virtual bool Raycast(const TSpatialVisitorData<FPayload>& Instance, FQueryFastData& CurData) override
+		{
+			return VisitRaycast(Instance, CurData);
+		}
+
+		virtual bool Sweep(const TSpatialVisitorData<FPayload>& Instance, FQueryFastData& CurData) override
+		{
+			return VisitSweep(Instance, CurData);
+		}
+	};
+
+
+	template <typename T>
+	auto BuildBoxes(TUniquePtr<TBox<T,3>>& Box, float BoxSize = 100, const TVector<T, 3>& BoxGridDimensions = TVector<T, 3>(10,10,10))
+	{
+		Box = MakeUnique<TBox<T, 3>>(TVector<T, 3>(0, 0, 0), TVector<T, 3>(BoxSize, BoxSize, BoxSize));
 		auto Boxes = MakeUnique<TGeometryParticles<T, 3>>();
-		const int32 NumRows = 10;
-		const int32 NumCols = 10;
-		const int32 NumHeight = 10;
+		const int32 NumRows = BoxGridDimensions.X;
+		const int32 NumCols = BoxGridDimensions.Y;
+		const int32 NumHeight = BoxGridDimensions.Z;
 
 		Boxes->AddParticles(NumRows * NumCols * NumHeight);
 
@@ -726,4 +773,153 @@ namespace ChaosTest
 			EXPECT_EQ(HitNum, 0);
 		}
 	}
+
+	template<typename T>
+	void SpatialAccelerationDirtyAndGlobalQueryStrestTest()
+	{
+		using AABBTreeType = TAABBTree<TAccelerationStructureHandle<T, 3>, TAABBTreeLeafArray<TAccelerationStructureHandle<T, 3>, T>, T>;
+
+		// Construct 100000 Particles
+		const int32 NumRows = 100;
+		const int32 NumCols = 100;
+		const int32 NumHeight = 10;
+		const int32 ParticleCount = NumRows * NumCols * NumHeight;
+		const float BoxSize = 100;
+
+		TPBDRigidsSOAs<T, 3> Particles;
+		TArray<TPBDRigidParticleHandle<T, 3>*> ParticleHandles = Particles.CreateDynamicParticles(ParticleCount);
+		for (auto& Handle : ParticleHandles)
+		{
+			Handle->GTGeometryParticle() = TGeometryParticle<T, 3>::CreateParticle().Release();
+		}
+		const auto& ParticlesView = Particles.GetAllParticlesView();
+
+		// ensure these can't be filtered out.
+		FCollisionFilterData FilterData;
+		FilterData.Word0 = TNumericLimits<uint32>::Max();
+		FilterData.Word1 = TNumericLimits<uint32>::Max();
+		FilterData.Word2 = TNumericLimits<uint32>::Max();
+		FilterData.Word3 = TNumericLimits<uint32>::Max();
+
+		TSharedPtr<TBox<T, 3>, ESPMode::ThreadSafe> Box = MakeShared<TBox<T, 3>, ESPMode::ThreadSafe>(TVector<T, 3>(0, 0, 0), TVector<T, 3>(BoxSize, BoxSize, BoxSize));
+
+		int32 Idx = 0;
+		for (int32 Height = 0; Height < NumHeight; ++Height)
+		{
+			for (int32 Row = 0; Row < NumRows; ++Row)
+			{
+				for (int32 Col = 0; Col < NumCols; ++Col)
+				{
+					TGeometryParticle<T, 3>* GTParticle = ParticleHandles[Idx]->GTGeometryParticle();
+					TPBDRigidParticleHandle<T, 3>* Handle = ParticleHandles[Idx];
+
+					Handle->SetGeometry(MakeSerializable(Box));
+					Handle->ShapesArray()[0]->SetQueryData(FilterData);
+					GTParticle->SetGeometry(Box);
+					GTParticle->ShapesArray()[0]->SetQueryData(FilterData);
+					Handle->SetX(TVector<T, 3>(Col * BoxSize, Row * BoxSize, Height * BoxSize));
+					GTParticle->SetX(TVector<T, 3>(Col * BoxSize, Row * BoxSize, Height * BoxSize));
+					Handle->SetR(TRotation<T, 3>::Identity);
+					GTParticle->SetR(TRotation<T, 3>::Identity);
+					Handle->SetUniqueIdx(FUniqueIdx(Idx));
+					GTParticle->SetUniqueIdx(FUniqueIdx(Idx));
+					Handle->SetLocalBounds(Box->BoundingBox());
+					Handle->SetHasBounds(true);
+					Handle->SetWorldSpaceInflatedBounds(Box->BoundingBox().TransformedAABB(TRigidTransform<T, 3>(GTParticle->X(), GTParticle->R())));
+					++Idx;
+				}
+			}
+		}
+
+		int32 DirtyNum = 800;
+		int32 Queries = 500;
+		ensure(DirtyNum < ParticleCount);
+
+		// Construct tree
+		AABBTreeType Spatial(ParticlesView);
+
+		// Update DirtyNum elements, so they are pulled out of leaves.
+		for (int32 i = 0; i < DirtyNum; ++i)
+		{
+			TAccelerationStructureHandle<T, 3> Payload(ParticleHandles[i]->GTGeometryParticle());
+			TAABB<T, 3> Bounds = ParticleHandles[i]->WorldSpaceInflatedBounds();
+			Spatial.UpdateElement(Payload, Bounds, true);
+		}
+
+		// RAYCASTS
+		{
+			// Setup raycast params
+			const FVec3 Start(500, 500, 500);
+			const FVec3 Dir(1, 0, 0);
+			const FReal Length = 1000;
+			TStressTestVisitor<T> Visitor;
+
+			// Measure raycasts
+			uint32 Cycles = 0.0;
+			for (int32 Query = 0; Query < Queries; ++Query)
+			{
+				uint32 StartTime = FPlatformTime::Cycles();
+
+				Spatial.Raycast(Start, Dir, Length, Visitor);
+
+				Cycles += FPlatformTime::Cycles() - StartTime;
+			}
+
+			float Milliseconds = FPlatformTime::ToMilliseconds(Cycles);
+			float AvgMicroseconds = (Milliseconds * 1000) / Queries;
+
+			UE_LOG(LogHeadlessChaos, Warning, TEXT("Raycast Test: Dirty Particles: %d, Queries: %d, Avg Query Time: %f(us), Total:%f(ms)"), DirtyNum, Queries, AvgMicroseconds, Milliseconds);
+		}
+
+		// SWEEPS
+		{
+			// Setup Sweep params
+			const FVec3 Start(500, 500, 500);
+			const FVec3 Dir(1, 0, 0);
+			const FReal Length = 1000;
+			const FVec3 HalfExtents(50, 50, 50);
+			TStressTestVisitor<T> Visitor;
+
+			// Measure raycasts
+			uint32 Cycles = 0.0;
+			for (int32 Query = 0; Query < Queries; ++Query)
+			{
+				uint32 StartTime = FPlatformTime::Cycles();
+
+				Spatial.Sweep(Start, Dir, Length, HalfExtents, Visitor);
+
+				Cycles += FPlatformTime::Cycles() - StartTime;
+			}
+
+			float Milliseconds = FPlatformTime::ToMilliseconds(Cycles);
+			float AvgMicroseconds = (Milliseconds * 1000) / Queries;
+
+			UE_LOG(LogHeadlessChaos, Warning, TEXT("Sweep Test: Dirty Particles: %d, Queries: %d, Avg Query Time: %f(us), Total:%f(ms)"), DirtyNum, Queries, AvgMicroseconds, Milliseconds);
+		}
+
+		// OVERLAPS
+		{
+			TStressTestVisitor<T> Visitor;
+			const TAABB<T, 3> QueryBounds(TVector<T, 3>(-50, -50, -50), TVector<T, 3>(50,50,50));
+
+			// Measure raycasts
+			uint32 Cycles = 0.0;
+			for (int32 Query = 0; Query < Queries; ++Query)
+			{
+				uint32 StartTime = FPlatformTime::Cycles();
+
+				Spatial.Overlap(QueryBounds, Visitor);
+
+				Cycles += FPlatformTime::Cycles() - StartTime;
+			}
+
+			float Milliseconds = FPlatformTime::ToMilliseconds(Cycles);
+			float AvgMicroseconds = (Milliseconds * 1000) / Queries;
+
+			UE_LOG(LogHeadlessChaos, Error, TEXT("Overlap Test: Dirty Particles: %d, Queries: %d, Avg Query Time: %f(us), Total:%f(ms)"), DirtyNum, Queries, AvgMicroseconds, Milliseconds);
+		}
+	}
+
+	template void SpatialAccelerationDirtyAndGlobalQueryStrestTest<float>();
+
 }
\ No newline at end of file
diff --git a/Engine/Source/Programs/HeadlessChaos/Public/HeadlessChaosTestBroadphase.h b/Engine/Source/Programs/HeadlessChaos/Public/HeadlessChaosTestBroadphase.h
--- a/Engine/Source/Programs/HeadlessChaos/Public/HeadlessChaosTestBroadphase.h
+++ b/Engine/Source/Programs/HeadlessChaos/Public/HeadlessChaosTestBroadphase.h
@@ -23,4 +23,7 @@ namespace ChaosTest {
 	void BroadphaseCollectionTest();
 
 	void TestPendingSpatialDataHandlePointerConflict();
+
+	template<typename T>
+	void SpatialAccelerationDirtyAndGlobalQueryStrestTest();
 }
\ No newline at end of file
