ChaosPhysics collision detection:
- EPA calculate correct normals for certain edge cases (when starting with a point simplex)
- Make sure initial EPA tetrahedron is correctly orientated with touching cases
- Rename template parameter to be more descriptive

diff --git a/Engine/Source/Programs/HeadlessChaos/Private/HeadlessChaosTestEPA.cpp b/Engine/Source/Programs/HeadlessChaos/Private/HeadlessChaosTestEPA.cpp
--- a/Engine/Source/Programs/HeadlessChaos/Private/HeadlessChaosTestEPA.cpp
+++ b/Engine/Source/Programs/HeadlessChaos/Private/HeadlessChaosTestEPA.cpp
@@ -1004,7 +1004,7 @@ namespace ChaosTest
 	// This test reproduces an in-game example. 
 	// NOTE: These are shapes in a skeletal mesh, hence the rotation (mesh is on its side so its Y is "up")
 	//
-	GTEST_TEST(EPATests, DISABLED_EPARealFailures_BoxBox)
+	GTEST_TEST(EPATests, EPARealFailures_BoxBox)
 	{
 		TBox<FReal, 3> A(FVec3(-250.000000, 250.000000, -750.000000), FVec3(1250.00000, 350.000000, 750.000000));
 		TBox<FReal, 3> B(FVec3(50.0000000, -50.0000000, -50.0000000), FVec3(150.000000, 50.0000000, 50.0000000));
@@ -1025,27 +1025,30 @@ namespace ChaosTest
 		EXPECT_GT(FMath::Abs(Normal.Z), 0.8f);
 	}
 
-	GTEST_TEST(EPATests, DISABLED_EPA_EdgeCases)
+	GTEST_TEST(EPATests, EPA_EdgeCases)
 	{
-		TBox<FReal, 3> A(FVec3(0.0f, 0.0f, 0.0f), FVec3(100.0f, 100.0f, 100.0f));
-		TBox<FReal, 3> B(FVec3(100.0f + 0.00001f, 0.0f, 0.0f), FVec3(200.0f, 100.0f, 100.0f)); // This small epsilon can actually be zero as well.
-		FRigidTransform3 ATM = FRigidTransform3(FVec3(0.0f, 0.0f, 0.0f), FRotation3::FromElements(0.0f, 0.0f, 0.0f, 1.0f));
-		FRigidTransform3 BTM = FRigidTransform3(FVec3(0.0f, 0.0f, 0.0f), FRotation3::FromElements(0.0f, 0.0f, 0.0f, 1.0f));
-		FVec3 InitialDir = FVec3(1.0f, 0.0f, 0.0f);
-
-		const FRigidTransform3 BToATM = BTM.GetRelativeTransform(ATM);
-		FReal Penetration;
-		FVec3 ClosestA, ClosestB, NormalA;
-		int32 NumIterations = 0;
-		GJKPenetration<true>(A, B, BToATM, Penetration, ClosestA, ClosestB, NormalA, 0.0f, InitialDir, 0.0f, &NumIterations);
+		// Two boxes exactly touching each other (This was a failing case that was fixed)
+		{
+			TBox<FReal, 3> A(FVec3(0.0f, 0.0f, 0.0f), FVec3(100.0f, 100.0f, 100.0f));
+			TBox<FReal, 3> B(FVec3(100.0f, 0.0f, 0.0f), FVec3(200.0f, 100.0f, 100.0f));
+			FRigidTransform3 ATM = FRigidTransform3(FVec3(0.0f, 0.0f, 0.0f), FRotation3::FromElements(0.0f, 0.0f, 0.0f, 1.0f));
+			FRigidTransform3 BTM = FRigidTransform3(FVec3(0.0f, 0.0f, 0.0f), FRotation3::FromElements(0.0f, 0.0f, 0.0f, 1.0f));
+			FVec3 InitialDir = FVec3(1.0f, 0.0f, 0.0f);
+
+			const FRigidTransform3 BToATM = BTM.GetRelativeTransform(ATM);
+			FReal Penetration;
+			FVec3 ClosestA, ClosestB, NormalA;
+			int32 NumIterations = 0;
+			GJKPenetration<true>(A, B, BToATM, Penetration, ClosestA, ClosestB, NormalA, 0.0f, InitialDir, 0.0f, &NumIterations);
 
-		FVec3 Location = ATM.TransformPosition(ClosestA);
-		FVec3 Normal = ATM.TransformVectorNoScale(NormalA);
-		FReal Phi = -Penetration;
+			FVec3 Location = ATM.TransformPosition(ClosestA);  // These transforms are not really necessary since they are identity
+			FVec3 Normal = ATM.TransformVectorNoScale(NormalA);
+			FReal Phi = -Penetration;
 
-		EXPECT_NEAR(Normal.X, 1.0f, 1e-3f);
-		EXPECT_NEAR(Normal.Y, 0.0f, 1e-3f);
-		EXPECT_NEAR(Normal.Z, 0.0f, 1e-3f);
+			EXPECT_NEAR(Normal.X, 1.0f, 1e-3f);
+			EXPECT_NEAR(Normal.Y, 0.0f, 1e-3f);
+			EXPECT_NEAR(Normal.Z, 0.0f, 1e-3f);
+		}
 	}
 
 
diff --git a/Engine/Source/Runtime/Experimental/Chaos/Public/Chaos/EPA.h b/Engine/Source/Runtime/Experimental/Chaos/Public/Chaos/EPA.h
--- a/Engine/Source/Runtime/Experimental/Chaos/Public/Chaos/EPA.h
+++ b/Engine/Source/Runtime/Experimental/Chaos/Public/Chaos/EPA.h
@@ -164,8 +164,10 @@ bool InitializeEPA(TArray<TVec3<T>>& VertsA, TArray<TVec3<T>>& VertsB, const Sup
 	{
 		case 1:
 		{
-			//assuming it's a touching hit at origin
-			return false;
+			//assuming it's a touching hit at origin, but we still need to calculate a separating normal
+			AddFartherPoint(OutTouchNormal); // Use an arbitrary direction
+
+			// Now we might have a line! So fall trough to the next case
 		}
 		case 2:
 		{
@@ -173,7 +175,7 @@ bool InitializeEPA(TArray<TVec3<T>>& VertsA, TArray<TVec3<T>>& VertsB, const Sup
 			TVec3<T> Dir = MinkowskiVert(VertsA.GetData(), VertsB.GetData(), 1) - MinkowskiVert(VertsA.GetData(), VertsB.GetData(), 0);
 
 			bValid = Dir.SizeSquared() > 1e-4;
-			if (CHAOS_ENSURE(bValid))	//two verts given should be distinct
+			if (bValid)	//two verts given are distinct
 			{
 				//find most opposing axis
 				int32 BestAxis = 0;
@@ -205,6 +207,12 @@ bool InitializeEPA(TArray<TVec3<T>>& VertsA, TArray<TVec3<T>>& VertsB, const Sup
 					return false;
 				}
 			}
+			else
+			{
+				// The two vertices are not distinct that may happen when the single vertex case above was hit and our CSO is very thin in that direction
+				CHAOS_ENSURE(NumVerts == 1); // If this ensure fires we were given 2 vertices that are not distinct to start with
+				return false;
+			}
 			break;
 		}
 		case 3:
@@ -255,7 +263,19 @@ bool InitializeEPA(TArray<TVec3<T>>& VertsA, TArray<TVec3<T>>& VertsB, const Sup
 	if (bValid)
 	{
 		//make sure normals are pointing out of tetrahedron
-		if (TVec3<T>::DotProduct(OutEntries[0].PlaneNormal, MinkowskiVert(VertsA.GetData(), VertsB.GetData(), 0)) > 0)
+		// In the usual case the distances will either all be positive or negative, 
+		// but the tetrahedron can be very close to (or touching) the origin
+		// Look for farthest plane to decide
+		T MaxSignedDistance = 0;
+		for (TEPAEntry<T>& Entry : OutEntries)
+		{
+			if (FMath::Abs(Entry.Distance) > FMath::Abs(MaxSignedDistance))
+			{
+				MaxSignedDistance = Entry.Distance;
+			}
+		}
+
+		if (MaxSignedDistance < 0.0f)
 		{
 			for (TEPAEntry<T>& Entry : OutEntries)
 			{
diff --git a/Engine/Source/Runtime/Experimental/Chaos/Public/Chaos/GJK.h b/Engine/Source/Runtime/Experimental/Chaos/Public/Chaos/GJK.h
--- a/Engine/Source/Runtime/Experimental/Chaos/Public/Chaos/GJK.h
+++ b/Engine/Source/Runtime/Experimental/Chaos/Public/Chaos/GJK.h
@@ -87,8 +87,8 @@ namespace Chaos
 	}
 
 	
-	// This function will be faster if bNegativePenetrationSupport is false, so don't use the feature if not required
-	template <bool bNegativePenetrationSupport = false, typename T, typename TGeometryA, typename TGeometryB>
+	// This function will be faster if bNegativePenetrationAllowed is false, so don't use the feature if not required
+	template <bool bNegativePenetrationAllowed = false, typename T, typename TGeometryA, typename TGeometryB>
 	bool GJKPenetration(const TGeometryA& A, const TGeometryB& B, const TRigidTransform<T, 3>& BToATM, T& OutPenetration, TVec3<T>& OutClosestA, TVec3<T>& OutClosestB, TVec3<T>& OutNormal, const T InThicknessA = 0, const TVector<T, 3>& InitialDir = TVector<T, 3>(1, 0, 0), const T InThicknessB = 0, int32* OutNumIterations = nullptr)
 	{
 		auto SupportAFunc = [&A](const TVec3<T>& V)
@@ -139,7 +139,7 @@ namespace Chaos
 			const TVector<T, 3> SupportB = SupportBFunc(V);
 			const TVector<T, 3> W = SupportA - SupportB;
 
-			if (!bNegativePenetrationSupport && TVector<T, 3>::DotProduct(V, W) > Inflation)
+			if (!bNegativePenetrationAllowed && TVector<T, 3>::DotProduct(V, W) > Inflation)
 			{
 				return false;
 			}
@@ -189,7 +189,7 @@ namespace Chaos
 			const T PreDist = FMath::Sqrt(PrevDist2);
 			OutNormal = (ClosestBInA - ClosestA).GetUnsafeNormal();	//Note1: should we just use PreDist2? //Note2: we can just use -V.GetUnsafeNormal() here instead if it improves accuracy
 			T Penetration = ThicknessA + ThicknessB - PreDist;
-			if (!bNegativePenetrationSupport)
+			if (!bNegativePenetrationAllowed)
 			{
 				if (Penetration < 0.0f) return false;
 			}
