Chaos Physics:
EPA resulted in large penetrations when shapes are slightly separated (in some cases): fixed
EPA now supports small separations more robustly (This is necessary since GJK won't handle this case)
Minor simplex numerical accuracy change

diff --git a/Engine/Source/Programs/HeadlessChaos/Private/HeadlessChaosTestEPA.cpp b/Engine/Source/Programs/HeadlessChaos/Private/HeadlessChaosTestEPA.cpp
--- a/Engine/Source/Programs/HeadlessChaos/Private/HeadlessChaosTestEPA.cpp
+++ b/Engine/Source/Programs/HeadlessChaos/Private/HeadlessChaosTestEPA.cpp
@@ -522,6 +522,69 @@ namespace ChaosTest
 			EXPECT_FLOAT_EQ(Penetration, 1);
 			EXPECT_NEAR(WitnessA.Size(), 1, 1e-1);	//don't know exact point, but should be 1 away from origin
 		}
+
+		// Tetrahedron that does not quite contain the origin (so not in penetration)
+		{
+			T eps = 0.00001f;
+			TVec3<T> HullVerts[4] = { {-1, 0, 0}, {0 - eps, 1, -1}, {0 - eps, 0, 1}, {0 - eps, -1, -1}};
+
+			auto Support = [&HullVerts](const auto& V)
+			{
+				auto MaxDist = TNumericLimits<T>::Lowest();
+				auto BestVert = HullVerts[0];
+				for (const auto& Vert : HullVerts)
+				{
+					const auto Dist = TVec3<T>::DotProduct(V, Vert);
+					if (Dist > MaxDist)
+					{
+						MaxDist = Dist;
+						BestVert = Vert;
+					}
+				}
+				return BestVert;
+			};
+
+			TArray<TVec3<T>> Tetrahedron = { HullVerts[0], HullVerts[1], HullVerts[2], HullVerts[3] };
+			TArray <TVec3<T>> Zeros = { TVec3<T>(0), TVec3<T>(0), TVec3<T>(0), TVec3<T>(0) };
+
+			T Penetration;
+			TVec3<T> Dir, WitnessA, WitnessB;
+			EXPECT_EQ(EPA(Tetrahedron, Zeros, Support, ZeroSupport, Penetration, Dir, WitnessA, WitnessB), EPAResult::Ok);
+			EXPECT_LT(Penetration, 0); // Negative penetration
+			EXPECT_NEAR(Dir.X, 1.0f, 0.001f);
+		}
+
+		// Tetrahedron that is quite a bit away from the origin (so not in penetration)
+		// This takes a slightly different code path as compared to the previous test
+		{
+			T eps = 0.1f;
+			TVec3<T> HullVerts[4] = { {-1, 0, 0}, {0 - eps, 1, -1}, {0 - eps, 0, 1}, {0 - eps, -1, -1} };
+
+			auto Support = [&HullVerts](const auto& V)
+			{
+				auto MaxDist = TNumericLimits<T>::Lowest();
+				auto BestVert = HullVerts[0];
+				for (const auto& Vert : HullVerts)
+				{
+					const auto Dist = TVec3<T>::DotProduct(V, Vert);
+					if (Dist > MaxDist)
+					{
+						MaxDist = Dist;
+						BestVert = Vert;
+					}
+				}
+				return BestVert;
+			};
+
+			TArray<TVec3<T>> Tetrahedron = { HullVerts[0], HullVerts[1], HullVerts[2], HullVerts[3] };
+			TArray <TVec3<T>> Zeros = { TVec3<T>(0), TVec3<T>(0), TVec3<T>(0), TVec3<T>(0) };
+
+			T Penetration;
+			TVec3<T> Dir, WitnessA, WitnessB;
+			EXPECT_EQ(EPA(Tetrahedron, Zeros, Support, ZeroSupport, Penetration, Dir, WitnessA, WitnessB), EPAResult::Ok);
+			EXPECT_LT(Penetration, 0); // Negative penetration
+			EXPECT_NEAR(Dir.X, 1.0f, 0.001f);
+		}
 	}
 
 	template void EPASimpleTest<float>();
@@ -531,15 +594,15 @@ namespace ChaosTest
 	{
 		{
 			//get to EPA from GJKPenetration
+			// Boxes that are very close to each other (Almost penetrating).
 			TAABB<float, 3> Box({ -50, -50, -50 }, { 50, 50, 50 });
 
 			const FRigidTransform3 BToATM({ -8.74146843, 4.58291769, -100.029655 }, FRotation3::FromElements(6.63562241e-05, -0.000235952888, 0.00664712908, 0.999977887));
 			FVec3 ClosestA, ClosestB, Normal;
 			float Penetration;
-			GJKPenetration(Box, Box, BToATM, Penetration, ClosestA, ClosestB, Normal);
 
-			EXPECT_GT(Penetration, 0);
-			EXPECT_LT(Penetration, 2);
+			GJKPenetration<true>(Box, Box, BToATM, Penetration, ClosestA, ClosestB, Normal);
+			EXPECT_NEAR(Penetration, 0.0, 0.01);
 		}
 
 		// Problem: EPA was selecting the wrong face on the second box, resulting in a large penetration depth (131cm, but the box is only 20cm thick)
@@ -1053,6 +1116,31 @@ namespace ChaosTest
 			EXPECT_NEAR(Normal.Y, 0.0f, 1e-3f);
 			EXPECT_NEAR(Normal.Z, 0.0f, 1e-3f);
 		}
+
+		// Two boxes almost touching each other (Separation is small enough for GJK to fail and fall back to EPA)
+		// (This was a failing case that was fixed)
+		{
+			TBox<FReal, 3> A(FVec3(0.0f, 0.0f, 0.0f), FVec3(100.0f, 100.0f, 100.0f));
+			TBox<FReal, 3> B(FVec3(100.0f + 0.00001f, 0.0f, 0.0f), FVec3(200.0f, 100.0f, 100.0f));
+			FRigidTransform3 ATM = FRigidTransform3(FVec3(0.0f, 0.0f, 0.0f), FRotation3::FromElements(0.0f, 0.0f, 0.0f, 1.0f));
+			FRigidTransform3 BTM = FRigidTransform3(FVec3(0.0f, 0.0f, 0.0f), FRotation3::FromElements(0.0f, 0.0f, 0.0f, 1.0f));
+			FVec3 InitialDir = FVec3(1.0f, 0.0f, 0.0f);
+
+			const FRigidTransform3 BToATM = BTM.GetRelativeTransform(ATM);
+			FReal Penetration;
+			FVec3 ClosestA, ClosestB, NormalA;
+			int32 NumIterations = 0;
+			GJKPenetration<true>(A, B, BToATM, Penetration, ClosestA, ClosestB, NormalA, 0.0f, InitialDir, 0.0f, &NumIterations);
+
+			FVec3 Location = ATM.TransformPosition(ClosestA);  // These transforms are not really necessary since they are identity
+			FVec3 Normal = ATM.TransformVectorNoScale(NormalA);
+			FReal Phi = -Penetration;
+
+			EXPECT_NEAR(Normal.X, 1.0f, 1e-3f);
+			EXPECT_NEAR(Normal.Y, 0.0f, 1e-3f);
+			EXPECT_NEAR(Normal.Z, 0.0f, 1e-3f);
+		}
+
 	}
 
 
diff --git a/Engine/Source/Runtime/Experimental/Chaos/Public/Chaos/EPA.h b/Engine/Source/Runtime/Experimental/Chaos/Public/Chaos/EPA.h
--- a/Engine/Source/Runtime/Experimental/Chaos/Public/Chaos/EPA.h
+++ b/Engine/Source/Runtime/Experimental/Chaos/Public/Chaos/EPA.h
@@ -346,6 +346,10 @@ void ComputeEPAResults(const TVec3<T>* VertsA, const TVec3<T>* VertsB, const TEP
 	if (OutPenetration < 1e-4)	//if closest point is on the origin (edge case when surface is right on the origin)
 	{
 		OutDir = Entry.PlaneNormal;	//just fall back on plane normal
+		if (Entry.Distance < 0)
+		{
+			OutPenetration = -OutPenetration; // We are a bit outside of the shape so penetration is negative
+		}
 	}
 	else
 	{
@@ -354,6 +358,7 @@ void ComputeEPAResults(const TVec3<T>* VertsA, const TVec3<T>* VertsB, const TEP
 		{
 			//The origin is on the outside, so the direction is reversed
 			OutDir = -OutDir;
+			OutPenetration = -OutPenetration; // We are a bit outside of the shape so penetration is negative
 		}
 	}
 
@@ -419,7 +424,8 @@ EPAResult EPA(TArray<TVec3<T>>& VertsABuffer, TArray<TVec3<T>>& VertsBBuffer, co
 	for(int32 Idx = 0; Idx < Entries.Num(); ++Idx)
 	{
 		//ensure(Entries[Idx].Distance > -Eps);
-		if(Entries[Idx].IsOriginProjectedInside(VertsABuffer.GetData(), VertsBBuffer.GetData()))
+		// Entries[Idx].Distance <= 0.0f is true if the origin is a bit out of the polytope (we need to support this case for robustness)
+		if(Entries[Idx].Distance <= 0.0f || Entries[Idx].IsOriginProjectedInside(VertsABuffer.GetData(), VertsBBuffer.GetData()))
 		{
 			Queue.push(FEPAEntryWrapper {&Entries, Idx});
 		}
@@ -524,11 +530,14 @@ EPAResult EPA(TArray<TVec3<T>>& VertsABuffer, TArray<TVec3<T>>& VertsBBuffer, co
 					break;
 				}
 
-				//We should never need to check the lower bound, but in the case of bad precision this can happen
-				//We simply ignore this direction as it likely has even more bad precision
-				if (bValidTri && NewEntry.Distance >= LowerBound && NewEntry.Distance <= UpperBound)
+				// Due to numerical inaccuracies NewEntry.Distance >= LowerBound may be false!
+				// However these Entries still have good normals, and needs to be included to prevent
+				// this exiting with very deep penetration results
+				
+				if (bValidTri && NewEntry.Distance <= UpperBound)
 				{
-					if (NewEntry.IsOriginProjectedInside(VertsABuffer.GetData(), VertsBBuffer.GetData()))
+					// NewEntry.Distance <= 0.0f is if the origin is a bit out of the polytope
+					if (NewEntry.Distance <= 0.0f || NewEntry.IsOriginProjectedInside(VertsABuffer.GetData(), VertsBBuffer.GetData()))
 					{
 						Queue.push(FEPAEntryWrapper{ &Entries, NewIdx });
 					}
diff --git a/Engine/Source/Runtime/Experimental/Chaos/Public/Chaos/Simplex.h b/Engine/Source/Runtime/Experimental/Chaos/Public/Chaos/Simplex.h
--- a/Engine/Source/Runtime/Experimental/Chaos/Public/Chaos/Simplex.h
+++ b/Engine/Source/Runtime/Experimental/Chaos/Public/Chaos/Simplex.h
@@ -25,7 +25,7 @@ namespace Chaos
 
 		const T X0ToX1Squared = X0ToX1.SizeSquared();
 
-		if (X0ToX1Squared < Dot || X0ToX1Squared <= std::numeric_limits<T>::min())	//if dividing gives 1+ or the line is degenerate
+		if (X0ToX1Squared <= Dot || X0ToX1Squared <= std::numeric_limits<T>::min())	//if dividing gives 1+ or the line is degenerate
 		{
 			NumVerts = 1;
 			Idxs[0] = Idxs[1];
