Added velocity recording

diff --git a/Engine/Source/Programs/HeadlessChaos/Private/HeadlessChaosTestRewind.cpp b/Engine/Source/Programs/HeadlessChaos/Private/HeadlessChaosTestRewind.cpp
--- a/Engine/Source/Programs/HeadlessChaos/Private/HeadlessChaosTestRewind.cpp
+++ b/Engine/Source/Programs/HeadlessChaos/Private/HeadlessChaosTestRewind.cpp
@@ -28,7 +28,7 @@ namespace ChaosTest {
 		Solver->UpdateGameThreadStructures();
 	}
 
-	GTEST_TEST(RewindTest,DataCapture1)
+	GTEST_TEST(RewindTest,MovingGeomChange)
 	{
 		auto Sphere = TSharedPtr<FImplicitObject,ESPMode::ThreadSafe>(new TSphere<float,3>(TVector<float,3>(0),10));
 		auto Box = TSharedPtr<FImplicitObject,ESPMode::ThreadSafe>(new TBox<float,3>(FVec3(0),FVec3(1)));
@@ -120,7 +120,7 @@ namespace ChaosTest {
 	}
 
 
-	GTEST_TEST(RewindTest,DataCapture2)
+	GTEST_TEST(RewindTest,AddForce)
 	{
 		auto Sphere = TSharedPtr<FImplicitObject,ESPMode::ThreadSafe>(new TSphere<float,3>(TVector<float,3>(0),10));
 
@@ -169,7 +169,7 @@ namespace ChaosTest {
 		Module->DestroySolver(Solver);
 	}
 
-	GTEST_TEST(RewindTest,DataCapture3)
+	GTEST_TEST(RewindTest,IntermittentForce)
 	{
 		auto Sphere = TSharedPtr<FImplicitObject,ESPMode::ThreadSafe>(new TSphere<float,3>(TVector<float,3>(0),10));
 
@@ -242,7 +242,7 @@ namespace ChaosTest {
 		Module->DestroySolver(Solver);
 	}
 
-	GTEST_TEST(RewindTest,DataCapture4)
+	GTEST_TEST(RewindTest,IntermittentGeomChange)
 	{
 		auto Sphere = TSharedPtr<FImplicitObject,ESPMode::ThreadSafe>(new TSphere<float,3>(TVector<float,3>(0),10));
 		auto Box = TSharedPtr<FImplicitObject,ESPMode::ThreadSafe>(new TBox<float,3>(FVec3(0),FVec3(1)));
@@ -328,4 +328,54 @@ namespace ChaosTest {
 
 		Module->DestroySolver(Solver);
 	}
+
+	GTEST_TEST(RewindTest,FallingObject)
+	{
+		auto Sphere = TSharedPtr<FImplicitObject,ESPMode::ThreadSafe>(new TSphere<float,3>(TVector<float,3>(0),10));
+
+		FChaosSolversModule* Module = FChaosSolversModule::GetModule();
+		Module->ChangeThreadingMode(EChaosThreadingMode::SingleThread);
+
+		// Make a solver
+		FPhysicsSolver* Solver = Module->CreateSolver(nullptr,ESolverFlags::Standalone);
+		Solver->SetEnabled(true);
+
+		Solver->EnableRewindCapture(20);
+
+
+		// Make particles
+		auto Particle = TPBDRigidParticle<float,3>::CreateParticle();
+
+		Particle->SetGeometry(Sphere);
+		Solver->RegisterObject(Particle.Get());
+		Particle->SetGravityEnabled(true);
+		Particle->SetX(FVec3(0,0,100));
+
+		TArray<FVec3> X;
+		TArray<FVec3> V;
+
+		for(int Step = 0; Step < 10; ++Step)
+		{
+			X.Add(Particle->X());
+			V.Add(Particle->V());
+			TickSolverHelper(Module,Solver);
+		}
+
+		const FRewindData* RewindData = Solver->GetRewindData();
+
+
+		for(int Step = 0; Step < 9; ++Step)
+		{
+			const auto ParticleState = RewindData->GetStateAtFrame(*Particle,Step);
+			EXPECT_NE(ParticleState,nullptr);
+			
+			EXPECT_EQ(ParticleState->X()[2],X[Step][2]);
+			EXPECT_EQ(ParticleState->V()[2],V[Step][2]);
+		}
+
+		// Throw out the proxy
+		Solver->UnregisterObject(Particle.Get());
+
+		Module->DestroySolver(Solver);
+	}
 }
diff --git a/Engine/Source/Runtime/Experimental/ChaosSolvers/Public/RewindData.h b/Engine/Source/Runtime/Experimental/ChaosSolvers/Public/RewindData.h
--- a/Engine/Source/Runtime/Experimental/ChaosSolvers/Public/RewindData.h
+++ b/Engine/Source/Runtime/Experimental/ChaosSolvers/Public/RewindData.h
@@ -86,6 +86,21 @@ public:
 		return ParticlePositionRotation.IsSet() ? ParticlePositionRotation.Read().X : Particle.X();
 	}
 
+	const FRotation3& R() const
+	{
+		return ParticlePositionRotation.IsSet() ? ParticlePositionRotation.Read().R : Particle.R();
+	}
+
+	const FVec3& V() const
+	{
+		return Velocities.IsSet() ? Velocities.Read().V : Particle.CastToKinematicParticle()->V();
+	}
+
+	const FVec3& W() const
+	{
+		return Velocities.IsSet() ? Velocities.Read().W : Particle.CastToKinematicParticle()->W();
+	}
+
 	TSerializablePtr<FImplicitObject> Geometry() const
 	{
 		return NonFrequentData.IsSet() ? MakeSerializable(NonFrequentData.Read().Geometry) : Particle.Geometry();
@@ -110,7 +125,11 @@ public:
 			Data.R = Rigid.R();
 		});
 
-		//todo: velocities
+		Velocities.SyncRemoteData(Manager,Idx,Dirty,[&Rigid](auto& Data)
+		{
+			Data.V = Rigid.PreV();
+			Data.W = Rigid.PreW();
+		});
 	}
 
 	void SyncDirtyDynamics(FDirtyPropertiesManager& DestManager,int32 DataIdx, const FParticleDirtyData& Dirty, const FDirtyPropertiesManager& SrcManager)
@@ -143,6 +162,15 @@ public:
 				Data.X = Handle->X();
 				Data.R = Handle->R();
 			});
+
+			if(auto Kinematic = Handle->CastToKinematicParticle())
+			{
+				Velocities.SyncRemoteData(Manager,Idx,Dirty.ParticleData,[Kinematic](auto& Data)
+				{
+					Data.V = Kinematic->V();
+					Data.W = Kinematic->W();
+				});
+			}
 		}
 		
 		NonFrequentData.SyncRemoteData(Manager,Idx,Dirty.ParticleData, [Handle](FParticleNonFrequentData& Data)
@@ -174,6 +202,12 @@ public:
 			bCoalesced = true;
 		}
 
+		if(!Velocities.IsSet() && LatestState.Velocities.IsSet())
+		{
+			Velocities = LatestState.Velocities;
+			bCoalesced = true;
+		}
+
 		//dynamics do not coalesce since they are always written when dirty
 		
 		return bCoalesced;
@@ -184,6 +218,7 @@ protected:
 private:
 	TParticleStateProperty<FParticlePositionRotation,EParticleProperty::XR> ParticlePositionRotation;
 	TParticleStateProperty<FParticleNonFrequentData, EParticleProperty::NonFrequentData> NonFrequentData;
+	TParticleStateProperty<FParticleVelocities,EParticleProperty::Velocities> Velocities;
 	TParticleStateProperty<FParticleDynamics, EParticleProperty::Dynamics> Dynamics;
 	/*
 	PARTICLE_PROPERTY(XR,FParticlePositionRotation)
