Make sure to test teleport of simulated object

diff --git a/Engine/Source/Programs/HeadlessChaos/Private/HeadlessChaosTestRewind.cpp b/Engine/Source/Programs/HeadlessChaos/Private/HeadlessChaosTestRewind.cpp
--- a/Engine/Source/Programs/HeadlessChaos/Private/HeadlessChaosTestRewind.cpp
+++ b/Engine/Source/Programs/HeadlessChaos/Private/HeadlessChaosTestRewind.cpp
@@ -329,7 +329,7 @@ namespace ChaosTest {
 		Module->DestroySolver(Solver);
 	}
 
-	GTEST_TEST(RewindTest,FallingObject)
+	GTEST_TEST(RewindTest,FallingObjectWithTeleport)
 	{
 		auto Sphere = TSharedPtr<FImplicitObject,ESPMode::ThreadSafe>(new TSphere<float,3>(TVector<float,3>(0),10));
 
@@ -356,6 +356,13 @@ namespace ChaosTest {
 
 		for(int Step = 0; Step < 10; ++Step)
 		{
+			//teleport from GT
+			if(Step == 5)
+			{
+				Particle->SetX(FVec3(0,0,10));
+				Particle->SetV(FVec3(0,0,1));
+			}
+
 			X.Add(Particle->X());
 			V.Add(Particle->V());
 			TickSolverHelper(Module,Solver);
