fix clustering tests for HeadlessChaos

diff --git a/Engine/Source/Programs/HeadlessChaos/Private/HeadlessChaosTestClustering.cpp b/Engine/Source/Programs/HeadlessChaos/Private/HeadlessChaosTestClustering.cpp
--- a/Engine/Source/Programs/HeadlessChaos/Private/HeadlessChaosTestClustering.cpp
+++ b/Engine/Source/Programs/HeadlessChaos/Private/HeadlessChaosTestClustering.cpp
@@ -24,25 +24,29 @@ namespace ChaosTest {
 	template<class T>
 	void ImplicitCluster()
 	{
-#if TODO_REIMPLEMENT_RIGID_CLUSTERING
 		TPBDRigidsSOAs<T, 3> Particles;
 		THandleArray<FChaosPhysicsMaterial> PhysicalMaterials;
 		FPBDRigidsEvolutionGBF Evolution(Particles, PhysicalMaterials);
-		auto& ClusteredParticles = Particles.GetClusteredParticles();
+		TPBDRigidClusteredParticles<T, 3>& ClusteredParticles = Particles.GetClusteredParticles();
 		
 		uint32 FirstId = ClusteredParticles.Size();
-		auto& Box1 = *AppendClusteredParticleBox<T>(Particles, TVector<T, 3>((T)100, (T)100, (T)100));
+		TPBDRigidParticleHandle<T, 3>* Box1 = AppendClusteredParticleBox<T>(Particles, TVector<T, 3>((T)100, (T)100, (T)100));
 		uint32 BoxId = FirstId++;
-		auto& Box2 = *AppendClusteredParticleBox<T>(Particles, TVector<T, 3>((T)100, (T)100, (T)100));
+		TPBDRigidParticleHandle<T, 3>* Box2 = AppendClusteredParticleBox<T>(Particles, TVector<T, 3>((T)100, (T)100, (T)100));
 		uint32 Box2Id = FirstId++;
 		
-		Box2.X() = TVector<T, 3>((T)100, (T)0, (T)0);
-		Box2.P() = Box2.X();
+		Box2->X() = TVector<T, 3>((T)100, (T)0, (T)0);
+		Box2->P() = Box2->X();
 
 		Evolution.AdvanceOneTimeStep(0);	//hack to initialize islands
 		//Evolution.InitializeAccelerationStructures();	//make sure islands are created
+		FClusterCreationParameters<T> ClusterParams;
+		
+		TArray<Chaos::TPBDRigidParticleHandle<T, 3>*> ClusterChildren;
+		ClusterChildren.Add(Box1);
+		ClusterChildren.Add(Box2);
 
-		Evolution.GetRigidClustering().CreateClusterParticle(0,TArray<uint32>{ BoxId, Box2Id }, TSerializablePtr<FImplicitObject>(), nullptr);
+		Evolution.GetRigidClustering().CreateClusterParticle(0, MoveTemp(ClusterChildren), ClusterParams);
 		EXPECT_EQ(ClusteredParticles.Size(), 3);
 
 		TVector<T, 3> ClusterX = ClusteredParticles.X(2);
@@ -61,15 +65,12 @@ namespace ChaosTest {
 		EXPECT_TRUE(Normal.Equals(TVector<T, 3>{(T)1, (T)0, (T)0}));
 
 		//EXPECT_TRUE(Evolution.GetParticles().Geometry(2)->IsConvex());  we don't actually guarantee this
-#endif
 	}
 	template void ImplicitCluster<float>();
 
 	template<class T>
 	void FractureCluster()
 	{
-#if TODO_REIMPLEMENT_RIGID_CLUSTERING
-
 		TPBDRigidsSOAs<T, 3> Particles;
 		THandleArray<FChaosPhysicsMaterial> PhysicalMaterials;
 		FPBDRigidsEvolutionGBF Evolution(Particles, PhysicalMaterials);
@@ -78,12 +79,12 @@ namespace ChaosTest {
 		//create a long row of boxes - the depth 0 cluster is the entire row, the depth 1 clusters 4 boxes each, the depth 2 clusters are 1 box each
 
 		constexpr int32 NumBoxes = 32;
-		TArray<TGeometryParticleHandle<T, 3>*> Boxes;
+		TArray<TPBDRigidParticleHandle<T, 3>*> Boxes;
 		TArray<uint32> BoxIDs;
 		for (int i = 0; i < NumBoxes; ++i)
 		{
 			BoxIDs.Add(ClusteredParticles.Size());
-			auto Box = AppendClusteredParticleBox<T>(Particles, TVector<T, 3>((T)100, (T)100, (T)100));
+			TPBDRigidParticleHandle<T, 3>* Box = AppendClusteredParticleBox<T>(Particles, TVector<T, 3>((T)100, (T)100, (T)100));
 			Box->X() = TVector<T, 3>((T)i * (T)100, (T)0, (T)0);
 			Box->P() = Box->X();
 			Boxes.Add(Box);
@@ -91,16 +92,25 @@ namespace ChaosTest {
 
 		Evolution.AdvanceOneTimeStep(0);	//hack to generate islands
 
-		TArray<uint32> ClusterIDs;
+		TArray<Chaos::TPBDRigidParticleHandle<T, 3>* > ClusterHandles;
 
 		for (int i = 0; i < NumBoxes / 4; ++i)
 		{
-			ClusterIDs.Add(Evolution.GetRigidClustering().CreateClusterParticle(0, TArray<uint32>{ BoxIDs[i*4], BoxIDs[i*4+1], BoxIDs[i*4+2], BoxIDs[i*4+3] }, TSerializablePtr<FImplicitObject>(), nullptr));
+			FClusterCreationParameters<T> ClusterParams;
+
+			TArray<Chaos::TPBDRigidParticleHandle<float, 3>*> ClusterChildren;
+			ClusterChildren.Add(Boxes[i * 4]);
+			ClusterChildren.Add(Boxes[i * 4+1]);
+			ClusterChildren.Add(Boxes[i * 4+2]);
+			ClusterChildren.Add(Boxes[i * 4+3]);
+			ClusterHandles.Add(Evolution.GetRigidClustering().CreateClusterParticle(0, MoveTemp(ClusterChildren), ClusterParams));
 		}
 
-		uint32 RootClusterID = Evolution.GetRigidClustering().CreateClusterParticle(0, TArray<uint32>(ClusterIDs), TSerializablePtr<FImplicitObject>(), nullptr);
+		FClusterCreationParameters<T> ClusterParams;
+		Chaos::TPBDRigidParticleHandle<T, 3>* RootClusterHandle = Evolution.GetRigidClustering().CreateClusterParticle(0, MoveTemp(ClusterHandles), ClusterParams);
 		TVector<T, 3> InitialVelocity((T)50, (T)20, (T)100);
-		ClusteredParticles.V(RootClusterID) = InitialVelocity;
+
+		RootClusterHandle->V() = InitialVelocity;		
 		
 		constexpr int NumParticles = NumBoxes + NumBoxes / 4 + 1;
 		EXPECT_EQ(ClusteredParticles.Size(), NumParticles);
@@ -110,16 +120,16 @@ namespace ChaosTest {
 			EXPECT_TRUE(ClusteredParticles.Disabled(i));
 		}
 
-		EXPECT_TRUE(ClusteredParticles.Disabled(RootClusterID) == false);
+		EXPECT_TRUE(RootClusterHandle->Disabled() == false);
 		EXPECT_EQ(Particles.GetNonDisabledView().Num(), 1);
-		EXPECT_EQ(Particles.GetNonDisabledView().Begin()->Handle(), ClusteredParticles.Handle(RootClusterID));
+		EXPECT_EQ(Particles.GetNonDisabledView().Begin()->Handle(), RootClusterHandle);
 
 		const T Dt = 0;	//don't want to integrate gravity, just want to test fracture
 		Evolution.AdvanceOneTimeStep(Dt);
-		EXPECT_TRUE(ClusteredParticles.Disabled(RootClusterID));	//not a cluster anymore, so disabled
+		EXPECT_TRUE(RootClusterHandle->Disabled());	//not a cluster anymore, so disabled
 		for (auto& Particle : Particles.GetNonDisabledView())
 		{
-			EXPECT_NE(Particle.Handle(), ClusteredParticles.Handle(RootClusterID));	//view no longer contains root
+			EXPECT_NE(Particle.Handle(), RootClusterHandle);	//view no longer contains root
 		}
 
 		EXPECT_EQ(Particles.GetNonDisabledView().Num(), NumBoxes / 4);
@@ -139,55 +149,54 @@ namespace ChaosTest {
 			}
 		}
 
-		for (uint32 ClusterID : ClusterIDs)
+		for (Chaos::TPBDRigidParticleHandle<T, 3>* ClusterHandle : ClusterHandles)
 		{
-			EXPECT_TRUE(ClusteredParticles.Disabled(ClusterID) == false);	//not a cluster anymore, so disabled
+			EXPECT_TRUE(ClusterHandle->Disabled() == false);	//not a cluster anymore, so disabled
 			bool bFoundInNonDisabled = false;
 			for (auto& Particle : Particles.GetNonDisabledView())
 			{
-				bFoundInNonDisabled |= Particle.Handle() == ClusteredParticles.Handle(ClusterID);
+				bFoundInNonDisabled |= Particle.Handle() == ClusterHandle;
 			}
 
 			EXPECT_TRUE(bFoundInNonDisabled);	//clusters are enabled and in non disabled array
-			EXPECT_TRUE(ClusteredParticles.V(ClusterID).Equals(InitialVelocity));
+			EXPECT_TRUE(ClusterHandle->V().Equals(InitialVelocity));
 		}
 
 		Evolution.AdvanceOneTimeStep(Dt);
 		//second fracture, all clusters are now disabled
-		for (uint32 ClusterID : ClusterIDs)
+		for (Chaos::TPBDRigidParticleHandle<T, 3>* ClusterHandle : ClusterHandles)
 		{
-			EXPECT_TRUE(ClusteredParticles.Disabled(ClusterID) == true);	//not a cluster anymore, so disabled
+			EXPECT_TRUE(ClusterHandle->Disabled() == true);	//not a cluster anymore, so disabled
 			for (auto& Particle : Particles.GetNonDisabledView())
 			{
-				EXPECT_NE(Particle.Handle(), ClusteredParticles.Handle(ClusterID));	//make sure boxes are not in non disabled array
+				EXPECT_NE(Particle.Handle(), ClusterHandle);	//make sure boxes are not in non disabled array
 			}
 
 			for (int32 Island = 0; Island < Evolution.NumIslands(); ++Island)
 			{
-				EXPECT_TRUE(Evolution.GetIslandParticles(Island).Contains(ClusteredParticles.Handle(ClusterID)) == false);
+				EXPECT_TRUE(Evolution.GetIslandParticles(Island).Contains(ClusterHandle) == false);
 			}
 		}
 
 		EXPECT_EQ(Particles.GetNonDisabledView().Num(), NumBoxes);
-		for (uint32 BoxID : BoxIDs)
+		
+		for (TPBDRigidParticleHandle<T, 3>* BoxHandle : Boxes)
 		{
-			EXPECT_TRUE(ClusteredParticles.Disabled(BoxID) == false);
+			EXPECT_TRUE(BoxHandle->Disabled() == false);
 			bool bFoundInNonDisabled = false;
 			for (auto& Particle : Particles.GetNonDisabledView())
 			{
-				bFoundInNonDisabled |= Particle.Handle() == ClusteredParticles.Handle(BoxID);
+				bFoundInNonDisabled |= Particle.Handle() == BoxHandle;
 			}
 			EXPECT_TRUE(bFoundInNonDisabled);
-			EXPECT_TRUE(ClusteredParticles.V(BoxID).Equals(InitialVelocity));
+			EXPECT_TRUE(BoxHandle->V().Equals(InitialVelocity));
 		}
-#endif
 	}
 	template void FractureCluster<float>();
 
 	template<class T>
 	void PartialFractureCluster()
 	{
-#if TODO_REIMPLEMENT_RIGID_CLUSTERING
 		TPBDRigidsSOAs<T, 3> Particles;
 		THandleArray<FChaosPhysicsMaterial> PhysicalMaterials;
 		FPBDRigidsEvolutionGBF Evolution(Particles, PhysicalMaterials);
@@ -196,12 +205,12 @@ namespace ChaosTest {
 		//create a long row of boxes - the depth 0 cluster is the entire row, the depth 1 clusters 4 boxes each, the depth 2 clusters are 1 box each
 
 		constexpr int32 NumBoxes = 32;
-		TArray<TGeometryParticleHandle<T, 3>*> Boxes;
+		TArray<TPBDRigidParticleHandle<T, 3>*> Boxes;
 		TArray<uint32> BoxIDs;
 		for (int i = 0; i < NumBoxes; ++i)
 		{
 			BoxIDs.Add(ClusteredParticles.Size());
-			auto Box = AppendClusteredParticleBox<T>(Particles, TVector<T, 3>((T)100, (T)100, (T)100));
+			TPBDRigidParticleHandle<T, 3>* Box = AppendClusteredParticleBox<T>(Particles, TVector<T, 3>((T)100, (T)100, (T)100));
 			Box->X() = TVector<T, 3>((T)i * (T)100, (T)0, (T)0);
 			Box->P() = Box->X();
 			Boxes.Add(Box);
@@ -209,18 +218,27 @@ namespace ChaosTest {
 
 		Evolution.AdvanceOneTimeStep(0);	//hack to generate islands
 
-		TArray<uint32> ClusterIDs;
+		TArray<Chaos::TPBDRigidParticleHandle<T, 3>* > ClusterHandles;
 
 		for (int i = 0; i < NumBoxes / 4; ++i)
 		{
-			ClusterIDs.Add(Evolution.GetRigidClustering().CreateClusterParticle(0, TArray<uint32>{ BoxIDs[i * 4], BoxIDs[i * 4 + 1], BoxIDs[i * 4 + 2], BoxIDs[i * 4 + 3] }, TSerializablePtr<FImplicitObject>(), nullptr));
+			FClusterCreationParameters<T> ClusterParams;
+
+			TArray<Chaos::TPBDRigidParticleHandle<float, 3>*> ClusterChildren;
+			ClusterChildren.Add(Boxes[i * 4]);
+			ClusterChildren.Add(Boxes[i * 4 + 1]);
+			ClusterChildren.Add(Boxes[i * 4 + 2]);
+			ClusterChildren.Add(Boxes[i * 4 + 3]);
+			ClusterHandles.Add(Evolution.GetRigidClustering().CreateClusterParticle(0, MoveTemp(ClusterChildren), ClusterParams));
 		}
 
-		TArray<uint32> ClusterIDsDup = ClusterIDs;
+		TArray<Chaos::TPBDRigidParticleHandle<T, 3>* > ClusterHandlesDup = ClusterHandles;
 
-		uint32 RootClusterID = Evolution.GetRigidClustering().CreateClusterParticle(0, TArray<uint32>(ClusterIDs), TSerializablePtr<FImplicitObject>(), nullptr);
+		FClusterCreationParameters<T> ClusterParams;
+		Chaos::TPBDRigidParticleHandle<T, 3>* RootClusterHandle = Evolution.GetRigidClustering().CreateClusterParticle(0, MoveTemp(ClusterHandles), ClusterParams);
 		TVector<T, 3> InitialVelocity((T)50, (T)20, (T)100);
-		ClusteredParticles.V(RootClusterID) = InitialVelocity;
+
+		RootClusterHandle->V() = InitialVelocity;
 
 		TUniquePtr<FChaosPhysicsMaterial> PhysicalMaterial = MakeUnique<FChaosPhysicsMaterial>();
 		PhysicalMaterial->Friction = 0;
@@ -239,12 +257,14 @@ namespace ChaosTest {
 		}
 
 		Evolution.AdvanceOneTimeStep((T)1 / (T)60);
-		EXPECT_TRUE(ClusteredParticles.Disabled(RootClusterID) == false);	//strain > 0 so no fracture yet
+		EXPECT_TRUE(RootClusterHandle->Disabled() == false);	//strain > 0 so no fracture yet
+
+		// todo: is this the correct replacement for strain?
+		static_cast<Chaos::TPBDRigidClusteredParticleHandle<T, 3>*>(ClusterHandlesDup[2])->SetStrain((T)0);	//fracture the third cluster, this should leave us with three pieces (0, 1), (2), (3,4,5,6,7)
 
-		SolverStrainArray[ClusterIDsDup[2]] = (T)0;	//fracture the third cluster, this should leave us with three pieces (0, 1), (2), (3,4,5,6,7)
 		Evolution.AdvanceOneTimeStep((T)1 / (T)60);
-		//EXPECT_TRUE(Evolution.GetParticles().Disabled(RootClusterID) == false);	//one of the connected pieces should re-use this
-		EXPECT_TRUE(ClusteredParticles.Disabled(ClusterIDsDup[2]) == false);	//this cluster is on its own and should be enabled 
+		//EXPECT_TRUE(Evolution.GetParticles().Disabled(RootClusterHandle) == false);	//one of the connected pieces should re-use this
+		EXPECT_TRUE(ClusterHandlesDup[2]->Disabled() == false);	//this cluster is on its own and should be enabled 
 		
 		EXPECT_EQ(Evolution.GetActiveClusteredArray().Num(), 3);	//there should only be 3 pieces
 		for (uint32 BoxID : BoxIDs)
@@ -263,8 +283,8 @@ namespace ChaosTest {
 		Evolution.SetPhysicsMaterial(ClusteredParticles.Handle(NumBoxes + NumBoxes / 4 + 2), MakeSerializable(PhysicalMaterial));
 
 		Evolution.AdvanceOneTimeStep((T)1 / (T)60);	//next frame nothing should fracture
-		//EXPECT_TRUE(Evolution.GetParticles().Disabled(RootClusterID) == false);	//one of the connected pieces should re-use this
-		EXPECT_TRUE(ClusteredParticles.Disabled(ClusterIDsDup[2]) == false);	//this cluster is on its own and should be enabled 
+		//EXPECT_TRUE(Evolution.GetParticles().Disabled(RootClusterHandle) == false);	//one of the connected pieces should re-use this
+		EXPECT_TRUE(ClusterHandlesDup[2]->Disabled() == false);	//this cluster is on its own and should be enabled 
 
 		EXPECT_EQ(Evolution.GetActiveClusteredArray().Num(), 3);	//there should only be 3 pieces
 		for (uint32 BoxID : BoxIDs)
@@ -276,11 +296,7 @@ namespace ChaosTest {
 				EXPECT_TRUE(Evolution.GetIslandParticles(Island).Contains(ClusteredParticles.Handle(BoxID)) == false);
 			}
 		}
-#endif
 	}
 	template void PartialFractureCluster<float>();
-
-
-
 }
 
