Chaos
- remove templates on collision types
- fix collision rejection in main evolution

diff --git a/Engine/Source/Programs/HeadlessChaos/Private/HeadlessChaosCollisionConstraints.cpp b/Engine/Source/Programs/HeadlessChaos/Private/HeadlessChaosCollisionConstraints.cpp
--- a/Engine/Source/Programs/HeadlessChaos/Private/HeadlessChaosCollisionConstraints.cpp
+++ b/Engine/Source/Programs/HeadlessChaos/Private/HeadlessChaosCollisionConstraints.cpp
@@ -6,6 +6,5 @@
 #include "Chaos/Utilities.h"
 
 namespace Chaos {
-	template class TPBDCollisionConstraintAccessor<float, 3>;
 }
 
diff --git a/Engine/Source/Programs/HeadlessChaos/Private/HeadlessChaosTestCollisions.cpp b/Engine/Source/Programs/HeadlessChaos/Private/HeadlessChaosTestCollisions.cpp
--- a/Engine/Source/Programs/HeadlessChaos/Private/HeadlessChaosTestCollisions.cpp
+++ b/Engine/Source/Programs/HeadlessChaos/Private/HeadlessChaosTestCollisions.cpp
@@ -17,9 +17,9 @@
 #define SMALL_THRESHOLD 1e-4
 
 #define RESET_PQ(Particle) Particle->P() = Particle->X(); Particle->Q() = Particle->R()
-#define INVARIANT_XR_START(Particle) TVector<T,3> InvariantPreX_##Particle = Particle->X(); TRotation<T,3> InvariantPreR_##Particle = Particle->R()
+#define INVARIANT_XR_START(Particle) TVector<T, 3> InvariantPreX_##Particle = Particle->X(); TRotation<T, 3> InvariantPreR_##Particle = Particle->R()
 #define INVARIANT_XR_END(Particle) EXPECT_TRUE(InvariantPreX_##Particle.Equals(Particle->X())); EXPECT_TRUE(InvariantPreR_##Particle.Equals(Particle->R()))
-#define INVARIANT_VW_START(Particle) TVector<T,3> InvariantPreV_##Particle = Particle->V(); TVector<T,3> InvariantPreW_##Particle = Particle->W()
+#define INVARIANT_VW_START(Particle) TVector<T, 3> InvariantPreV_##Particle = Particle->V(); TVector<T, 3> InvariantPreW_##Particle = Particle->W()
 #define INVARIANT_VW_END(Particle) EXPECT_TRUE(InvariantPreV_##Particle.Equals(Particle->V())); EXPECT_TRUE(InvariantPreW_##Particle.Equals(Particle->W()))
 
 namespace ChaosTest {
@@ -55,17 +55,17 @@ namespace ChaosTest {
 		Box2->Q() = Box2->R();
 		Box2->AuxilaryValue(PhysicsMaterials) = MakeSerializable(PhysicsMaterial);
 
-		TPBDCollisionConstraintAccessor<T, 3> Collisions(Particles, Collided, PhysicsMaterials, 1, 1, T(0));
+		FPBDCollisionConstraintAccessor Collisions(Particles, Collided, PhysicsMaterials, 1, 1, T(0));
 		Collisions.ComputeConstraints(0.f);
 		EXPECT_EQ(Collisions.NumConstraints(), 1);
 
-		TPBDCollisionConstraints<T, 3>::FConstraintBase& Constraint = Collisions.GetConstraint(0);
+		FCollisionConstraintBase& Constraint = Collisions.GetConstraint(0);
 		if (auto PBDRigid = Constraint.Particle[0]->CastToRigidParticle())
 		{
 			//Question: non dynamics don't have collision particles, seems wrong if the levelset is dynamic and the static is something like a box
 			PBDRigid->CollisionParticles()->UpdateAccelerationStructures();
 		}
-		Collisions.UpdateLevelsetConstraint(*Constraint.template As<TPBDCollisionConstraints<T,3>::FPointContactConstraint>());
+		Collisions.UpdateLevelsetConstraint(*Constraint.template As<FPBDCollisionConstraints::FPointContactConstraint>());
 
 		EXPECT_EQ(Constraint.Particle[0], Box2);
 		EXPECT_EQ(Constraint.Particle[1], Box1);
@@ -88,7 +88,7 @@ namespace ChaosTest {
 		Particles.GetParticleHandles().AddArray(&Collided);
 		Particles.GetParticleHandles().AddArray(&PhysicsMaterials);
 
-		auto Box1 = AppendDynamicParticleConvexBox<T>(Particles, TVector<T,3>(1.f) );
+		auto Box1 = AppendDynamicParticleConvexBox<T>(Particles, TVector<T, 3>(1.f) );
 		Box1->X() = TVector<T, 3>(0.f);
 		Box1->R() = TRotation<T, 3>(FQuat::Identity);
 		Box1->P() = Box1->X();
@@ -102,12 +102,12 @@ namespace ChaosTest {
 		Box2->Q() = Box2->R();
 		Box2->AuxilaryValue(PhysicsMaterials) = MakeSerializable(PhysicsMaterial);
 
-		TPBDCollisionConstraintAccessor<T, 3> Collisions(Particles, Collided, PhysicsMaterials, 1, 1, T(0));
+		FPBDCollisionConstraintAccessor Collisions(Particles, Collided, PhysicsMaterials, 1, 1, T(0));
 		Collisions.ComputeConstraints(0.f);
 		EXPECT_EQ(Collisions.NumConstraints(), 1);
 
-		TPBDCollisionConstraints<T, 3>::FConstraintBase & Constraint = Collisions.GetConstraint(0);
-		Collisions.UpdateLevelsetConstraint(*Constraint.template As<TPBDCollisionConstraints<T,3>::FPointContactConstraint>());
+		FCollisionConstraintBase & Constraint = Collisions.GetConstraint(0);
+		Collisions.UpdateLevelsetConstraint(*Constraint.template As<FPBDCollisionConstraints::FPointContactConstraint>());
 		
 		EXPECT_EQ(Constraint.Particle[0], Box2);
 		EXPECT_EQ(Constraint.Particle[1], Box1);
@@ -141,17 +141,17 @@ namespace ChaosTest {
 
 		const float Dt = 1 / 24.;
 
-		TPBDCollisionConstraintAccessor<T, 3> Collisions(Particles, Collided, PhysicsMaterials, 2, 5, T(0));
+		FPBDCollisionConstraintAccessor Collisions(Particles, Collided, PhysicsMaterials, 2, 5, T(0));
 
 		Collisions.ComputeConstraints(Dt);
 		EXPECT_EQ(Collisions.NumConstraints(), 1);
 
-		TPBDCollisionConstraints<T, 3>::FConstraintBase & Constraint = Collisions.GetConstraint(0);
+		FCollisionConstraintBase & Constraint = Collisions.GetConstraint(0);
 		if (auto PBDRigid = Constraint.Particle[0]->CastToRigidParticle())
 		{
 			PBDRigid->CollisionParticles()->UpdateAccelerationStructures();
 		}
-		Collisions.UpdateLevelsetConstraint(*Constraint.template As<TPBDCollisionConstraints<T, 3>::FPointContactConstraint>());
+		Collisions.UpdateLevelsetConstraint(*Constraint.template As<FPBDCollisionConstraints::FPointContactConstraint>());
 
 		EXPECT_EQ(Constraint.Particle[0], Box);
 		EXPECT_EQ(Constraint.Particle[1], Floor);
@@ -174,7 +174,7 @@ namespace ChaosTest {
 			{
 				INVARIANT_XR_START(Box);
 				INVARIANT_VW_START(Box);
-				Collisions.ApplyPushOut(Dt, { Collisions.GetConstraintHandle(0) }, TSet<const TGeometryParticleHandle<T,3>*>(), 0, 1);
+				Collisions.ApplyPushOut(Dt, { Collisions.GetConstraintHandle(0) }, TSet<const TGeometryParticleHandle<T, 3>*>(), 0, 1);
 				INVARIANT_XR_END(Box);
 				INVARIANT_VW_END(Box);
 			}
@@ -211,20 +211,21 @@ namespace ChaosTest {
 
 		const float Dt = 1 / 24.;
 
-		TPBDCollisionConstraintAccessor<T, 3> Collisions(Particles, Collided, PhysicsMaterials, 2, 5, T(0));
+		FPBDCollisionConstraintAccessor Collisions(Particles, Collided, PhysicsMaterials, 2, 5, T(0));
 
 		Collisions.ComputeConstraints(Dt);
 		EXPECT_EQ(Collisions.NumConstraints(), 1);
 
-		TPBDCollisionConstraints<T, 3>::FMultiPointContactConstraint * Constraint = Collisions.GetConstraint(0).template As<TPBDCollisionConstraints<T, 3>::FMultiPointContactConstraint>();
+		FRigidBodyMultiPointContactConstraint * Constraint = Collisions.GetConstraint(0).template As<FRigidBodyMultiPointContactConstraint>();
 		EXPECT_TRUE(Constraint != nullptr);
 
-		Collisions.UpdateManifold(*Constraint);
+		//Collisions.UpdateManifold(*Constraint);
+		Collisions.Update(*Constraint);
 
 		EXPECT_EQ(Constraint->Particle[0], Box);
 		EXPECT_EQ(Constraint->Particle[1], Floor);
 		EXPECT_TRUE(Constraint->GetNormal().operator==(TVector<T, 3>(0, 0, 1)));
-		EXPECT_TRUE(FMath::Abs( Constraint->GetLocation().Z - TVector<T,3>(0,0,-1).Z ) < SMALL_THRESHOLD);
+		EXPECT_TRUE(FMath::Abs( Constraint->GetLocation().Z - TVector<T, 3>(0,0,-1).Z ) < SMALL_THRESHOLD);
 		EXPECT_TRUE(FMath::Abs(Constraint->GetPhi() - T(-1.0)) < SMALL_THRESHOLD);
 
 		{
@@ -280,17 +281,17 @@ namespace ChaosTest {
 
 		const float Dt = 1 / 24.;
 
-		TPBDCollisionConstraintAccessor<T, 3> Collisions(Particles, Collided, PhysicsMaterials, 2, 5, T(0));
+		FPBDCollisionConstraintAccessor Collisions(Particles, Collided, PhysicsMaterials, 2, 5, T(0));
 
 		Collisions.ComputeConstraints(Dt);
 		EXPECT_EQ(Collisions.NumConstraints(), 1);
 
-		TPBDCollisionConstraints<T, 3>::FConstraintBase & Constraint = Collisions.GetConstraint(0);
+		FCollisionConstraintBase & Constraint = Collisions.GetConstraint(0);
 		if (auto PBDRigid = Constraint.Particle[0]->CastToRigidParticle())
 		{
 			PBDRigid->CollisionParticles()->UpdateAccelerationStructures();
 		}
-		Collisions.UpdateLevelsetConstraint(*Constraint.template As<TPBDCollisionConstraints<T,3>::FPointContactConstraint>());
+		Collisions.UpdateLevelsetConstraint(*Constraint.template As<FPBDCollisionConstraints::FPointContactConstraint>());
 
 		EXPECT_EQ(Constraint.Particle[0], Box);
 		EXPECT_EQ(Constraint.Particle[1], Floor);
@@ -313,7 +314,7 @@ namespace ChaosTest {
 			{
 				INVARIANT_XR_START(Box);
 				INVARIANT_VW_START(Box);
-				Collisions.ApplyPushOut(Dt, { Collisions.GetConstraintHandle(0) }, TSet<const TGeometryParticleHandle<T,3>*>(), 0, 1);
+				Collisions.ApplyPushOut(Dt, { Collisions.GetConstraintHandle(0) }, TSet<const TGeometryParticleHandle<T, 3>*>(), 0, 1);
 				INVARIANT_XR_END(Box);
 				INVARIANT_VW_END(Box);
 			}
@@ -348,17 +349,17 @@ namespace ChaosTest {
 
 		const float Dt = 1 / 24.;
 
-		TPBDCollisionConstraintAccessor<T, 3> Collisions(Particles, Collided, PhysicsMaterials, 2, 5, T(0));
+		FPBDCollisionConstraintAccessor Collisions(Particles, Collided, PhysicsMaterials, 2, 5, T(0));
 
 		Collisions.ComputeConstraints(Dt);
 		EXPECT_EQ(Collisions.NumConstraints(), 1);
 
-		TPBDCollisionConstraints<T, 3>::FConstraintBase & Constraint = Collisions.GetConstraint(0);
+		FCollisionConstraintBase & Constraint = Collisions.GetConstraint(0);
 		if (auto PBDRigid = Constraint.Particle[0]->CastToRigidParticle())
 		{
 			PBDRigid->CollisionParticles()->UpdateAccelerationStructures();
 		}
-		Collisions.UpdateLevelsetConstraint(*Constraint.template As<TPBDCollisionConstraints<T,3>::FPointContactConstraint>());
+		Collisions.UpdateLevelsetConstraint(*Constraint.template As<FPBDCollisionConstraints::FPointContactConstraint>());
 		EXPECT_EQ(Constraint.Particle[0], Box);
 		EXPECT_EQ(Constraint.Particle[1], Floor);
 		EXPECT_TRUE(Constraint.GetNormal().operator==(TVector<T, 3>(0, 0, 1)));
@@ -383,7 +384,7 @@ namespace ChaosTest {
 			{
 				INVARIANT_XR_START(Box);
 				INVARIANT_VW_START(Box);
-				Collisions.ApplyPushOut(Dt, { Collisions.GetConstraintHandle(0) }, TSet<const TGeometryParticleHandle<T,3>*>(), 0, 1);
+				Collisions.ApplyPushOut(Dt, { Collisions.GetConstraintHandle(0) }, TSet<const TGeometryParticleHandle<T, 3>*>(), 0, 1);
 				INVARIANT_XR_END(Box);
 				INVARIANT_VW_END(Box);
 			}
@@ -424,11 +425,12 @@ namespace ChaosTest {
 
 		float Dt = 1 / 24.;
 
-		TPBDCollisionConstraintAccessor<T, 3> Collisions(Particles, Collided, PhysicsMaterials, 1, 1, T(0));
+		FPBDCollisionConstraintAccessor Collisions(Particles, Collided, PhysicsMaterials, 1, 1, T(0));
 		Collisions.ComputeConstraints(Dt);
 		EXPECT_EQ(Collisions.NumConstraints(), 1);
 
-		TPBDCollisionConstraints<T, 3>::FConstraintBase & Constraint = Collisions.GetConstraint(0);
+		FCollisionConstraintBase & Constraint = Collisions.GetConstraint(0);
+		Collisions.Update(Constraint, 100.0f);
 
 		//Collisions.UpdateLevelsetConstraintGJK(Particles, Constraint);
 		//EXPECT_EQ(Constraint.ParticleIndex, 1);
diff --git a/Engine/Source/Programs/HeadlessChaos/Private/HeadlessChaosTestMinEvolution.cpp b/Engine/Source/Programs/HeadlessChaos/Private/HeadlessChaosTestMinEvolution.cpp
--- a/Engine/Source/Programs/HeadlessChaos/Private/HeadlessChaosTestMinEvolution.cpp
+++ b/Engine/Source/Programs/HeadlessChaos/Private/HeadlessChaosTestMinEvolution.cpp
@@ -23,7 +23,7 @@ namespace ChaosTest
 	GTEST_TEST(MinEvolutionTests, TestSpringConstraints)
 	{
 		// @todo(ccaulfield): remove template parameters on collisions and other constraints
-		using FCollisionConstraints = TPBDCollisionConstraints<FReal, 3>;
+		using FCollisionConstraints = FPBDCollisionConstraints;
 		using FCollisionDetector = TCollisionDetector<FParticlePairBroadPhase, FNarrowPhase, FSyncCollisionReceiver, FCollisionConstraints>;
 		using FRigidParticleSOAs = TPBDRigidsSOAs<FReal, 3>;
 		using FParticleHandle = TPBDRigidParticleHandle<FReal, 3>;
diff --git a/Engine/Source/Programs/HeadlessChaos/Private/HeadlessChaosTestSim.cpp b/Engine/Source/Programs/HeadlessChaos/Private/HeadlessChaosTestSim.cpp
--- a/Engine/Source/Programs/HeadlessChaos/Private/HeadlessChaosTestSim.cpp
+++ b/Engine/Source/Programs/HeadlessChaos/Private/HeadlessChaosTestSim.cpp
@@ -34,7 +34,7 @@ namespace ChaosTest {
 
 
 		const FReal Dt = 1 / 60.f;
-		for (int i = 0; i < 100; ++i)
+		for (int i = 0; i < 200; ++i)
 		{
 			Evolution.AdvanceOneTimeStep(1 / 60.f);
 			Evolution.EndFrame(Dt);
diff --git a/Engine/Source/Programs/HeadlessChaos/Public/HeadlessChaosCollisionConstraints.h b/Engine/Source/Programs/HeadlessChaos/Public/HeadlessChaosCollisionConstraints.h
--- a/Engine/Source/Programs/HeadlessChaos/Public/HeadlessChaosCollisionConstraints.h
+++ b/Engine/Source/Programs/HeadlessChaos/Public/HeadlessChaosCollisionConstraints.h
@@ -23,46 +23,43 @@ namespace Chaos
 /**
  * Test collision constraints.
  */
-template <class T, int d>
-class TPBDCollisionConstraintAccessor
+class FPBDCollisionConstraintAccessor
 {
 public:
-	using FCollisionConstraints = TPBDCollisionConstraints<T, d>;
-	using FConstraintContainerHandle = TPBDCollisionConstraintHandle<T, d>;
-	using FContactConstraintBase = TCollisionConstraintBase<T, d>;
-	using FMultiPointContactConstraint = TRigidBodyMultiPointContactConstraint<T, d>;
-	using FPointContactConstraint = TRigidBodyPointContactConstraint<T, d>;
-	using FMultiPointContactConstraint = TRigidBodyMultiPointContactConstraint<T, d>;
-	using FConstraintHandleAllocator = TConstraintHandleAllocator<TPBDCollisionConstraints<T, d>>;
-	using FConstraintHandleID = TPair<const TGeometryParticleHandle<T, d>*, const TGeometryParticleHandle<T, d>*>;
+	using FCollisionConstraints = FPBDCollisionConstraints;
+	using FConstraintContainerHandle = FPBDCollisionConstraintHandle;
+	using FMultiPointContactConstraint = FRigidBodyMultiPointContactConstraint;
+	using FPointContactConstraint = FRigidBodyPointContactConstraint;
+	using FConstraintHandleAllocator = TConstraintHandleAllocator<FPBDCollisionConstraints>;
+	using FConstraintHandleID = TPair<const TGeometryParticleHandle<FReal, 3>*, const TGeometryParticleHandle<FReal, 3>*>;
 	using FCollisionDetector = TCollisionDetector<FSpatialAccelerationBroadPhase, FNarrowPhase, FAsyncCollisionReceiver, FCollisionConstraints>;
-	using FAccelerationStructure = TBoundingVolume<TAccelerationStructureHandle<T, d>, T, d>;
+	using FAccelerationStructure = TBoundingVolume<TAccelerationStructureHandle<FReal, 3>, FReal, 3>;
 
-	TPBDCollisionConstraintAccessor() 
+	FPBDCollisionConstraintAccessor()
 		: SpatialAcceleration(EmptyParticles.GetNonDisabledView())
-		, BroadPhase(EmptyParticles, (T)1, (T)0)
+		, BroadPhase(EmptyParticles, (FReal)1, (FReal)0)
 		, CollisionConstraints(EmptyParticles, EmptyCollided, EmptyPhysicsMaterials, 1, 1)
 		, CollisionDetector(BroadPhase, CollisionConstraints)
 	{}
 
-	TPBDCollisionConstraintAccessor(const TPBDRigidsSOAs<T, d>& InParticles, TArrayCollectionArray<bool>& Collided, const TArrayCollectionArray<TSerializablePtr<FChaosPhysicsMaterial>>& PerParticleMaterials,
-	const int32 PushOutIterations, const int32 PushOutPairIterations, const T Thickness) 
+	FPBDCollisionConstraintAccessor(const TPBDRigidsSOAs<FReal, 3>& InParticles, TArrayCollectionArray<bool>& Collided, const TArrayCollectionArray<TSerializablePtr<FChaosPhysicsMaterial>>& PerParticleMaterials,
+	const int32 PushOutIterations, const int32 PushOutPairIterations, const FReal Thickness) 
 		: SpatialAcceleration(InParticles.GetNonDisabledView())
-		, BroadPhase(InParticles, (T)1, (T)0)
+		, BroadPhase(InParticles, (FReal)1, (FReal)0)
 		, CollisionConstraints(InParticles, Collided, PerParticleMaterials, 1, 1, Thickness)
 		, CollisionDetector(BroadPhase, CollisionConstraints)
 	{}
 
-	virtual ~TPBDCollisionConstraintAccessor() {}
+	virtual ~FPBDCollisionConstraintAccessor() {}
 	
-	void ComputeConstraints(T Dt)
+	void ComputeConstraints(FReal Dt)
 	{
 		CollisionStats::FStatData StatData(false);
 		CollisionDetector.GetBroadPhase().SetSpatialAcceleration(&SpatialAcceleration);
 		CollisionDetector.DetectCollisions(Dt, StatData);
 	}
 
-	void Update(FContactConstraintBase& Constraint, T CullDistance = T(0) )
+	void Update(FCollisionConstraintBase& Constraint, FReal CullDistance = FReal(0) )
 	{
 		if (Constraint.GetType() == FPointContactConstraint::StaticType())
 		{
@@ -74,9 +71,9 @@ public:
 		}
 	}
 
-	void UpdateManifold(FContactConstraintBase& Constraint, T CullDistance = T(0))
+	void UpdateManifold(FCollisionConstraintBase& Constraint, FReal CullDistance = FReal(0))
 	{
-		if (Constraint.GetType() == FContactConstraintBase::FType::MultiPoint)
+		if (Constraint.GetType() == FCollisionConstraintBase::FType::MultiPoint)
 		{
 			FCollisionContext Context;
 			Collisions::UpdateManifold(*Constraint.As<FRigidBodyMultiPointContactConstraint>(), CullDistance, Context);
@@ -86,7 +83,7 @@ public:
 
 	void UpdateLevelsetConstraint(FPointContactConstraint& Constraint) 
 	{
-		Collisions::UpdateLevelsetLevelsetConstraint<ECollisionUpdateType::Deepest>(T(0), Constraint);
+		Collisions::UpdateLevelsetLevelsetConstraint<ECollisionUpdateType::Deepest>(FReal(0), Constraint);
 	}
 
 	int32 NumConstraints() const
@@ -94,7 +91,7 @@ public:
 		return CollisionConstraints.NumConstraints();
 	}
 
-	TCollisionConstraintBase<T,d>& GetConstraint(int32 Index)
+	FCollisionConstraintBase& GetConstraint(int32 Index)
 	{
 		if (Index < CollisionConstraints.NumConstraints())
 		{
@@ -113,18 +110,18 @@ public:
 		return CollisionConstraints.GetAllConstraintHandles()[ConstraintIndex];
 	}
 
-	void Apply(const T Dt, const TArray<FConstraintContainerHandle*>& InConstraintHandles, const int32 It, const int32 NumIts)
+	void Apply(const FReal Dt, const TArray<FConstraintContainerHandle*>& InConstraintHandles, const int32 It, const int32 NumIts)
 	{
 		CollisionConstraints.Apply(Dt, InConstraintHandles, It, NumIts);
 	}
 
-	bool ApplyPushOut(const T Dt, const TArray<FConstraintContainerHandle*>& InConstraintHandles, const TSet<const TGeometryParticleHandle<T, d>*>& IsTemporarilyStatic, int32 Iteration, int32 NumIterations)
+	bool ApplyPushOut(const FReal Dt, const TArray<FConstraintContainerHandle*>& InConstraintHandles, const TSet<const TGeometryParticleHandle<FReal, 3>*>& IsTemporarilyStatic, int32 Iteration, int32 NumIterations)
 	{
 		return CollisionConstraints.ApplyPushOut(Dt, InConstraintHandles, IsTemporarilyStatic, Iteration, NumIterations);
 	}
 
 	FPointContactConstraint EmptyConstraint;
-	TPBDRigidsSOAs<T, d> EmptyParticles;
+	TPBDRigidsSOAs<FReal, 3> EmptyParticles;
 	TArrayCollectionArray<bool> EmptyCollided;
 	TArrayCollectionArray<TSerializablePtr<FChaosPhysicsMaterial>> EmptyPhysicsMaterials;
 
diff --git a/Engine/Source/Runtime/Engine/Private/PhysicsEngine/Experimental/PhysScene_Chaos.cpp b/Engine/Source/Runtime/Engine/Private/PhysicsEngine/Experimental/PhysScene_Chaos.cpp
--- a/Engine/Source/Runtime/Engine/Private/PhysicsEngine/Experimental/PhysScene_Chaos.cpp
+++ b/Engine/Source/Runtime/Engine/Private/PhysicsEngine/Experimental/PhysScene_Chaos.cpp
@@ -72,7 +72,7 @@ CSV_DEFINE_CATEGORY(ChaosPhysics, true);
 DEFINE_LOG_CATEGORY_STATIC(LogFPhysScene_ChaosSolver, Log, All);
 
 #if WITH_CHAOS
-Chaos::TCollisionModifierCallback<float, 3> FPhysScene_ChaosInterface::CollisionModifierCallback;
+Chaos::FCollisionModifierCallback FPhysScene_ChaosInterface::CollisionModifierCallback;
 #endif // WITH_CHAOS
 
 void DumpHierarchyStats(const TArray<FString>& Args)
diff --git a/Engine/Source/Runtime/Engine/Private/PhysicsEngine/ImmediatePhysics/ImmediatePhysicsChaos/ImmediatePhysicsSimulation_Chaos.cpp b/Engine/Source/Runtime/Engine/Private/PhysicsEngine/ImmediatePhysics/ImmediatePhysicsChaos/ImmediatePhysicsSimulation_Chaos.cpp
--- a/Engine/Source/Runtime/Engine/Private/PhysicsEngine/ImmediatePhysics/ImmediatePhysicsChaos/ImmediatePhysicsSimulation_Chaos.cpp
+++ b/Engine/Source/Runtime/Engine/Private/PhysicsEngine/ImmediatePhysics/ImmediatePhysicsChaos/ImmediatePhysicsSimulation_Chaos.cpp
@@ -200,7 +200,7 @@ namespace ImmediatePhysics_Chaos
 				DebugDrawConstraints(3, 3, 1.0f);
 			});
 		Collisions.SetPostApplyCallback(
-			[this](const float Dt, const TArray<TPBDCollisionConstraintHandle<float, 3>*>& InConstraintHandles)
+			[this](const float Dt, const TArray<FPBDCollisionConstraintHandle*>& InConstraintHandles)
 			{
 				if (ChaosImmediate_DebugDrawCollisions == 4)
 				{
@@ -209,7 +209,7 @@ namespace ImmediatePhysics_Chaos
 				DebugDrawDynamicParticles(4, 4, FColor(128, 128, 0));
 			});
 		Collisions.SetPostApplyPushOutCallback(
-			[this](const float Dt, const TArray<TPBDCollisionConstraintHandle<float, 3>*>& InConstraintHandles, bool bRequiresAnotherIteration)
+			[this](const float Dt, const TArray<FPBDCollisionConstraintHandle*>& InConstraintHandles, bool bRequiresAnotherIteration)
 			{
 				if (ChaosImmediate_DebugDrawCollisions == 4)
 				{
diff --git a/Engine/Source/Runtime/Engine/Public/Physics/Experimental/PhysScene_Chaos.h b/Engine/Source/Runtime/Engine/Public/Physics/Experimental/PhysScene_Chaos.h
--- a/Engine/Source/Runtime/Engine/Public/Physics/Experimental/PhysScene_Chaos.h
+++ b/Engine/Source/Runtime/Engine/Public/Physics/Experimental/PhysScene_Chaos.h
@@ -417,7 +417,7 @@ public:
 	ENGINE_API bool ExecPxVis(uint32 SceneType, const TCHAR* Cmd, FOutputDevice* Ar);
 	ENGINE_API bool ExecApexVis(uint32 SceneType, const TCHAR* Cmd, FOutputDevice* Ar);
 
-	ENGINE_API static Chaos::TCollisionModifierCallback<float, 3> CollisionModifierCallback;
+	ENGINE_API static Chaos::FCollisionModifierCallback CollisionModifierCallback;
 
 #if XGE_FIXED
 	template<typename PayloadType>
diff --git a/Engine/Source/Runtime/Engine/Public/Physics/ImmediatePhysics/ImmediatePhysicsChaos/ImmediatePhysicsSimulation_Chaos.h b/Engine/Source/Runtime/Engine/Public/Physics/ImmediatePhysics/ImmediatePhysicsChaos/ImmediatePhysicsSimulation_Chaos.h
--- a/Engine/Source/Runtime/Engine/Public/Physics/ImmediatePhysics/ImmediatePhysicsChaos/ImmediatePhysicsSimulation_Chaos.h
+++ b/Engine/Source/Runtime/Engine/Public/Physics/ImmediatePhysics/ImmediatePhysicsChaos/ImmediatePhysicsSimulation_Chaos.h
@@ -80,7 +80,7 @@ namespace ImmediatePhysics_Chaos
 		void DebugDrawDynamicParticles(const int32 MinDebugLevel, const int32 MaxDebugLevel, const FColor& Color);
 		void DebugDrawConstraints(const int32 MinDebugLevel, const int32 MaxDebugLevel, const float ColorScale);
 
-		using FCollisionConstraints = Chaos::TPBDCollisionConstraints<FReal, 3>;
+		using FCollisionConstraints = Chaos::FPBDCollisionConstraints;
 		using FCollisionDetector = Chaos::TCollisionDetector<Chaos::FParticlePairBroadPhase, Chaos::FNarrowPhase, Chaos::FSyncCollisionReceiver, FCollisionConstraints>;
 		using FRigidParticleSOAs = Chaos::TPBDRigidsSOAs<FReal, 3>;
 		using FParticleHandle = Chaos::TGeometryParticleHandle<FReal, Dimensions>;
diff --git a/Engine/Source/Runtime/Experimental/Chaos/Private/Chaos/ChaosDebugDraw.cpp b/Engine/Source/Runtime/Experimental/Chaos/Private/Chaos/ChaosDebugDraw.cpp
--- a/Engine/Source/Runtime/Experimental/Chaos/Private/Chaos/ChaosDebugDraw.cpp
+++ b/Engine/Source/Runtime/Experimental/Chaos/Private/Chaos/ChaosDebugDraw.cpp
@@ -225,7 +225,7 @@ namespace Chaos
 			}
 		}
 		
-		void DrawCollisionImpl(const FRigidTransform3& SpaceTransform, const TPBDCollisionConstraintHandle<float, 3>* ConstraintHandle, float ColorScale)
+		void DrawCollisionImpl(const FRigidTransform3& SpaceTransform, const FPBDCollisionConstraintHandle* ConstraintHandle, float ColorScale)
 		{
 			DrawCollisionImpl(SpaceTransform, ConstraintHandle->GetContact(), ColorScale);
 		}
@@ -492,12 +492,12 @@ namespace Chaos
 #endif
 		}
 
-		void DrawParticleCollisions(const FRigidTransform3& SpaceTransform, const TGeometryParticleHandle<float, 3>* Particle, const TPBDCollisionConstraints<float, 3>& Collisions)
+		void DrawParticleCollisions(const FRigidTransform3& SpaceTransform, const TGeometryParticleHandle<float, 3>* Particle, const FPBDCollisionConstraints& Collisions)
 		{
 #if CHAOS_DEBUG_DRAW
 			if (FDebugDrawQueue::IsDebugDrawingEnabled())
 			{
-				for (const Chaos::TPBDCollisionConstraintHandle<float, 3> * ConstraintHandle : Collisions.GetConstConstraintHandles())
+				for (const Chaos::FPBDCollisionConstraintHandle * ConstraintHandle : Collisions.GetConstConstraintHandles())
 				{
 					TVector<const TGeometryParticleHandle<float, 3>*, 2> ConstrainedParticles = ConstraintHandle->GetConstrainedParticles();
 					if ((ConstrainedParticles[0] == Particle) || (ConstrainedParticles[1] == Particle))
@@ -509,7 +509,7 @@ namespace Chaos
 #endif
 		}
 
-		void DrawCollisions(const FRigidTransform3& SpaceTransform, const TPBDCollisionConstraints<float, 3>& Collisions, float ColorScale)
+		void DrawCollisions(const FRigidTransform3& SpaceTransform, const FPBDCollisionConstraints& Collisions, float ColorScale)
 		{
 #if CHAOS_DEBUG_DRAW
 			if (FDebugDrawQueue::IsDebugDrawingEnabled())
@@ -522,12 +522,12 @@ namespace Chaos
 #endif
 		}
 
-		void DrawCollisions(const FRigidTransform3& SpaceTransform, const TArray<TPBDCollisionConstraintHandle<float, 3>*>& ConstraintHandles, float ColorScale)
+		void DrawCollisions(const FRigidTransform3& SpaceTransform, const TArray<FPBDCollisionConstraintHandle*>& ConstraintHandles, float ColorScale)
 		{
 #if CHAOS_DEBUG_DRAW
 			if (FDebugDrawQueue::IsDebugDrawingEnabled())
 			{
-				for (const TPBDCollisionConstraintHandle<float, 3>* ConstraintHandle : ConstraintHandles)
+				for (const FPBDCollisionConstraintHandle* ConstraintHandle : ConstraintHandles)
 				{
 					DrawCollisionImpl(SpaceTransform, ConstraintHandle, ColorScale);
 				}
diff --git a/Engine/Source/Runtime/Experimental/Chaos/Private/Chaos/Collision/CollisionDetector.cpp b/Engine/Source/Runtime/Experimental/Chaos/Private/Chaos/Collision/CollisionDetector.cpp
--- a/Engine/Source/Runtime/Experimental/Chaos/Private/Chaos/Collision/CollisionDetector.cpp
+++ b/Engine/Source/Runtime/Experimental/Chaos/Private/Chaos/Collision/CollisionDetector.cpp
@@ -12,6 +12,6 @@ namespace Chaos
 	DEFINE_STAT(STAT_Collisions_NarrowPhase);
 #endif
 
-	template class TCollisionDetector<FParticlePairBroadPhase, FNarrowPhase, FSyncCollisionReceiver, TPBDCollisionConstraints<FReal, 3>>;
-	template class TCollisionDetector<FSpatialAccelerationBroadPhase, FNarrowPhase, FAsyncCollisionReceiver, TPBDCollisionConstraints<FReal, 3>>;
+	template class TCollisionDetector<FParticlePairBroadPhase, FNarrowPhase, FSyncCollisionReceiver, FPBDCollisionConstraints>;
+	template class TCollisionDetector<FSpatialAccelerationBroadPhase, FNarrowPhase, FAsyncCollisionReceiver, FPBDCollisionConstraints>;
 }
\ No newline at end of file
diff --git a/Engine/Source/Runtime/Experimental/Chaos/Private/Chaos/Collision/PBDCollisionConstraintHandle.cpp b/Engine/Source/Runtime/Experimental/Chaos/Private/Chaos/Collision/PBDCollisionConstraintHandle.cpp
new file mode 100644
--- /dev/null
+++ b/Engine/Source/Runtime/Experimental/Chaos/Private/Chaos/Collision/PBDCollisionConstraintHandle.cpp
@@ -0,0 +1,81 @@
+// Copyright Epic Games, Inc. All Rights Reserved.
+#include "Chaos/Collision/PBDCollisionConstraintHandle.h"
+#include "Chaos/PBDCollisionConstraints.h"
+
+namespace Chaos
+{
+	typename FPBDCollisionConstraintHandle::FHandleKey FPBDCollisionConstraintHandle::GetKey()
+	{
+		const FCollisionConstraintBase& Contact = GetContact();
+		return FHandleKey(
+			FImplicitPair(Contact.Manifold.Implicit[0], Contact.Manifold.Implicit[1]),
+			FGeometryPair(Contact.Particle[0], Contact.Particle[1]));
+	}
+
+	const FCollisionConstraintBase& FPBDCollisionConstraintHandle::GetContact() const
+	{
+		if (GetType() == FCollisionConstraintBase::FType::SinglePoint)
+		{
+			return ConstraintContainer->Constraints.SinglePointConstraints[ConstraintIndex];
+		}
+		else if (GetType() == FCollisionConstraintBase::FType::SinglePointSwept)
+		{
+			return ConstraintContainer->Constraints.SinglePointSweptConstraints[ConstraintIndex];
+		}
+
+		check(GetType() == FCollisionConstraintBase::FType::MultiPoint);
+		return ConstraintContainer->Constraints.MultiPointConstraints[ConstraintIndex];
+	}
+
+	FCollisionConstraintBase& FPBDCollisionConstraintHandle::GetContact()
+	{
+		if (GetType() == FCollisionConstraintBase::FType::SinglePoint)
+		{
+			return ConstraintContainer->Constraints.SinglePointConstraints[ConstraintIndex];
+		}
+		else if (GetType() == FCollisionConstraintBase::FType::SinglePointSwept)
+		{
+			return ConstraintContainer->Constraints.SinglePointSweptConstraints[ConstraintIndex];
+		}
+
+		check(GetType() == FCollisionConstraintBase::FType::MultiPoint);
+		return ConstraintContainer->Constraints.MultiPointConstraints[ConstraintIndex];
+	}
+
+	const FRigidBodyPointContactConstraint& FPBDCollisionConstraintHandle::GetPointContact() const
+	{ 
+		check(GetType() == FCollisionConstraintBase::FType::SinglePoint);
+		return ConstraintContainer->Constraints.SinglePointConstraints[ConstraintIndex]; 
+	}
+	
+	FRigidBodyPointContactConstraint& FPBDCollisionConstraintHandle::GetPointContact() 
+	{ 
+		check(GetType() == FCollisionConstraintBase::FType::SinglePoint);
+		return ConstraintContainer->Constraints.SinglePointConstraints[ConstraintIndex];
+	}
+
+	const FRigidBodySweptPointContactConstraint& FPBDCollisionConstraintHandle::GetSweptPointContact() const
+	{
+		check(GetType() == FCollisionConstraintBase::FType::SinglePointSwept);
+		return ConstraintContainer->Constraints.SinglePointSweptConstraints[ConstraintIndex];
+	}
+	
+	FRigidBodySweptPointContactConstraint& FPBDCollisionConstraintHandle::GetSweptPointContact()
+	{
+		check(GetType() == FCollisionConstraintBase::FType::SinglePointSwept);
+		return ConstraintContainer->Constraints.SinglePointSweptConstraints[ConstraintIndex];
+	}
+
+	const FRigidBodyMultiPointContactConstraint& FPBDCollisionConstraintHandle::GetMultiPointContact() const
+	{
+		check(GetType() == FCollisionConstraintBase::FType::MultiPoint);
+		return ConstraintContainer->Constraints.MultiPointConstraints[ConstraintIndex];
+	}
+	
+	FRigidBodyMultiPointContactConstraint& FPBDCollisionConstraintHandle::GetMultiPointContact()
+	{
+		check(GetType() == FCollisionConstraintBase::FType::MultiPoint);
+		return ConstraintContainer->Constraints.MultiPointConstraints[ConstraintIndex];
+	}
+
+}
\ No newline at end of file
