Chaos - fix HeadlessChaos compile

diff --git a/Engine/Source/Programs/HeadlessChaos/Private/HeadlessChaosTestMinEvolution.cpp b/Engine/Source/Programs/HeadlessChaos/Private/HeadlessChaosTestMinEvolution.cpp
--- a/Engine/Source/Programs/HeadlessChaos/Private/HeadlessChaosTestMinEvolution.cpp
+++ b/Engine/Source/Programs/HeadlessChaos/Private/HeadlessChaosTestMinEvolution.cpp
@@ -2,10 +2,10 @@
 
 #include "HeadlessChaos.h"
 
-#include "Chaos/Collision/CollisionDetector.h"
 #include "Chaos/Collision/CollisionReceiver.h"
 #include "Chaos/Collision/NarrowPhase.h"
 #include "Chaos/Collision/ParticlePairBroadPhase.h"
+#include "Chaos/Collision/ParticlePairCollisionDetector.h"
 #include "Chaos/Evolution/PBDMinEvolution.h"
 #include "Chaos/ParticleHandle.h"
 #include "Chaos/PBDCollisionConstraints.h"
@@ -24,7 +24,7 @@ namespace ChaosTest
 	{
 		// @todo(ccaulfield): remove template parameters on collisions and other constraints
 		using FCollisionConstraints = FPBDCollisionConstraints;
-		using FCollisionDetector = TCollisionDetector<FParticlePairBroadPhase, FNarrowPhase, FSyncCollisionReceiver, FCollisionConstraints>;
+		using FCollisionDetector = FParticlePairCollisionDetector;
 		using FRigidParticleSOAs = TPBDRigidsSOAs<FReal, 3>;
 		using FParticleHandle = TPBDRigidParticleHandle<FReal, 3>;
 		using FParticlePair = TVector<TGeometryParticleHandle<FReal, 3>*, 2>;
@@ -40,7 +40,8 @@ namespace ChaosTest
 		TArrayCollectionArray<TUniquePtr<Chaos::FChaosPhysicsMaterial>> PerParticleMaterials;
 		FCollisionConstraints Collisions(ParticlesContainer, CollidedParticles, ParticleMaterials);
 		FParticlePairBroadPhase BroadPhase(ActivePotentiallyCollidingPairs, 0);
-		FCollisionDetector CollisionDetector(BroadPhase, Collisions);
+		FNarrowPhase NarrowPhase;
+		FCollisionDetector CollisionDetector(BroadPhase, NarrowPhase, Collisions);
 		TSimpleConstraintRule<FCollisionConstraints> CollisionsRule(1, Collisions);
 		// End collisions stuff
 
diff --git a/Engine/Source/Programs/HeadlessChaos/Public/HeadlessChaosCollisionConstraints.h b/Engine/Source/Programs/HeadlessChaos/Public/HeadlessChaosCollisionConstraints.h
--- a/Engine/Source/Programs/HeadlessChaos/Public/HeadlessChaosCollisionConstraints.h
+++ b/Engine/Source/Programs/HeadlessChaos/Public/HeadlessChaosCollisionConstraints.h
@@ -8,8 +8,7 @@
 #include "Chaos/PBDCollisionConstraints.h"
 #include "Chaos/PBDCollisionConstraintsContact.h"
 #include "Chaos/CollisionResolution.h"
-#include "Chaos/Collision/CollisionContext.h"
-#include "Chaos/Collision/CollisionDetector.h"
+#include "Chaos/Collision/SpatialAccelerationCollisionDetector.h"
 #include "Chaos/CollisionResolutionUtil.h"
 #include "Chaos/Plane.h"
 #include "Chaos/Sphere.h"
@@ -32,14 +31,14 @@ public:
 	using FPointContactConstraint = FRigidBodyPointContactConstraint;
 	using FConstraintHandleAllocator = TConstraintHandleAllocator<FPBDCollisionConstraints>;
 	using FConstraintHandleID = TPair<const TGeometryParticleHandle<FReal, 3>*, const TGeometryParticleHandle<FReal, 3>*>;
-	using FCollisionDetector = TCollisionDetector<FSpatialAccelerationBroadPhase, FNarrowPhase, FAsyncCollisionReceiver, FCollisionConstraints>;
+	using FCollisionDetector = FSpatialAccelerationCollisionDetector;
 	using FAccelerationStructure = TBoundingVolume<TAccelerationStructureHandle<FReal, 3>, FReal, 3>;
 
 	FPBDCollisionConstraintAccessor()
 		: SpatialAcceleration(EmptyParticles.GetNonDisabledView())
 		, BroadPhase(EmptyParticles, (FReal)1, (FReal)0)
 		, CollisionConstraints(EmptyParticles, EmptyCollided, EmptyPhysicsMaterials, 1, 1)
-		, CollisionDetector(BroadPhase, CollisionConstraints)
+		, CollisionDetector(BroadPhase, NarrowPhase, CollisionConstraints)
 	{}
 
 	FPBDCollisionConstraintAccessor(const TPBDRigidsSOAs<FReal, 3>& InParticles, TArrayCollectionArray<bool>& Collided, const TArrayCollectionArray<TSerializablePtr<FChaosPhysicsMaterial>>& PerParticleMaterials,
@@ -47,16 +46,15 @@ public:
 		: SpatialAcceleration(InParticles.GetNonDisabledView())
 		, BroadPhase(InParticles, (FReal)1, (FReal)0)
 		, CollisionConstraints(InParticles, Collided, PerParticleMaterials, 1, 1, Thickness)
-		, CollisionDetector(BroadPhase, CollisionConstraints)
+		, CollisionDetector(BroadPhase, NarrowPhase, CollisionConstraints)
 	{}
 
 	virtual ~FPBDCollisionConstraintAccessor() {}
 	
 	void ComputeConstraints(FReal Dt)
 	{
-		CollisionStats::FStatData StatData(false);
 		CollisionDetector.GetBroadPhase().SetSpatialAcceleration(&SpatialAcceleration);
-		CollisionDetector.DetectCollisions(Dt, StatData);
+		CollisionDetector.DetectCollisions(Dt);
 	}
 
 	void Update(FCollisionConstraintBase& Constraint, FReal CullDistance = FReal(0) )
@@ -127,6 +125,7 @@ public:
 
 	FAccelerationStructure SpatialAcceleration;
 	FSpatialAccelerationBroadPhase BroadPhase;
+	FNarrowPhase NarrowPhase;
 	FCollisionConstraints CollisionConstraints;
 	FCollisionDetector CollisionDetector;
 };
