Chaos: Reenable 2 failing tests

diff --git a/Engine/Source/Programs/HeadlessChaos/Private/HeadlessChaosTestSim.cpp b/Engine/Source/Programs/HeadlessChaos/Private/HeadlessChaosTestSim.cpp
--- a/Engine/Source/Programs/HeadlessChaos/Private/HeadlessChaosTestSim.cpp
+++ b/Engine/Source/Programs/HeadlessChaos/Private/HeadlessChaosTestSim.cpp
@@ -13,8 +13,7 @@ namespace ChaosTest {
 
 	using namespace Chaos;
 
-	// Todo: Fix this
-	GTEST_TEST(SimTests, DISABLED_SphereSphereSimTest)
+	GTEST_TEST(SimTests, SphereSphereSimTest)
 	{
 		TPBDRigidsSOAs<FReal, 3> Particles;
 		THandleArray<FChaosPhysicsMaterial> PhysicalMaterials;
diff --git a/Engine/Source/Programs/HeadlessChaos/Private/HeadlessChaosTestSolverProxies.cpp b/Engine/Source/Programs/HeadlessChaos/Private/HeadlessChaosTestSolverProxies.cpp
--- a/Engine/Source/Programs/HeadlessChaos/Private/HeadlessChaosTestSolverProxies.cpp
+++ b/Engine/Source/Programs/HeadlessChaos/Private/HeadlessChaosTestSolverProxies.cpp
@@ -136,8 +136,6 @@ namespace ChaosTest {
 	template void SingleParticleProxyTaskGraphTest<float>();
 
 
-
-	// TODO: Fix this test (2 lines are commented out)
 	template<class T>
 	void SingleParticleProxyWakeEventPropergationTest()
 	{
@@ -183,9 +181,8 @@ namespace ChaosTest {
 		EXPECT_EQ(Particle->HasAwakeEvent(), false);
 		EXPECT_EQ(Particle->ObjectState(), Chaos::EObjectStateType::Dynamic);
 
-		// TODO: Fix this
-		//EXPECT_EQ(Particle2->HasAwakeEvent(), true);
-		//EXPECT_EQ(Particle2->ObjectState(), Chaos::EObjectStateType::Dynamic);
+		EXPECT_EQ(Particle2->HasAwakeEvent(), true);
+		EXPECT_EQ(Particle2->ObjectState(), Chaos::EObjectStateType::Dynamic);
 
 		Particle2->ClearEvents();
 		EXPECT_EQ(Particle2->HasAwakeEvent(), false);
