Chaos - fix unit tests and add Prepare/Unprepare calls to EvolutionGBF

diff --git a/Engine/Source/Programs/HeadlessChaos/Private/HeadlessChaosTestMinEvolution.cpp b/Engine/Source/Programs/HeadlessChaos/Private/HeadlessChaosTestMinEvolution.cpp
--- a/Engine/Source/Programs/HeadlessChaos/Private/HeadlessChaosTestMinEvolution.cpp
+++ b/Engine/Source/Programs/HeadlessChaos/Private/HeadlessChaosTestMinEvolution.cpp
@@ -38,6 +38,8 @@ namespace ChaosTest
 		TArrayCollectionArray<bool> CollidedParticles;
 		TArrayCollectionArray<Chaos::TSerializablePtr<Chaos::FChaosPhysicsMaterial>> ParticleMaterials;
 		TArrayCollectionArray<TUniquePtr<Chaos::FChaosPhysicsMaterial>> PerParticleMaterials;
+		TArrayCollectionArray<FVec3> ParticlePrevXs;
+		TArrayCollectionArray<FRotation3> ParticlePrevRs;
 		FCollisionConstraints Collisions(ParticlesContainer, CollidedParticles, ParticleMaterials);
 		FParticlePairBroadPhase BroadPhase(ActivePotentiallyCollidingPairs, 0);
 		FNarrowPhase NarrowPhase;
@@ -51,7 +53,7 @@ namespace ChaosTest
 
 		// Evolution
 		// @todo(ccaulfield): this should start with some reasonable default iterations
-		FPBDMinEvolution Evolution(ParticlesContainer, CollisionDetector, 0);
+		FPBDMinEvolution Evolution(ParticlesContainer, ParticlePrevXs, ParticlePrevRs, CollisionDetector, 0);
 		Evolution.SetNumIterations(1);
 		Evolution.SetNumPushOutIterations(0);
 
@@ -61,6 +63,8 @@ namespace ChaosTest
 		FReal Dt = 1.0f / 30.0f;
 
 		// Add a couple dynamic particles connected by a spring
+		ParticlesContainer.GetParticleHandles().AddArray(&ParticlePrevXs);
+		ParticlesContainer.GetParticleHandles().AddArray(&ParticlePrevRs);
 		TArray<FParticleHandle*> Particles = ParticlesContainer.CreateDynamicParticles(2);
 
 		// Set up Particles
@@ -70,12 +74,16 @@ namespace ChaosTest
 		Particles[0]->I() = FMatrix33(100.0f, 100.0f, 100.0f);
 		Particles[0]->InvM() = 1.0f;
 		Particles[0]->InvI() = FMatrix33(1.0f / 100.0f, 1.0f / 100.0f, 1.0f / 100.0f);
+		Particles[0]->AuxilaryValue(ParticlePrevXs) = Particles[0]->X();
+		Particles[0]->AuxilaryValue(ParticlePrevRs) = Particles[0]->R();
 
 		Particles[1]->X() = FVec3(50, 0, 0);
 		Particles[1]->M() = 1.0f;
 		Particles[1]->I() = FMatrix33(100.0f, 100.0f, 100.0f);
 		Particles[1]->InvM() = 1.0f;
 		Particles[1]->InvI() = FMatrix33(1.0f / 100.0f, 1.0f / 100.0f, 1.0f / 100.0f);
+		Particles[1]->AuxilaryValue(ParticlePrevXs) = Particles[1]->X();
+		Particles[1]->AuxilaryValue(ParticlePrevRs) = Particles[1]->R();
 
 		// Spring connectors at particle centres
 		TArray<FVec3> Locations =
@@ -90,7 +98,7 @@ namespace ChaosTest
 		for (int32 TimeIndex = 0; TimeIndex < 1000; ++TimeIndex)
 		{
 			//UE_LOG(LogChaos, Warning, TEXT("%d: %f %f %f - %f %f %f"), TimeIndex, Particles[0]->X().X, Particles[0]->X().Y, Particles[0]->X().Z, Particles[1]->X().X, Particles[1]->X().Y, Particles[1]->X().Z);
-			Evolution.Advance(Dt, 1);
+			Evolution.Advance(Dt, 1, 0.0f);
 		}
 
 		// Particles should be separated by the spring's rest length
diff --git a/Engine/Source/Runtime/Experimental/Chaos/Private/Chaos/PBDRigidsEvolutionGBF.cpp b/Engine/Source/Runtime/Experimental/Chaos/Private/Chaos/PBDRigidsEvolutionGBF.cpp
--- a/Engine/Source/Runtime/Experimental/Chaos/Private/Chaos/PBDRigidsEvolutionGBF.cpp
+++ b/Engine/Source/Runtime/Experimental/Chaos/Private/Chaos/PBDRigidsEvolutionGBF.cpp
@@ -133,7 +133,7 @@ void FPBDRigidsEvolutionGBF::Advance(const FReal Dt, const FReal MaxStepDt, cons
 		
 			UE_LOG(LogChaos, Verbose, TEXT("Advance dt = %f [%d/%d]"), StepDt, Step + 1, NumSteps);
 
-			AdvanceOneTimeStep(StepDt, StepFraction);
+			AdvanceOneTimeStepImpl(StepDt, StepFraction);
 		}
 
 		UnprepareTick();
@@ -249,8 +249,16 @@ void CCDHack(const FReal Dt, TParticleView<TPBDRigidParticles<FReal, 3>>& Partic
 //
 //
 
-
 void FPBDRigidsEvolutionGBF::AdvanceOneTimeStep(const FReal Dt, const FReal StepFraction)
+{
+	PrepareTick();
+
+	AdvanceOneTimeStepImpl(Dt, StepFraction);
+
+	UnprepareTick();
+}
+
+void FPBDRigidsEvolutionGBF::AdvanceOneTimeStepImpl(const FReal Dt, const FReal StepFraction)
 {
 	SCOPE_CYCLE_COUNTER(STAT_Evolution_AdvanceOneTimeStep);
 
diff --git a/Engine/Source/Runtime/Experimental/Chaos/Public/Chaos/PBDRigidsEvolutionGBF.h b/Engine/Source/Runtime/Experimental/Chaos/Public/Chaos/PBDRigidsEvolutionGBF.h
--- a/Engine/Source/Runtime/Experimental/Chaos/Public/Chaos/PBDRigidsEvolutionGBF.h
+++ b/Engine/Source/Runtime/Experimental/Chaos/Public/Chaos/PBDRigidsEvolutionGBF.h
@@ -194,6 +194,9 @@ namespace Chaos
 		CHAOS_API void Serialize(FChaosArchive& Ar);
 
 	protected:
+
+		CHAOS_API void AdvanceOneTimeStepImpl(const FReal dt, const FReal StepFraction);
+
 		TPBDRigidClustering<FPBDRigidsEvolutionGBF, FPBDCollisionConstraints, FReal, 3> Clustering;
 
 		FGravityForces GravityForces;
