Record delta time passed into the physics solver for rewind

diff --git a/Engine/Source/Programs/HeadlessChaos/Private/HeadlessChaosTestRewind.cpp b/Engine/Source/Programs/HeadlessChaos/Private/HeadlessChaosTestRewind.cpp
--- a/Engine/Source/Programs/HeadlessChaos/Private/HeadlessChaosTestRewind.cpp
+++ b/Engine/Source/Programs/HeadlessChaos/Private/HeadlessChaosTestRewind.cpp
@@ -18,10 +18,10 @@ namespace ChaosTest {
 
     using namespace Chaos;
 
-	void TickSolverHelper(FChaosSolversModule* Module, FPhysicsSolver* Solver)
+	void TickSolverHelper(FChaosSolversModule* Module, FPhysicsSolver* Solver, FReal Dt = 1.0)
 	{
 		Solver->PushPhysicsState(Module->GetDispatcher());
-		FPhysicsSolverAdvanceTask AdvanceTask(Solver,1.0f);
+		FPhysicsSolverAdvanceTask AdvanceTask(Solver,Dt);
 		AdvanceTask.DoTask(ENamedThreads::GameThread,FGraphEventRef());
 		Solver->BufferPhysicsResults();
 		Solver->FlipBuffers();
@@ -778,4 +778,50 @@ namespace ChaosTest {
 
 		Module->DestroySolver(Solver);
 	}
+
+	GTEST_TEST(RewindTest,DeltaTimeRecord)
+	{
+		auto Sphere = TSharedPtr<FImplicitObject,ESPMode::ThreadSafe>(new TSphere<float,3>(TVector<float,3>(0),10));
+
+		FChaosSolversModule* Module = FChaosSolversModule::GetModule();
+		Module->ChangeThreadingMode(EChaosThreadingMode::SingleThread);
+
+		// Make a solver
+		FPhysicsSolver* Solver = Module->CreateSolver(nullptr,ESolverFlags::Standalone);
+		Solver->SetEnabled(true);
+
+		Solver->EnableRewindCapture(7);
+
+		// Make particles
+		auto Particle = TPBDRigidParticle<float,3>::CreateParticle();
+
+		Particle->SetGeometry(Sphere);
+		Solver->RegisterObject(Particle.Get());
+		Particle->SetGravityEnabled(true);
+
+		const int LastStep = 11;
+		TArray<FReal> DTs;
+		FReal Dt = 1;
+		for(int Step = 0; Step <= LastStep; ++Step)
+		{
+			DTs.Add(Dt);
+			TickSolverHelper(Module,Solver, Dt);
+			Dt += 0.1;
+		}
+		
+		const int RewindStep = 5;
+
+		FRewindData* RewindData = Solver->GetRewindData();
+		EXPECT_TRUE(RewindData->RewindToFrame(RewindStep));
+
+		for(int Step = RewindStep; Step <= LastStep; ++Step)
+		{
+			EXPECT_EQ(DTs[Step],RewindData->GetDeltaTimeForFrame(Step));
+		}
+
+		// Throw out the proxy
+		Solver->UnregisterObject(Particle.Get());
+
+		Module->DestroySolver(Solver);
+	}
 }
diff --git a/Engine/Source/Runtime/Experimental/ChaosSolvers/Private/PBDRigidsSolver.cpp b/Engine/Source/Runtime/Experimental/ChaosSolvers/Private/PBDRigidsSolver.cpp
--- a/Engine/Source/Runtime/Experimental/ChaosSolvers/Private/PBDRigidsSolver.cpp
+++ b/Engine/Source/Runtime/Experimental/ChaosSolvers/Private/PBDRigidsSolver.cpp
@@ -108,7 +108,7 @@ namespace Chaos
 
 			if(FRewindData* RewindData = MSolver->GetRewindData())
 			{
-				RewindData->AdvanceFrame();
+				RewindData->AdvanceFrame(MDeltaTime);
 			}
 
 			{
diff --git a/Engine/Source/Runtime/Experimental/ChaosSolvers/Public/RewindData.h b/Engine/Source/Runtime/Experimental/ChaosSolvers/Public/RewindData.h
--- a/Engine/Source/Runtime/Experimental/ChaosSolvers/Public/RewindData.h
+++ b/Engine/Source/Runtime/Experimental/ChaosSolvers/Public/RewindData.h
@@ -677,6 +677,12 @@ public:
 
 	int32 Capacity() const { return Managers.Capacity(); }
 
+	FReal GetDeltaTimeForFrame(int32 Frame) const
+	{
+		ensure(Managers[Frame].FrameCreatedFor == Frame);
+		return Managers[Frame].DeltaTime;
+	}
+
 	bool RewindToFrame(int32 Frame)
 	{
 		ensure(IsInGameThread());
@@ -799,9 +805,11 @@ public:
 		return EFutureQueryResult::Untracked;
 	}
 
-	void AdvanceFrame()
+	void AdvanceFrame(FReal DeltaTime)
 	{
 		QUICK_SCOPE_CYCLE_COUNTER(RewindDataAdvance);
+		Managers[CurFrame].DeltaTime = DeltaTime;
+
 		++CurFrame;
 		LatestFrame = FMath::Max(LatestFrame,CurFrame);
 		FramesSaved = FMath::Min(FramesSaved+1,static_cast<int32>(Managers.Capacity()));
@@ -1063,6 +1071,7 @@ private:
 		//A manager can have data for two frames at once, the important part is just knowing which frame it was created on so we know whether the physics data can rely on it
 		//Consider the case where nothing is dirty from GT and then an object moves from the simulation, in that case it needs a manager to record the data into
 		int32 FrameCreatedFor;
+		FReal DeltaTime;
 	};
 
 	struct FDirtyParticleInfo
