Fixed rewind data not properly recording mass props and dynamic misc. Added unit test

diff --git a/Engine/Source/Programs/HeadlessChaos/Private/HeadlessChaosTestRewind.cpp b/Engine/Source/Programs/HeadlessChaos/Private/HeadlessChaosTestRewind.cpp
--- a/Engine/Source/Programs/HeadlessChaos/Private/HeadlessChaosTestRewind.cpp
+++ b/Engine/Source/Programs/HeadlessChaos/Private/HeadlessChaosTestRewind.cpp
@@ -779,6 +779,74 @@ namespace ChaosTest {
 		Module->DestroySolver(Solver);
 	}
 
+	GTEST_TEST(RewindTest,ResimDesyncAfterChangingMass)
+	{
+		auto Sphere = TSharedPtr<FImplicitObject,ESPMode::ThreadSafe>(new TSphere<float,3>(TVector<float,3>(0),10));
+
+		FChaosSolversModule* Module = FChaosSolversModule::GetModule();
+		Module->ChangeThreadingMode(EChaosThreadingMode::SingleThread);
+
+		// Make a solver
+		FPhysicsSolver* Solver = Module->CreateSolver(nullptr,ESolverFlags::Standalone);
+		Solver->SetEnabled(true);
+
+		Solver->EnableRewindCapture(7);
+
+		// Make particles
+		auto Particle = TPBDRigidParticle<float,3>::CreateParticle();
+
+		Particle->SetGeometry(Sphere);
+		Solver->RegisterObject(Particle.Get());
+		Particle->SetGravityEnabled(true);
+
+		FReal CurMass = 1.0;
+		Particle->SetM(CurMass);
+		int32 LastStep = 11;
+
+		for(int Step = 0; Step <= LastStep; ++Step)
+		{
+			if(Step == 7)
+			{
+				Particle->SetM(2);
+			}
+
+			if(Step == 9)
+			{
+				Particle->SetM(3);
+			}
+			TickSolverHelper(Module,Solver);
+		}
+
+		const int RewindStep = 5;
+
+		FRewindData* RewindData = Solver->GetRewindData();
+		EXPECT_TRUE(RewindData->RewindToFrame(RewindStep));
+
+		for(int Step = RewindStep; Step <= LastStep; ++Step)
+		{
+			FGeometryParticleState FutureState(*Particle);
+			EXPECT_EQ(RewindData->GetFutureStateAtFrame(FutureState,Step),Step < 10 ? EFutureQueryResult::Ok : EFutureQueryResult::Desync);
+			if(Step < 7)
+			{
+				EXPECT_EQ(1,FutureState.M());
+			}
+
+			if(Step == 7)
+			{
+				Particle->SetM(2);
+			}
+
+			//skip step 9 SetM to trigger a desync
+
+			TickSolverHelper(Module,Solver);
+		}
+
+		// Throw out the proxy
+		Solver->UnregisterObject(Particle.Get());
+
+		Module->DestroySolver(Solver);
+	}
+
 	GTEST_TEST(RewindTest,DeltaTimeRecord)
 	{
 		auto Sphere = TSharedPtr<FImplicitObject,ESPMode::ThreadSafe>(new TSphere<float,3>(TVector<float,3>(0),10));
diff --git a/Engine/Source/Runtime/Experimental/ChaosSolvers/Public/RewindData.h b/Engine/Source/Runtime/Experimental/ChaosSolvers/Public/RewindData.h
--- a/Engine/Source/Runtime/Experimental/ChaosSolvers/Public/RewindData.h
+++ b/Engine/Source/Runtime/Experimental/ChaosSolvers/Public/RewindData.h
@@ -379,10 +379,19 @@ public:
 			Data.CopyFrom(*Handle);
 		});
 
-		NonFrequentData.SyncRemoteData(Manager,Idx,Dirty.ParticleData,[Handle](FParticleNonFrequentData& Data)
+		if(auto Rigid = Handle->CastToRigidParticle())
 		{
-			Data.CopyFrom(*Handle);
-		});
+			DynamicsMisc.SyncRemoteData(Manager,Idx,Dirty.ParticleData,[Rigid](FParticleDynamicMisc& Data)
+			{
+				Data.CopyFrom(*Rigid);
+			});
+
+			MassProps.SyncRemoteData(Manager,Idx,Dirty.ParticleData,[Rigid](FParticleMassProps& Data)
+			{
+				Data.CopyFrom(*Rigid);
+			});
+		}
+		
 	}
 
 	void SyncIfDirty(FDirtyPropertiesManager& Manager,int32 Idx,const TGeometryParticle<FReal,3>& InParticle, const FGeometryParticleStateBase& RewindState)
@@ -475,15 +484,39 @@ public:
 			{
 				return true;
 			}
+
+			if(!NonFrequentData.IsInSync(SrcManager,DataIdxIn,Flags,Handle))
+			{
+				return true;
+			}
+
+			if(auto Kinematic = Handle.CastToKinematicParticle())
+			{
+				if(!Velocities.IsInSync(SrcManager,DataIdxIn,Flags,*Kinematic))
+				{
+					return true;
+				}
+			}
+
+			if(auto Rigid = Handle.CastToRigidParticle())
+			{
+				if(!Dynamics.IsInSync(SrcManager,DataIdxIn,Flags,*Rigid))
+				{
+					return true;
+				}
+
+				if(!DynamicsMisc.IsInSync(SrcManager,DataIdxIn,Flags,*Rigid))
+				{
+					return true;
+				}
+
+				if(!MassProps.IsInSync(SrcManager,DataIdxIn,Flags,*Rigid))
+				{
+					return true;
+				}
+			}
 		}
 
-		//TODO: test other properties, should probably find a better way to do this to avoid getting into the individual variables
-		
-		/*Desynced = ParticlePositionRotation.IsDesynced(SrcManager,DataIdxIn,Dirty);
-		Desynced = Desynced || NonFrequentData.IsDesynced(SrcManager,DataIdxIn,Dirty);
-		Desynced = Desynced || Velocities.IsDesynced(SrcManager,DataIdxIn,Dirty);
-		Desynced = Desynced || Dynamics.IsDesynced(SrcManager,DataIdxIn,Dirty);
-		return Desynced;*/
 		return false;
 	}
 
