Fix Headless Chaos Tests

diff --git a/Engine/Source/Programs/HeadlessChaos/Private/HeadlessChaosTestSim.cpp b/Engine/Source/Programs/HeadlessChaos/Private/HeadlessChaosTestSim.cpp
--- a/Engine/Source/Programs/HeadlessChaos/Private/HeadlessChaosTestSim.cpp
+++ b/Engine/Source/Programs/HeadlessChaos/Private/HeadlessChaosTestSim.cpp
@@ -21,12 +21,17 @@ namespace ChaosTest {
 		auto Static = Evolution.CreateStaticParticles(1)[0];
 		auto Dynamic = Evolution.CreateDynamicParticles(1)[0];
 
+		TUniquePtr<FChaosPhysicsMaterial> PhysicsMaterial = MakeUnique<FChaosPhysicsMaterial>();
+		PhysicsMaterial->SleepCounterThreshold = 2;
+
 		TUniquePtr<FImplicitObject> Box(new TSphere<FReal, 3>(FVec3(0, 0, 0), 50));
 		Static->SetGeometry(MakeSerializable(Box));
 		Dynamic->SetGeometry(MakeSerializable(Box));
 
+		Evolution.SetPhysicsMaterial(Dynamic, MakeSerializable(PhysicsMaterial));
+
 		Static->X() = FVec3(10, 10, 10);
-		Dynamic->X() = FVec3(10, 10, 300);
+		Dynamic->X() = FVec3(10, 10, 150);
 		Dynamic->I() = FMatrix33(100000.0f, 100000.0f, 100000.0f);
 		Dynamic->InvI() = FMatrix33(1.0f / 100000.0f, 1.0f / 100000.0f, 1.0f / 100000.0f);
 
diff --git a/Engine/Source/Programs/HeadlessChaos/Private/HeadlessChaosTestSolverProxies.cpp b/Engine/Source/Programs/HeadlessChaos/Private/HeadlessChaosTestSolverProxies.cpp
--- a/Engine/Source/Programs/HeadlessChaos/Private/HeadlessChaosTestSolverProxies.cpp
+++ b/Engine/Source/Programs/HeadlessChaos/Private/HeadlessChaosTestSolverProxies.cpp
@@ -152,7 +152,7 @@ namespace ChaosTest {
 
 		TUniquePtr<Chaos::TPBDRigidParticle<float, 3>> Particle = Chaos::TPBDRigidParticle<float, 3>::CreateParticle();
 		Particle->SetGeometry(Sphere);
-		Particle->SetX(TVector<float, 3>(0, 0, 1000));
+		Particle->SetX(TVector<float, 3>(0, 0, 220));
 		Particle->SetV(TVector<float, 3>(0, 0, -10));
 		Solver->RegisterObject(Particle.Get());
 		Solver->AddDirtyProxy(Particle->GetProxy());
@@ -163,18 +163,23 @@ namespace ChaosTest {
 		Particle2->SetV(TVector<float, 3>(0, 0, 0));
 		Solver->RegisterObject(Particle2.Get());
 		Solver->AddDirtyProxy(Particle2->GetProxy());
-		Particle2->SetObjectState(Chaos::EObjectStateType::Sleeping, false);
+		Particle2->SetObjectState(Chaos::EObjectStateType::Sleeping, true);
 
-		::ChaosTest::SetParticleSimDataToCollide({Particle.Get(),Particle2.Get() });
+		::ChaosTest::SetParticleSimDataToCollide({ Particle.Get(),Particle2.Get() });
 
 		Solver->PushPhysicsState(Module->GetDispatcher());
 
-		FPhysicsSolverAdvanceTask AdvanceTask(Solver, 100.0f);
-		AdvanceTask.DoTask(ENamedThreads::GameThread, FGraphEventRef());
+		// let top paticle collide and wake up second particle
+		int32 LoopCount = 0;
+		while (!Particle2->HasAwakeEvent() && LoopCount++ < 20)
+		{
+			FPhysicsSolverAdvanceTask AdvanceTask(Solver, 100.0f);
+			AdvanceTask.DoTask(ENamedThreads::GameThread, FGraphEventRef());
 
-		Solver->BufferPhysicsResults();
-		Solver->FlipBuffers();
-		Solver->UpdateGameThreadStructures();
+			Solver->BufferPhysicsResults();
+			Solver->FlipBuffers();
+			Solver->UpdateGameThreadStructures();
+		}
 
 		// Make sure game thread data has changed
 		TVector<float, 3> V = Particle->V();
