Added disabled test for sim collision desync

diff --git a/Engine/Source/Programs/HeadlessChaos/Private/HeadlessChaosTestRewind.cpp b/Engine/Source/Programs/HeadlessChaos/Private/HeadlessChaosTestRewind.cpp
--- a/Engine/Source/Programs/HeadlessChaos/Private/HeadlessChaosTestRewind.cpp
+++ b/Engine/Source/Programs/HeadlessChaos/Private/HeadlessChaosTestRewind.cpp
@@ -847,6 +847,72 @@ namespace ChaosTest {
 		Module->DestroySolver(Solver);
 	}
 
+	GTEST_TEST(RewindTest,DISABLED_DesyncFromPT)
+	{
+		//We want to detect when sim results change
+		//Detecting output of position and velocity is expensive and hard to track
+		//Instead we need to rely on fast forward mechanism, this is still in progress
+		auto Sphere = TSharedPtr<FImplicitObject,ESPMode::ThreadSafe>(new TSphere<FReal,3>(TVector<float,3>(0),10));
+		auto Box = TSharedPtr<FImplicitObject,ESPMode::ThreadSafe>(new TBox<FReal,3>(FVec3(-100,-100,-100),FVec3(100, 100, 0)));
+
+		FChaosSolversModule* Module = FChaosSolversModule::GetModule();
+		Module->ChangeThreadingMode(EChaosThreadingMode::SingleThread);
+
+		// Make a solver
+		FPhysicsSolver* Solver = Module->CreateSolver(nullptr,ESolverFlags::Standalone);
+		Solver->SetEnabled(true);
+		Solver->EnableRewindCapture(7);
+
+		// Make particles
+		auto Dynamic = TPBDRigidParticle<float,3>::CreateParticle();
+		auto Kinematic = TKinematicGeometryParticle<float,3>::CreateParticle();
+
+		Dynamic->SetGeometry(Sphere);
+		Dynamic->SetGravityEnabled(true);
+		Solver->RegisterObject(Dynamic.Get());
+
+		Kinematic->SetGeometry(Box);
+		Solver->RegisterObject(Kinematic.Get());
+
+		Dynamic->SetX(FVec3(0,0,17));
+		Dynamic->SetGravityEnabled(false);
+		Dynamic->SetV(FVec3(0,0,-1));
+		Dynamic->SetObjectState(EObjectStateType::Dynamic);
+
+		Kinematic->SetX(FVec3(0,0,0));
+
+		ChaosTest::SetParticleSimDataToCollide({Dynamic.Get(),Kinematic.Get()});
+
+		const int32 LastStep = 11;
+
+		for(int Step = 0; Step <= LastStep; ++Step)
+		{
+			TickSolverHelper(Module,Solver);
+		}
+
+		EXPECT_FLOAT_EQ(Dynamic->X()[2],10);
+		
+		const int RewindStep = 5;
+
+		FRewindData* RewindData = Solver->GetRewindData();
+		EXPECT_TRUE(RewindData->RewindToFrame(RewindStep));
+
+		Kinematic->SetX(FVec3(0,0,-1));
+
+		for(int Step = RewindStep; Step <= LastStep; ++Step)
+		{
+			//at Step 7 we're at z=10 but velocity will now be -1 instead of 0, so a desync has occured
+			FGeometryParticleState FutureState(*Dynamic);
+			EXPECT_EQ(RewindData->GetFutureStateAtFrame(FutureState,Step),Step < 7 ? EFutureQueryResult::Ok : EFutureQueryResult::Desync);
+
+			TickSolverHelper(Module,Solver);
+		}
+
+		EXPECT_FLOAT_EQ(Dynamic->X()[2],9);
+
+		Module->DestroySolver(Solver);
+	}
+
 	GTEST_TEST(RewindTest,DeltaTimeRecord)
 	{
 		auto Sphere = TSharedPtr<FImplicitObject,ESPMode::ThreadSafe>(new TSphere<float,3>(TVector<float,3>(0),10));
