Fix rewind applying dynamics instead of just using it to determine desync

diff --git a/Engine/Source/Programs/HeadlessChaos/Private/HeadlessChaosTestRewind.cpp b/Engine/Source/Programs/HeadlessChaos/Private/HeadlessChaosTestRewind.cpp
--- a/Engine/Source/Programs/HeadlessChaos/Private/HeadlessChaosTestRewind.cpp
+++ b/Engine/Source/Programs/HeadlessChaos/Private/HeadlessChaosTestRewind.cpp
@@ -958,4 +958,77 @@ namespace ChaosTest {
 
 		Module->DestroySolver(Solver);
 	}
+
+	GTEST_TEST(RewindTest,ResimDesyncFromChangeForce)
+	{
+		auto Sphere = TSharedPtr<FImplicitObject,ESPMode::ThreadSafe>(new TSphere<float,3>(TVector<float,3>(0),10));
+
+		FChaosSolversModule* Module = FChaosSolversModule::GetModule();
+		Module->ChangeThreadingMode(EChaosThreadingMode::SingleThread);
+
+		// Make a solver
+		FPhysicsSolver* Solver = Module->CreateSolver(nullptr,ESolverFlags::Standalone);
+		Solver->SetEnabled(true);
+
+		Solver->EnableRewindCapture(7);
+
+		// Make particles
+		auto Particle = TPBDRigidParticle<float,3>::CreateParticle();
+
+		Particle->SetGeometry(Sphere);
+		Solver->RegisterObject(Particle.Get());
+		Particle->SetGravityEnabled(false);
+		Particle->SetV(FVec3(0,0,10));
+
+		int32 LastStep = 11;
+
+		for(int Step = 0; Step <= LastStep; ++Step)
+		{
+			if(Step == 7)
+			{
+				Particle->SetF(FVec3(0,1,0));
+			}
+
+			if(Step == 9)
+			{
+				Particle->SetF(FVec3(100,0,0));
+			}
+			TickSolverHelper(Module,Solver);
+		}
+
+		const int RewindStep = 5;
+
+		{
+			FRewindData* RewindData = Solver->GetRewindData();
+			EXPECT_TRUE(RewindData->RewindToFrame(RewindStep));
+
+			for(int Step = RewindStep; Step <= LastStep; ++Step)
+			{
+				FGeometryParticleState FutureState(*Particle);
+				EXPECT_EQ(RewindData->GetFutureStateAtFrame(FutureState,Step),Step < 10 ? EFutureQueryResult::Ok : EFutureQueryResult::Desync);
+
+				if(Step == 7)
+				{
+					Particle->SetF(FVec3(0,1,0));
+				}
+
+				//skip step 9 SetF to trigger a desync
+
+				TickSolverHelper(Module,Solver);
+			}
+			EXPECT_EQ(Particle->V()[0],0);
+		}
+
+		//rewind to exactly step 7 to make sure force is not already applied for us
+		{
+			FRewindData* RewindData = Solver->GetRewindData();
+			EXPECT_TRUE(RewindData->RewindToFrame(7));
+			EXPECT_EQ(Particle->F()[1],0);
+		}
+
+		// Throw out the proxy
+		Solver->UnregisterObject(Particle.Get());
+
+		Module->DestroySolver(Solver);
+	}
 }
diff --git a/Engine/Source/Runtime/Experimental/ChaosSolvers/Public/RewindData.h b/Engine/Source/Runtime/Experimental/ChaosSolvers/Public/RewindData.h
--- a/Engine/Source/Runtime/Experimental/ChaosSolvers/Public/RewindData.h
+++ b/Engine/Source/Runtime/Experimental/ChaosSolvers/Public/RewindData.h
@@ -346,11 +346,6 @@ public:
 
 		if(auto Rigid = Particle.CastToRigidParticle())
 		{
-			Dynamics.SyncToParticle([Rigid](const auto& Data)
-			{
-				Rigid->SetDynamics(Data);
-			});
-
 			DynamicsMisc.SyncToParticle([Rigid](const auto& Data)
 			{
 				Rigid->SetDynamicMisc(Data);
