Starting to convert tests over to iterate over all evolutions

Chris.Caulfield

diff --git a/Engine/Source/Programs/HeadlessChaos/Private/HeadlessChaosApp.cpp b/Engine/Source/Programs/HeadlessChaos/Private/HeadlessChaosApp.cpp
--- a/Engine/Source/Programs/HeadlessChaos/Private/HeadlessChaosApp.cpp
+++ b/Engine/Source/Programs/HeadlessChaos/Private/HeadlessChaosApp.cpp
@@ -93,11 +93,6 @@ TEST(CollisionTests, PGS) {
 	SUCCEED();
 }
 
-TEST(ForceTests, Forces) {
-	ChaosTest::Gravity<Chaos::FPBDRigidsEvolution, float>();
-	SUCCEED();
-}
-
 TEST(Clustering, Clustering) {
 	ChaosTest::ImplicitCluster<float>();
 	ChaosTest::FractureCluster<float>();
diff --git a/Engine/Source/Programs/HeadlessChaos/Private/HeadlessChaosTestConstraintHandle.cpp b/Engine/Source/Programs/HeadlessChaos/Private/HeadlessChaosTestConstraintHandle.cpp
--- a/Engine/Source/Programs/HeadlessChaos/Private/HeadlessChaosTestConstraintHandle.cpp
+++ b/Engine/Source/Programs/HeadlessChaos/Private/HeadlessChaosTestConstraintHandle.cpp
@@ -188,37 +188,37 @@ namespace ChaosTest {
 		CheckConstraintHandles(Constraints, ParticleHandles, ConstraintsHandles);
 	}
 
-	TEST(ConstraintHandleTests, DISABLED_CollisionConstraintHandle)
+	TYPED_TEST(AllEvolutions, DISABLED_ConstraintHandleTests_CollisionConstraintHandle)
 	{
-		CollisionConstraintHandles<Chaos::FPBDRigidsEvolution>();
+		CollisionConstraintHandles<TypeParam>();
 
 		SUCCEED();
 	}
 
-	TEST(ConstraintHandleTests, JointConstraintHandle)
+	TYPED_TEST(AllEvolutions, ConstraintHandleTests_JointConstraintHandle)
 	{
-		JointConstraintHandles<Chaos::FPBDRigidsEvolution>();
+		JointConstraintHandles<TypeParam>();
 
 		SUCCEED();
 	}
 
-	TEST(ConstraintHandleTests, PositionConstraintHandles)
+	TYPED_TEST(AllEvolutions, ConstraintHandleTests_PositionConstraintHandles)
 	{
-		PositionConstraintHandles<Chaos::FPBDRigidsEvolution>();
+		PositionConstraintHandles<TypeParam>();
 
 		SUCCEED();
 	}
 
-	TEST(ConstraintHandleTests, RigidSpringConstraintHandles)
+	TYPED_TEST(AllEvolutions, ConstraintHandleTests_RigidSpringConstraintHandles)
 	{
-		RigidSpringConstraintHandles<Chaos::FPBDRigidsEvolution>();
+		RigidSpringConstraintHandles<TypeParam>();
 
 		SUCCEED();
 	}
 
-	TEST(ConstraintHandleTests, RigidDynamicSpringConstraintHandles)
+	TYPED_TEST(AllEvolutions, ConstraintHandleTests_RigidDynamicSpringConstraintHandles)
 	{
-		RigidDynamicSpringConstraintHandles<Chaos::FPBDRigidsEvolution>();
+		RigidDynamicSpringConstraintHandles<TypeParam>();
 
 		SUCCEED();
 	}
diff --git a/Engine/Source/Programs/HeadlessChaos/Private/HeadlessChaosTestConstraints.cpp b/Engine/Source/Programs/HeadlessChaos/Private/HeadlessChaosTestConstraints.cpp
--- a/Engine/Source/Programs/HeadlessChaos/Private/HeadlessChaosTestConstraints.cpp
+++ b/Engine/Source/Programs/HeadlessChaos/Private/HeadlessChaosTestConstraints.cpp
@@ -45,7 +45,7 @@ namespace ChaosTest {
 		{
 			TPBDRigidsSOAs<FReal, 3> Particles;
 			THandleArray<FChaosPhysicsMaterial> PhysicalMaterials;
-			FPBDRigidsEvolution Evolution(Particles, PhysicalMaterials);
+			TEvolution Evolution(Particles, PhysicalMaterials);
 			TArray<TPBDRigidParticleHandle<FReal, 3>*> Dynamics = Evolution.CreateDynamicParticles(1);
 			Dynamics[0]->SetGravityEnabled(false);
 
@@ -117,10 +117,10 @@ namespace ChaosTest {
 	}
 
 
-	TEST(ConstraintTests, Constraints)
+	TYPED_TEST(AllEvolutions, Constraints)
 	{
-		ChaosTest::Position<Chaos::FPBDRigidsEvolution>();
-		ChaosTest::PositionAndJoint<Chaos::FPBDRigidsEvolution>();
+		ChaosTest::Position<TypeParam>();
+		ChaosTest::PositionAndJoint<TypeParam>();
 
 		SUCCEED();
 	}
diff --git a/Engine/Source/Programs/HeadlessChaos/Private/HeadlessChaosTestForces.cpp b/Engine/Source/Programs/HeadlessChaos/Private/HeadlessChaosTestForces.cpp
--- a/Engine/Source/Programs/HeadlessChaos/Private/HeadlessChaosTestForces.cpp
+++ b/Engine/Source/Programs/HeadlessChaos/Private/HeadlessChaosTestForces.cpp
@@ -11,8 +11,8 @@
 #include "Chaos/Sphere.h"
 #include "Chaos/Utilities.h"
 
-namespace ChaosTest {
-
+namespace ChaosTest
+{
 	template<typename TEvolution, typename T>
 	void Gravity()
 	{
@@ -24,7 +24,10 @@ namespace ChaosTest {
 		Evolution.AdvanceOneTimeStep(0.1);
 		EXPECT_LT(Dynamics[0]->X()[2], 0);
 	}
-	template void Gravity<Chaos::FPBDRigidsEvolution, float>();
-
-
+	
+	TYPED_TEST(AllEvolutions,Forces)
+	{
+		ChaosTest::Gravity<TypeParam,float>();
+		SUCCEED();
+	}
 }
\ No newline at end of file
diff --git a/Engine/Source/Programs/HeadlessChaos/Public/HeadlessChaosTestForces.h b/Engine/Source/Programs/HeadlessChaos/Public/HeadlessChaosTestForces.h
--- a/Engine/Source/Programs/HeadlessChaos/Public/HeadlessChaosTestForces.h
+++ b/Engine/Source/Programs/HeadlessChaos/Public/HeadlessChaosTestForces.h
@@ -9,10 +9,3 @@
 #include "Chaos/PBDRigidParticles.h"
 #include "Chaos/Sphere.h"
 #include "Chaos/Utilities.h"
-
-namespace ChaosTest {
-
-	/**/
-	template<class FPBDRigidsEvolution, class T>
-	void Gravity();
-}
\ No newline at end of file
diff --git a/Engine/Source/Programs/HeadlessChaos/Public/HeadlessChaosTestUtility.h b/Engine/Source/Programs/HeadlessChaos/Public/HeadlessChaosTestUtility.h
--- a/Engine/Source/Programs/HeadlessChaos/Public/HeadlessChaosTestUtility.h
+++ b/Engine/Source/Programs/HeadlessChaos/Public/HeadlessChaosTestUtility.h
@@ -17,6 +17,7 @@
 #include "Chaos/Utilities.h"
 #include "Chaos/ParticleHandleFwd.h"
 #include "Chaos/PBDRigidsEvolutionFwd.h"
+#include "Chaos/EvolutionTraits.h"
 
 namespace Chaos
 {
@@ -27,6 +28,31 @@ namespace ChaosTest {
 	
 	using namespace Chaos;
 
+	template <typename... T>
+	struct TTypesWithoutVoid
+	{
+		using Types = ::testing::Types<T...>;
+	};
+
+	template <typename... T>
+	struct TTypesWithoutVoid<void,T...>
+	{
+		using Types = ::testing::Types<T...>;
+	};
+
+	//should be TAllEvolutions, but used in macros and logs and this makes it more readable
+	template <typename T>
+	class AllEvolutions : public ::testing::Test {};
+
+#define EVOLUTION_TRAIT(Trait) ,TPBDRigidsEvolutionGBF<Trait>
+	using AllEvolutionTypesTmp = TTypesWithoutVoid <
+		void
+#include "Chaos/EvolutionTraits.inl"
+	>;
+	using AllEvolutionTypes = AllEvolutionTypesTmp::Types;
+#undef EVOLUTION_TRAIT
+	TYPED_TEST_SUITE(AllEvolutions,AllEvolutionTypes);
+
 	MATCHER_P2(VectorNear, V, Tolerance, "") { return arg.Equals(V, Tolerance); }
 
 	// Expects each component of the vector is within T of its corresponding component in A. 
