Move more tests to use all evolutions

Chris.Caulfield

diff --git a/Engine/Source/Programs/HeadlessChaos/Private/HeadlessChaosApp.cpp b/Engine/Source/Programs/HeadlessChaos/Private/HeadlessChaosApp.cpp
--- a/Engine/Source/Programs/HeadlessChaos/Private/HeadlessChaosApp.cpp
+++ b/Engine/Source/Programs/HeadlessChaos/Private/HeadlessChaosApp.cpp
@@ -114,7 +114,6 @@ TEST(SerializationTests, Serialization) {
 
 TEST(BroadphaseTests, Broadphase) {
 	ChaosTest::BPPerfTest<float>();
-	ChaosTest::TestPendingSpatialDataHandlePointerConflict();
 	//ChaosTest::SpatialAccelerationDirtyAndGlobalQueryStrestTest<float>();
 	SUCCEED();
 }
diff --git a/Engine/Source/Programs/HeadlessChaos/Private/HeadlessChaosTestBroadphase.cpp b/Engine/Source/Programs/HeadlessChaos/Private/HeadlessChaosTestBroadphase.cpp
--- a/Engine/Source/Programs/HeadlessChaos/Private/HeadlessChaosTestBroadphase.cpp
+++ b/Engine/Source/Programs/HeadlessChaos/Private/HeadlessChaosTestBroadphase.cpp
@@ -682,6 +682,7 @@ namespace ChaosTest
 	}
 	template void BroadphaseCollectionTest<float>();
 
+	template <typename TEvolution>
 	void TestPendingSpatialDataHandlePointerConflict()
 	{
 		// We are testing an edge case in the way we queue updates to our acceleration structures.
@@ -702,7 +703,7 @@ namespace ChaosTest
 		Particle->X() = TVector<FReal, 3>(0, 0, 0);
 
 		THandleArray<FChaosPhysicsMaterial> PhysicalMaterials;
-		FPBDRigidsEvolution Evolution(Particles, PhysicalMaterials);
+		TEvolution Evolution(Particles, PhysicalMaterials);
 
 		// Flush spatial acceleration structures to put particle into structure.
 		Evolution.FlushSpatialAcceleration();
@@ -774,6 +775,11 @@ namespace ChaosTest
 		}
 	}
 
+	TYPED_TEST(AllEvolutions, TestPendingSpatialDataHandlePointerConflict)
+	{
+		TestPendingSpatialDataHandlePointerConflict<TypeParam>();
+	}
+
 	template<typename T>
 	void SpatialAccelerationDirtyAndGlobalQueryStrestTest()
 	{
diff --git a/Engine/Source/Programs/HeadlessChaos/Private/HeadlessChaosTestJoints.cpp b/Engine/Source/Programs/HeadlessChaos/Private/HeadlessChaosTestJoints.cpp
--- a/Engine/Source/Programs/HeadlessChaos/Private/HeadlessChaosTestJoints.cpp
+++ b/Engine/Source/Programs/HeadlessChaos/Private/HeadlessChaosTestJoints.cpp
@@ -31,11 +31,16 @@ namespace ChaosTest {
 	 * Base class for joint constraint tests.
 	 * Initialize the particle and joint data in the test code and call Create()
 	 */
-	class FJointConstraintsTest : public FConstraintsTest
+	template <typename TEvolution>
+	class FJointConstraintsTest : public FConstraintsTest<TEvolution>
 	{
 	public:
+		using Base = FConstraintsTest<TEvolution>;
+		using Base::Evolution;
+		using Base::AddParticleBox;
+
 		FJointConstraintsTest(const int32 NumIterations, const FReal Gravity)
-			: FConstraintsTest(NumIterations, Gravity)
+			: Base(NumIterations, Gravity)
 			, JointsRule(Joints)
 		{
 			Evolution.AddConstraintRule(&JointsRule);
@@ -77,13 +82,14 @@ namespace ChaosTest {
 	/**
 	 * One kinematic, one dynamic particle connected by a ball-socket joint in the middle.
 	 */
+	template <typename TEvolution>
 	void JointConstraint_Single()
 	{
 
 		const int32 NumIterations = 1;
 		const FReal Gravity = 980;
 
-		FJointConstraintsTest Test(NumIterations, Gravity);
+		FJointConstraintsTest<TEvolution> Test(NumIterations, Gravity);
 
 		Test.ParticlePositions =
 		{
@@ -133,6 +139,7 @@ namespace ChaosTest {
 		}
 	}
 
+	template <typename TEvolution>
 	void JointConstraint_SingleMoveRoot()
 	{
 		const int32 NumIterations = 5;
@@ -142,7 +149,7 @@ namespace ChaosTest {
 		const FReal Dt = (FReal)1 / 20;
 		const FVec3 RootDelta(1 * BoxSize, 0, 0);
 
-		FJointConstraintsTest Test(NumIterations, Gravity);
+		FJointConstraintsTest<TEvolution> Test(NumIterations, Gravity);
 
 		Test.ParticlePositions =
 		{
@@ -216,6 +223,7 @@ namespace ChaosTest {
 	/**
 	 * Pendulum with animated root.
 	 */
+	template <typename TEvolution>
 	void JointConstraint_SingleAnimated()
 	{
 		const int32 NumIterations = 5;
@@ -226,7 +234,7 @@ namespace ChaosTest {
 		const FReal AnimPeriod = (FReal)2;
 		const FVec3 AnimDelta = FVec3(10 * BoxSize, 0, 0);
 
-		FJointConstraintsTest Test(NumIterations, Gravity);
+		FJointConstraintsTest<TEvolution> Test(NumIterations, Gravity);
 
 		Test.ParticlePositions =
 		{
@@ -288,6 +296,7 @@ namespace ChaosTest {
 	/**
 	 * Pendulum with animated root.
 	 */
+	template <typename TEvolution>
 	void JointConstraint_ShortChainAnimated()
 	{
 		const int32 NumIterations = 10;
@@ -299,7 +308,7 @@ namespace ChaosTest {
 		const FVec3 AnimDelta = FVec3(5 * BoxSize, 0, 0);
 		const FReal AcceptableDistanceError = 5;
 
-		FJointConstraintsTest Test(NumIterations, Gravity);
+		FJointConstraintsTest<TEvolution> Test(NumIterations, Gravity);
 
 		Test.ParticlePositions =
 		{
@@ -373,6 +382,7 @@ namespace ChaosTest {
 	/**
 	 * Pendulum with animated root.
 	 */
+	template <typename TEvolution>
 	void JointConstraint_LongChainAnimated()
 	{
 		const int NumParticles = 10;
@@ -392,7 +402,7 @@ namespace ChaosTest {
 		FMath::RandInit(1048604845);
 
 		// Create a chain of connected particles, with particle 0 fixed
-		FJointConstraintsTest Test(NumIterations, Gravity);
+		FJointConstraintsTest<TEvolution> Test(NumIterations, Gravity);
 		for (int32 ParticleIndex = 0; ParticleIndex < NumParticles; ++ParticleIndex)
 		{
 			Test.ParticlePositions.Add(Begin + ParticleIndex * Separation * Dir);
@@ -458,6 +468,7 @@ namespace ChaosTest {
 		EXPECT_LT(MaxDistanceError, AcceptableDistanceError) << "On frame " << MaxDistanceErrorFrameIndex;
 	}
 
+	template <typename TEvolution>
 	void SpringConstraint()
 	{
 		TUniquePtr<FChaosPhysicsMaterial> PhysicalMaterial = MakeUnique<FChaosPhysicsMaterial>();
@@ -477,7 +488,7 @@ namespace ChaosTest {
 		Box2->X() = FVec3((FReal)500, (FReal)0, (FReal)1000);
 		Box2->P() = Box2->X();
 		THandleArray<FChaosPhysicsMaterial> PhysicalMaterials;
-		FPBDRigidsEvolution Evolution(Particles, PhysicalMaterials);
+		TEvolution Evolution(Particles, PhysicalMaterials);
 		TVector<TGeometryParticleHandle<FReal, 3>*,2> ConstrainedParticles = TVector<TGeometryParticleHandle<FReal,3>*, 2>(StaticBox, Box2);
 		TVector<FVec3, 2> Points = { FVec3((FReal)100, (FReal)0, (FReal)1000), FVec3((FReal)400, (FReal)0, (FReal)1000) };
 
@@ -498,6 +509,7 @@ namespace ChaosTest {
 		}
 	}
 
+	template <typename TEvolution>
 	void DynamicSpringConstraint()
 	{
 		TUniquePtr<FChaosPhysicsMaterial> PhysicalMaterial = MakeUnique<FChaosPhysicsMaterial>();
@@ -519,7 +531,7 @@ namespace ChaosTest {
 			Box2.P() = Box2.X();
 
 			THandleArray<FChaosPhysicsMaterial> PhysicalMaterials;
-			FPBDRigidsEvolution Evolution(Particles, PhysicalMaterials);
+			TEvolution Evolution(Particles, PhysicalMaterials);
 			TArray<TVector<TGeometryParticleHandle<FReal,3>*, 2>> Constraints = { TVector<TGeometryParticleHandle<FReal,3>*, 2>(&StaticBox, &Box2) };
 
 			Evolution.SetPhysicsMaterial(&StaticBox, MakeSerializable(PhysicalMaterial));
@@ -549,7 +561,7 @@ namespace ChaosTest {
 			Box2.P() = Box2.X();
 
 			THandleArray<FChaosPhysicsMaterial> PhysicalMaterials;
-			FPBDRigidsEvolution Evolution(Particles, PhysicalMaterials);
+			TEvolution Evolution(Particles, PhysicalMaterials);
 			TArray<TVector<TGeometryParticleHandle<FReal,3>*, 2>> Constraints = { TVector<TGeometryParticleHandle<FReal,3>*, 2>(&StaticBox, &Box2) };
 
 			Evolution.SetPhysicsMaterial(&StaticBox, MakeSerializable(PhysicalMaterial));
@@ -571,32 +583,32 @@ namespace ChaosTest {
 
 
 
-	TEST(JointTests, TestSingleConstraint) {
-		JointConstraint_Single();
+	TYPED_TEST(AllEvolutions, JointTests_TestSingleConstraint) {
+		JointConstraint_Single<TypeParam>();
 	}
 
-	TEST(JointTests, TestSingleConstraintWithLateralTranslation) {
-		JointConstraint_SingleMoveRoot();
+	TYPED_TEST(AllEvolutions, JointTests_TestSingleConstraintWithLateralTranslation) {
+		JointConstraint_SingleMoveRoot<TypeParam>();
 	}
 
-	TEST(JointTests, TestSingleConstraintWithAnimatedRoot) {
-		JointConstraint_SingleAnimated();
+	TYPED_TEST(AllEvolutions, JointTests_TestSingleConstraintWithAnimatedRoot) {
+		JointConstraint_SingleAnimated<TypeParam>();
 	}
 
-	TEST(JointTests, TestShortJointChainWithAnimatedRoot) {
-		JointConstraint_ShortChainAnimated();
+	TYPED_TEST(AllEvolutions, JointTests_TestShortJointChainWithAnimatedRoot) {
+		JointConstraint_ShortChainAnimated<TypeParam>();
 	}
 
-	TEST(JointTests, TestLongJointChainWithAnimatedRoot) {
-		JointConstraint_LongChainAnimated();
+	TYPED_TEST(AllEvolutions, JointTests_TestLongJointChainWithAnimatedRoot) {
+		JointConstraint_LongChainAnimated<TypeParam>();
 	}
 
-	TEST(JointTests, TestSingleSpringConstraint) {
-		SpringConstraint();
+	TYPED_TEST(AllEvolutions, JointTests_TestSingleSpringConstraint) {
+		SpringConstraint<TypeParam>();
 	}
 
-	TEST(JointTests, TestSingleDynamicSpringConstraint) {
-		DynamicSpringConstraint();
+	TYPED_TEST(AllEvolutions, JointTests_TestSingleDynamicSpringConstraint) {
+		DynamicSpringConstraint<TypeParam>();
 	}
 }
 
diff --git a/Engine/Source/Programs/HeadlessChaos/Private/HeadlessChaosTestSim.cpp b/Engine/Source/Programs/HeadlessChaos/Private/HeadlessChaosTestSim.cpp
--- a/Engine/Source/Programs/HeadlessChaos/Private/HeadlessChaosTestSim.cpp
+++ b/Engine/Source/Programs/HeadlessChaos/Private/HeadlessChaosTestSim.cpp
@@ -13,11 +13,12 @@ namespace ChaosTest {
 
 	using namespace Chaos;
 
-	GTEST_TEST(SimTests, SphereSphereSimTest)
+	TYPED_TEST(AllEvolutions, SimTests_SphereSphereSimTest)
 	{
+		using TEvolution = TypeParam;
 		TPBDRigidsSOAs<FReal, 3> Particles;
 		THandleArray<FChaosPhysicsMaterial> PhysicalMaterials;
-		FPBDRigidsEvolution Evolution(Particles, PhysicalMaterials);
+		TEvolution Evolution(Particles, PhysicalMaterials);
 		auto Static = Evolution.CreateStaticParticles(1)[0];
 		auto Dynamic = Evolution.CreateDynamicParticles(1)[0];
 
@@ -48,11 +49,12 @@ namespace ChaosTest {
 		EXPECT_NEAR(Dynamic->X().Z, 110, 1);
 	}
 
-	GTEST_TEST(SimTests, BoxBoxSimTest)
+	TYPED_TEST(AllEvolutions, SimTests_BoxBoxSimTest)
 	{
+		using TEvolution = TypeParam;
 		TPBDRigidsSOAs<FReal, 3> Particles;
 		THandleArray<FChaosPhysicsMaterial> PhysicalMaterials;
-		FPBDRigidsEvolution Evolution(Particles, PhysicalMaterials);
+		TEvolution Evolution(Particles, PhysicalMaterials);
 		auto Static = Evolution.CreateStaticParticles(1)[0];
 		auto Dynamic = Evolution.CreateDynamicParticles(1)[0];
 
@@ -87,11 +89,12 @@ namespace ChaosTest {
 	// This will be fixed if/when we have a multi-contact manifold between particle pairs and we simultaneously
 	// resolve contacts in that manifold.
 	//
-	GTEST_TEST(SimTests, DISABLED_VeryLowInertiaSimTest)
+	TYPED_TEST(AllEvolutions, DISABLED_SimTests_VeryLowInertiaSimTest)
 	{
+		using TEvolution = TypeParam;
 		TPBDRigidsSOAs<FReal, 3> Particles;
 		THandleArray<FChaosPhysicsMaterial> PhysicalMaterials;
-		FPBDRigidsEvolution Evolution(Particles, PhysicalMaterials);
+		TEvolution Evolution(Particles, PhysicalMaterials);
 		auto Static = Evolution.CreateStaticParticles(1)[0];
 		auto Dynamic = Evolution.CreateDynamicParticles(1)[0];
 
@@ -116,11 +119,12 @@ namespace ChaosTest {
 		EXPECT_NEAR(Dynamic->X().Z, 110, 10);
 	}
 
-	GTEST_TEST(SimTests, SleepAndWakeSimTest)
+	TYPED_TEST(AllEvolutions, SimTests_SleepAndWakeSimTest)
 	{
+		using TEvolution = TypeParam;
 		TPBDRigidsSOAs<FReal, 3> Particles;
 		THandleArray<FChaosPhysicsMaterial> PhysicalMaterials;
-		FPBDRigidsEvolution Evolution(Particles, PhysicalMaterials);
+		TEvolution Evolution(Particles, PhysicalMaterials);
 		auto Static = Evolution.CreateStaticParticles(1)[0];
 		auto Dynamic1 = Evolution.CreateDynamicParticles(1)[0];
 		auto Dynamic2 = Evolution.CreateDynamicParticles(1)[0];
diff --git a/Engine/Source/Programs/HeadlessChaos/Private/HeadlessChaosTestUtility.cpp b/Engine/Source/Programs/HeadlessChaos/Private/HeadlessChaosTestUtility.cpp
--- a/Engine/Source/Programs/HeadlessChaos/Private/HeadlessChaosTestUtility.cpp
+++ b/Engine/Source/Programs/HeadlessChaos/Private/HeadlessChaosTestUtility.cpp
@@ -44,30 +44,6 @@ namespace ChaosTest {
 	}
 	template int32 AppendAnalyticSphere(TPBDRigidParticles<float, 3> &, float Scale);
 
-	template<class T>
-	TPBDRigidParticleHandle<T, 3>* AppendAnalyticSphere2(FPBDRigidsEvolution& Evolution, T Scale)
-	{
-		TArray<TPBDRigidParticleHandle<T, 3>*> Particles = Evolution.CreateDynamicParticles(1);
-		TPBDRigidParticleHandle<T, 3>* Particle = Particles[0];
-
-		Particle->X() = TVector<T, 3>(0.f, 0.f, 0.f);
-		Particle->V() = TVector<T, 3>(0.f, 0.f, 0.f);
-		Particle->R() = TRotation<T, 3>::MakeFromEuler(TVector<T, 3>(0.f, 0.f, 0.f)).GetNormalized();
-		Particle->W() = TVector<T, 3>(0.f, 0.f, 0.f);
-		Particle->P() = Particle->X();
-		Particle->Q() = Particle->R();
-
-		Particle->M() = 1.f;
-		Particle->InvM() = 1.f;
-		Particle->I() = PMatrix<T, 3, 3>(1.f, 0.f, 0.f, 0.f, 1.f, 0.f, 0.f, 0.f, 1.f);
-		Particle->InvI() = PMatrix<T, 3, 3>(1.f, 0.f, 0.f, 0.f, 1.f, 0.f, 0.f, 0.f, 1.f);
-		Particle->SetDynamicGeometry(MakeUnique<TSphere<T, 3>>(TVector<T, 3>(0), Scale));
-
-		return Particle;
-	}
-	template TPBDRigidParticleHandle<float, 3>* AppendAnalyticSphere2(FPBDRigidsEvolution& Evolution, float Scale);
-
-
 	template<class T>
 	int32 AppendAnalyticBox(TPBDRigidParticles<T, 3> & InParticles, TVector<T, 3> Scale)
 	{
@@ -115,25 +91,6 @@ namespace ChaosTest {
 		}
 	}
 
-	template<class T>
-	TKinematicGeometryParticleHandle<T, 3>* AppendKinematicAnalyticBox2(FPBDRigidsEvolution& Evolution, TVector<T, 3> Scale)
-	{
-		TArray<TKinematicGeometryParticleHandle<T, 3>*> Particles = Evolution.CreateKinematicParticles(1);
-		InitAnalyticBox2(Particles[0], Scale);
-		return Particles[0];
-	}
-	template TKinematicGeometryParticleHandle<float, 3>* AppendKinematicAnalyticBox2(FPBDRigidsEvolution& Evolution, TVector<float, 3> Scale);
-
-	template<class T>
-	TPBDRigidParticleHandle<T, 3>* AppendDynamicAnalyticBox2(FPBDRigidsEvolution& Evolution, TVector<T, 3> Scale)
-	{
-		TArray<TPBDRigidParticleHandle<T, 3>*> Particles = Evolution.CreateDynamicParticles(1);
-		InitAnalyticBox2(Particles[0], Scale);
-		return Particles[0];
-	}
-	template TPBDRigidParticleHandle<float, 3>* AppendDynamicAnalyticBox2(FPBDRigidsEvolution& Evolution, TVector<float, 3> Scale);
-
-
 	template<class T>
 	int32 AppendParticleBox(TPBDRigidParticles<T, 3> & InParticles, TVector<T, 3> Scale, TArray<TVector<int32, 3>> * elements)
 	{
@@ -410,16 +367,6 @@ namespace ChaosTest {
 		}
 	}
 
-
-	template<class T>
-	TPBDRigidParticleHandle<T, 3>* AppendDynamicParticleBox(FPBDRigidsEvolution& Evolution, const TVector<T, 3>& Scale, TArray<TVector<int32, 3>>* OutElements)
-	{
-		TArray<TPBDRigidParticleHandle<T, 3>*> Particles = Evolution.CreateDynamicParticles(1);
-		InitDynamicParticleBox2(Particles[0], Scale, OutElements);
-		return Particles[0];
-	}
-	template TPBDRigidParticleHandle<float, 3>* AppendDynamicParticleBox(FPBDRigidsEvolution& Evolution, const TVector<float, 3>& Scale, TArray<Chaos::TVector<int32, 3>>* OutElements);
-
 	template<class T>
 	TPBDRigidParticleHandle<T, 3>* AppendDynamicParticleBox(TPBDRigidsSOAs<T, 3>& SOAs, const TVector<T, 3>& Scale, TArray<TVector<int32, 3>>* OutElements)
 	{
@@ -456,15 +403,6 @@ namespace ChaosTest {
 	}
 	template TPBDRigidParticleHandle<float, 3>* AppendDynamicParticleTaperedCylinder(TPBDRigidsSOAs<float, 3>& Evolution, const TVector<float, 3>& Scale, TArray<Chaos::TVector<int32, 3>>* OutElements);
 
-	template<class T>
-	TPBDRigidParticleHandle<T, 3>* AppendClusteredParticleBox(FPBDRigidsEvolution& Evolution, const TVector<T, 3>& Scale, TArray<TVector<int32, 3>>* OutElements)
-	{
-		auto Particles = Evolution.CreateClusteredParticles(1);
-		InitDynamicParticleBox2(Particles[0], Scale, OutElements);
-		return Particles[0];
-	}
-	template TPBDRigidParticleHandle<float, 3>* AppendClusteredParticleBox(FPBDRigidsEvolution& Evolution, const TVector<float, 3>& Scale, TArray<Chaos::TVector<int32, 3>>* OutElements);
-
 	template<class T>
 	TPBDRigidParticleHandle<T, 3>* AppendClusteredParticleBox(TPBDRigidsSOAs<T, 3>& SOAs, const TVector<T, 3>& Scale, TArray<TVector<int32, 3>>* OutElements)
 	{
@@ -554,24 +492,6 @@ namespace ChaosTest {
 	}
 	template int AppendStaticAnalyticFloor(TPBDRigidParticles<float, 3> &);
 
-	template<class T>
-	TKinematicGeometryParticleHandle<T, 3>* AppendStaticAnalyticFloor(FPBDRigidsEvolution& Evolution)
-	{
-		TArray<TKinematicGeometryParticleHandle<T, 3>*> Particles = Evolution.CreateKinematicParticles(1);
-		TKinematicGeometryParticleHandle<T, 3>* Particle = Particles[0];
-
-		Particle->X() = TVector<T, 3>(0.f, 0.f, 0.f);
-		Particle->V() = TVector<T, 3>(0.f, 0.f, 0.f);
-		Particle->R() = TRotation<T, 3>::MakeFromEuler(TVector<T, 3>(0.f, 0.f, 0.f)).GetNormalized();
-		Particle->W() = TVector<T, 3>(0.f, 0.f, 0.f);
-		Particle->SetDynamicGeometry( MakeUnique<TPlane<T, 3>>(TVector<T, 3>(0.f, 0.f, 0.f), TVector<T, 3>(0.f, 0.f, 1.f)));
-
-		::ChaosTest::SetParticleSimDataToCollide({ Particle });
-
-		return Particle;
-	}
-	template TKinematicGeometryParticleHandle<float, 3>* AppendStaticAnalyticFloor(FPBDRigidsEvolution& Evolution);
-
 	template<class T>
 	TKinematicGeometryParticleHandle<T, 3>* AppendStaticAnalyticFloor(TPBDRigidsSOAs<T, 3>& SOAs)
 	{
diff --git a/Engine/Source/Programs/HeadlessChaos/Public/HeadlessChaosTestBroadphase.h b/Engine/Source/Programs/HeadlessChaos/Public/HeadlessChaosTestBroadphase.h
--- a/Engine/Source/Programs/HeadlessChaos/Public/HeadlessChaosTestBroadphase.h
+++ b/Engine/Source/Programs/HeadlessChaos/Public/HeadlessChaosTestBroadphase.h
@@ -22,8 +22,6 @@ namespace ChaosTest {
 	template<typename T>
 	void BroadphaseCollectionTest();
 
-	void TestPendingSpatialDataHandlePointerConflict();
-
 	template<typename T>
 	void SpatialAccelerationDirtyAndGlobalQueryStrestTest();
 }
\ No newline at end of file
diff --git a/Engine/Source/Programs/HeadlessChaos/Public/HeadlessChaosTestConstraints.h b/Engine/Source/Programs/HeadlessChaos/Public/HeadlessChaosTestConstraints.h
--- a/Engine/Source/Programs/HeadlessChaos/Public/HeadlessChaosTestConstraints.h
+++ b/Engine/Source/Programs/HeadlessChaos/Public/HeadlessChaosTestConstraints.h
@@ -15,6 +15,7 @@ namespace ChaosTest
 	/**
 	 * Base class for constraint tests. Provides a basic sim with no builtin constraint support.
 	 */
+	template <typename TEvolution>
 	class FConstraintsTest
 	{
 	public:
@@ -61,7 +62,7 @@ namespace ChaosTest
 
 		// Solver state
 		TPBDRigidsSOAs<FReal, 3> SOAs;
-		FPBDRigidsEvolution Evolution;
+		TEvolution Evolution;
 		TUniquePtr<FChaosPhysicsMaterial> PhysicalMaterial;
 		THandleArray<FChaosPhysicsMaterial> PhysicalMaterials;
 
diff --git a/Engine/Source/Programs/HeadlessChaos/Public/HeadlessChaosTestUtility.h b/Engine/Source/Programs/HeadlessChaos/Public/HeadlessChaosTestUtility.h
--- a/Engine/Source/Programs/HeadlessChaos/Public/HeadlessChaosTestUtility.h
+++ b/Engine/Source/Programs/HeadlessChaos/Public/HeadlessChaosTestUtility.h
@@ -70,22 +70,14 @@ namespace ChaosTest {
 	/**/
 	template<class T>
 	int32 AppendAnalyticSphere(TPBDRigidParticles<T, 3> & InParticles, T Scale = (T)1);
-	template<class T>
-	TPBDRigidParticleHandle<T, 3>* AppendAnalyticSphere2(FPBDRigidsEvolution& Evolution, T Scale = (T)1);
 
 	/**/
 	template<class T>
 	int32 AppendAnalyticBox(TPBDRigidParticles<T, 3>& InParticles, TVector<T, 3> Scale = TVector<T, 3>(1));
-	template<class T>
-	TKinematicGeometryParticleHandle<T, 3>* AppendKinematicAnalyticBox2(FPBDRigidsEvolution& Evolution, TVector<T, 3> Scale = TVector<T, 3>(1));
-	template<class T>
-	TPBDRigidParticleHandle<T, 3>* AppendDynamicAnalyticBox2(FPBDRigidsEvolution& Evolution, TVector<T, 3> Scale = TVector<T, 3>(1));
 
 	/**/
 	template<class T>
 	int32 AppendParticleBox(TPBDRigidParticles<T, 3>& InParticles, TVector<T, 3> Scale = TVector<T, 3>(1), TArray<TVector<int32, 3>>* OutElements = nullptr);
-	template<class T>
-	TPBDRigidParticleHandle<T, 3>* AppendDynamicParticleBox(FPBDRigidsEvolution& Evolution, const TVector<T, 3>& Scale = TVector<T, 3>(1), TArray<TVector<int32, 3>>* OutElements = nullptr);
 
 	template<class T>
 	TPBDRigidParticleHandle<T, 3>* AppendDynamicParticleBox(TPBDRigidsSOAs<T, 3>& SOAs, const TVector<T, 3>& Scale = TVector<T, 3>(1), TArray<TVector<int32, 3>>* OutElements = nullptr);
@@ -102,17 +94,12 @@ namespace ChaosTest {
 	template<class T>
 	TGeometryParticleHandle<T, 3>* AppendStaticParticleBox(TPBDRigidsSOAs<T, 3>& SOAs, const TVector<T, 3>& Scale = TVector<T, 3>(1), TArray<TVector<int32, 3>>* OutElements = nullptr);
 
-	template<class T>
-	TPBDRigidParticleHandle<T, 3>* AppendClusteredParticleBox(FPBDRigidsEvolution& Evolution, const TVector<T, 3>& Scale = TVector<T, 3>(1), TArray<TVector<int32, 3>>* OutElements = nullptr);
-
 	template<class T>
 	TPBDRigidParticleHandle<T, 3>* AppendClusteredParticleBox(TPBDRigidsSOAs<T, 3>& SOAs, const TVector<T, 3>& Scale = TVector<T, 3>(1), TArray<TVector<int32, 3>>* OutElements = nullptr);
 		
 	/**/
 	template<class T>
 	int32 AppendStaticAnalyticFloor(TPBDRigidParticles<T, 3>& InParticles);
-	template<class T>
-	TKinematicGeometryParticleHandle<T, 3>* AppendStaticAnalyticFloor(FPBDRigidsEvolution& Evolution);
 
 	template<class T>
 	TKinematicGeometryParticleHandle<T, 3>* AppendStaticAnalyticFloor(TPBDRigidsSOAs<T, 3>& SOAs);
diff --git a/Engine/Source/Runtime/Experimental/Chaos/Public/Chaos/PBDRigidsEvolution.h b/Engine/Source/Runtime/Experimental/Chaos/Public/Chaos/PBDRigidsEvolution.h
--- a/Engine/Source/Runtime/Experimental/Chaos/Public/Chaos/PBDRigidsEvolution.h
+++ b/Engine/Source/Runtime/Experimental/Chaos/Public/Chaos/PBDRigidsEvolution.h
@@ -23,7 +23,11 @@ extern int32 ChaosNumPushOutIterationsOverride;
 extern int32 ChaosNumContactIterationsOverride;
 
 // Declaring so it can be friended for tests.
-namespace ChaosTest { void TestPendingSpatialDataHandlePointerConflict(); } 
+namespace ChaosTest
+{
+	template <typename TEvolution>
+	void TestPendingSpatialDataHandlePointerConflict();
+} 
 
 namespace Chaos
 {
@@ -232,6 +236,7 @@ class TPBDRigidsEvolutionBase
 	typedef TFunction<void(TPBDRigidParticles<FReal, 3>&, const FReal, const FReal, const int32)> FKinematicUpdateRule;
 	typedef TFunction<void(TParticleView<TPBDRigidParticles<FReal,3>>&)> FCaptureRewindRule;
 
+	template <typename TEvolution>
 	friend void ChaosTest::TestPendingSpatialDataHandlePointerConflict();
 
 	CHAOS_API TPBDRigidsEvolutionBase(TPBDRigidsSOAs<FReal, 3>& InParticles, THandleArray<FChaosPhysicsMaterial>& InSolverPhysicsMaterials, int32 InNumIterations = 1, int32 InNumPushOutIterations = 1, bool InIsSingleThreaded = false);
