Add unit test for real scaled trimesh sweep outtime bug.

diff --git a/Engine/Source/Programs/HeadlessChaos/Private/HeadlessChaosApp.cpp b/Engine/Source/Programs/HeadlessChaos/Private/HeadlessChaosApp.cpp
--- a/Engine/Source/Programs/HeadlessChaos/Private/HeadlessChaosApp.cpp
+++ b/Engine/Source/Programs/HeadlessChaos/Private/HeadlessChaosApp.cpp
@@ -26,6 +26,7 @@
 #include "HeadlessChaosTestSerialization.h"
 #include "HeadlessChaosTestBP.h"
 #include "HeadlessChaosTestRaycast.h"
+#include "HeadlessChaosTestSweep.h"
 #include "HeadlessChaosTestGJK.h"
 #include "HeadlessChaosTestEPA.h"
 #include "HeadlessChaosTestBroadphase.h"
@@ -140,6 +141,10 @@ TEST(RaycastTests, Raycast) {
 	SUCCEED();
 }
 
+TEST(SweepTests, Sweep) {
+	ChaosTest::CapsuleSweepAgainstTriMeshReal<float>();
+}
+
 TEST(MostOpposingTests, MostOpposing) {
 	ChaosTest::TrimeshMostOpposing<float>();
 	ChaosTest::ConvexMostOpposing<float>();
diff --git a/Engine/Source/Programs/HeadlessChaos/Private/HeadlessChaosTestSweep.cpp b/Engine/Source/Programs/HeadlessChaos/Private/HeadlessChaosTestSweep.cpp
new file mode 100644
--- /dev/null
+++ b/Engine/Source/Programs/HeadlessChaos/Private/HeadlessChaosTestSweep.cpp
@@ -0,0 +1,888 @@
+// Copyright Epic Games, Inc. All Rights Reserved.
+
+#include "HeadlessChaosTestSweep.h"
+
+#include "HeadlessChaos.h"
+#include "Chaos/Capsule.h"
+#include "Chaos/Convex.h"
+#include "Chaos/TriangleMeshImplicitObject.h"
+#include "Chaos/ImplicitObjectScaled.h"
+
+namespace ChaosTest
+{
+	using namespace Chaos;
+
+	template<typename T>
+	void CapsuleSweepAgainstTriMeshReal()
+	{
+		// Trimesh is from SM_Cattus_POI_Rib, this was a real world failure that is now fixed.
+
+		using namespace Chaos;
+		TParticles<FReal, 3> TrimeshParticles(
+		{
+			{29.0593967, -5.21321106, -10.2669592},
+			{34.5006638, -3.16600156, -14.5092020},
+			{28.8770218, -5.00888205, -10.4567394},
+			{38.9350929, -1.47836626, -17.6896191},
+			{39.1246262, -1.58615386, -17.4255066},
+			{42.4614334, -0.114946000, -19.8563728},
+			{38.7996712, -1.91535223, -16.3864536},
+			{29.4634132, -5.15883255, -9.64184284},
+			{29.6582336, -5.05185938, -9.37638092},
+			{45.0206337, 0.641714215, -21.0300026},
+			{46.2034111, 1.19252074, -22.0121441},
+			{44.3153610, 0.398101240, -19.9310055},
+			{38.2805557, -1.87191427, -15.7972298},
+			{48.8065224, 1.82796133, -22.7662964},
+			{49.7862473, 2.21876359, -23.4238491},
+			{49.1621628, 1.89441979, -22.1913338},
+			{52.2055054, 2.59018326, -23.6808910},
+			{55.6262436, 3.71029496, -24.4840908},
+			{52.2546806, 2.55433965, -22.8908424},
+			{47.6990814, 1.76376379, -21.3171616},
+			{55.9418335, 3.34045410, -24.2242432},
+			{56.7527275, 4.26803923, -24.7314072},
+			{56.6030273, 3.37918210, -23.5389385},
+			{57.1201477, 3.92609882, -24.5670509},
+			{57.5926590, 5.32376480, -24.8299332},
+			{57.3900146, 3.86687994, -24.3628445},
+			{58.3100967, 5.59623432, -24.5884380},
+			{57.9720459, 6.88918400, -24.9411621},
+			{58.3736610, 5.22753286, -23.6481304},
+			{58.3795776, 9.20724392, -24.4941502},
+			{57.6161652, 7.89536858, -25.0186253},
+			{57.7064209, 8.63112640, -24.8542480},
+			{57.4299011, 9.53679848, -24.7764454},
+			{58.6315613, 9.11761856, -24.1844330},
+			{58.5211067, 7.24213171, -23.5957718},
+			{57.5271606, 9.99959183, -24.4184647},
+			{58.2185135, 9.55542183, -23.2548218},
+			{57.3097725, 10.2978697, -23.5637169},
+			{56.3722000, 9.18606091, -24.9406128},
+			{58.2657166, 8.36739063, -23.1939850},
+			{56.3347893, 9.67211342, -24.4351349},
+			{57.3641586, 9.99938393, -23.0135994},
+			{58.1764908, 5.82101727, -23.1902943},
+			{57.3481255, 8.60126781, -22.9416485},
+			{56.1609459, 9.56552029, -23.1539116},
+			{56.2483864, 9.85045052, -23.4782410},
+			{55.7174797, 8.29436588, -23.0638103},
+			{54.7229729, 8.73790932, -23.3331184},
+			{54.0776367, 8.77919769, -23.7519341},
+			{53.9943695, 8.17912292, -24.5668526},
+			{51.3485641, 7.47948122, -23.0380058},
+			{51.2222366, 7.41138554, -23.6724663},
+			{53.9260139, 7.74872541, -24.6529694},
+			{52.2321625, 7.36892176, -22.8332882},
+			{50.7700424, 6.47493553, -23.9456253},
+			{56.6027756, 8.31463623, -25.0338459},
+			{57.2853279, 6.51692963, -25.0099850},
+			{50.7739677, 5.00270128, -24.0342751},
+			{47.1190453, 4.83312464, -22.8379707},
+			{49.7170258, 2.62024260, -23.6686172},
+			{43.0110207, 0.426635355, -20.4862480},
+			{42.7733459, 3.24669766, -20.7472992},
+			{34.1604767, -2.98149395, -14.5161381},
+			{28.7377853, -4.72821426, -10.5358887},
+			{28.2660236, -3.50151801, -10.7232647},
+			{34.4872894, 0.0475072451, -15.1762705},
+			{27.9901352, -2.42618680, -10.5657129},
+			{27.6759071, 0.205285028, -10.2063637},
+			{36.6210670, 2.21063828, -16.3613548},
+			{27.6877193, 0.690624714, -10.0054235},
+			{42.8507233, 3.80447602, -20.6911526},
+			{39.7815170, 3.57913852, -17.9840279},
+			{42.7781487, 4.36921978, -19.9032001},
+			{27.8298340, 0.836181641, -9.67990112},
+			{28.1982155, 0.777253330, -9.15566635},
+			{46.0600357, 5.31439161, -22.0218697},
+			{36.9135399, 2.88628912, -15.1429548},
+			{47.6548805, 5.92783451, -22.1377392},
+			{44.6417084, 4.89560652, -20.4432411},
+			{47.3978271, 5.70256472, -21.4681816},
+			{42.0819168, 3.81306863, -18.2309093},
+			{34.6343575, 1.72932339, -12.9041548},
+			{28.5061035, 0.184601068, -8.87326050},
+			{28.7382660, -0.523662448, -8.71943665},
+			{47.9805756, 5.56415987, -21.3477497},
+			{35.0854912, 1.10745859, -13.0555801},
+			{29.5277596, -4.36109161, -9.26433849},
+			{44.2550583, 3.73801875, -19.3567085},
+			{48.6162643, 5.16706753, -21.4052410},
+			{33.8666039, -3.11327124, -12.2873840},
+			{29.5907612, -4.66688442, -9.28112030},
+			{37.8056297, -1.60099864, -15.0567427},
+			{43.0328369, 1.30700612, -18.5348587},
+			{45.9731026, 1.34582365, -20.3940430},
+			{48.3442001, 2.43499756, -21.2875118},
+			{53.0610924, 3.33885503, -22.4557037},
+			{54.3709755, 5.39051533, -22.8000984},
+			{55.9649162, 3.45125723, -23.1164742},
+			{56.6591415, 4.77854252, -22.8085098},
+			{57.1962891, 4.12890625, -23.1115723},
+			{25.7303848, -5.68249369, -8.51585770},
+			{26.5309143, -4.24291611, -9.48308945},
+			{24.8292999, -5.12226152, -8.47907925},
+			{20.3731937, -6.52481079, -5.81806421},
+			{26.4056702, -3.28238487, -9.07682896},
+			{23.9174252, -5.14487743, -7.70250511},
+			{21.3869915, -5.92491102, -6.74390554},
+			{18.1200867, -6.53198242, -4.76391602},
+			{21.6681881, -5.50328684, -6.94691038},
+			{18.0004005, -5.80428123, -5.25099564},
+			{15.9813662, -6.10711002, -4.21955252},
+			{18.8343563, -5.61828232, -5.61174917},
+			{20.2135773, -4.88624334, -5.77969599},
+			{15.4357758, -5.21927977, -4.48837900},
+			{11.1496544, -2.60679626, -4.38681793},
+			{10.0865431, -3.16139269, -3.85664177},
+			{8.63047791, -2.38198924, -3.79941773},
+			{14.6292696, -4.35020351, -4.44057178},
+			{7.90794277, -1.09154844, -4.05677366},
+			{6.95699024, -1.20261848, -3.97343445},
+			{8.23572540, -0.744559944, -4.00086451},
+			{3.10794330, 0.166724160, -4.27361870},
+			{3.13311577, -0.346626908, -4.11541176},
+			{1.48072076, -0.108915798, -4.42642593},
+			{0.311369270, -0.0673256740, -4.49766636},
+			{-1.59366548, 0.725093842, -4.37339926},
+			{-3.45725179, -0.749191761, -4.34381390},
+			{-6.95951319, 0.0120848129, -3.59336853},
+			{-6.83396530, -1.23874116, -3.47961354},
+			{-7.89957952, -1.66941118, -2.82847047},
+			{0.703817546, 0.936612248, -4.19160843},
+			{-5.88422966, 0.830549538, -3.74061680},
+			{-8.81333447, -0.651533544, -2.60129619},
+			{3.10681534, 0.612948775, -3.77665353},
+			{-4.67421389, 1.29445243, -3.36058092},
+			{-0.481138408, 1.18588996, -3.53271747},
+			{7.11869955, -0.267135799, -3.69777322},
+			{-7.27621984, 0.737751126, -2.76471734},
+			{2.97556901, 1.46023333, -1.62857735},
+			{-8.61850739, -0.257534623, -2.48780465},
+			{-7.48897791, 0.836947024, -2.17553926},
+			{-6.28466797, 1.53784180, -1.55349731},
+			{-9.28733826, -2.43306065, -1.33127916},
+			{-5.33093262, 1.85986328, -1.26861572},
+			{-8.90132236, -2.71124601, -1.38725543},
+			{-9.17060471, -2.78321695, -0.819307029},
+			{-7.84990644, -2.40683627, -2.12729597},
+			{-7.91611624, -3.38985372, -0.952804804},
+			{-7.87829638, -3.03356981, -1.37831473},
+			{-2.43495178, 1.76181090, -1.93872476},
+			{-5.54367924, 1.81969070, -0.745879471},
+			{-1.47680271, 1.64132023, -1.37136936},
+			{-7.19018078, 0.901845515, -0.857508898},
+			{-2.50409842, 1.67723644, -1.04755175},
+			{-5.39768934, 1.59197235, -0.603915393},
+			{-8.16507626, 0.222914696, -1.11252105},
+			{-8.16396332, -0.777114809, -0.289943308},
+			{-9.22219181, -2.11877251, -0.794639707},
+			{-8.59923363, -0.829777956, -0.520817995},
+			{-8.82921314, -1.96039867, -0.298297912},
+			{-7.40556860, -1.55998039, -0.250481576},
+			{-8.14031315, -2.90432310, -0.262226462},
+			{-6.85163212, -3.75249863, -0.688983977},
+			{-7.04014444, -3.11551261, -0.126083776},
+			{-6.32870483, -3.58578491, -0.493814230},
+			{-6.28338480, -3.61894345, -0.929058373},
+			{-2.62741208, -3.09605527, -1.33044124},
+			{-1.79794896, -2.99078751, -1.21258605},
+			{-2.96854925, -2.36802864, -0.369529188},
+			{-4.30677605, -1.66542077, -0.303167671},
+			{-4.98426914, 0.696137369, -0.565792441},
+			{-2.85455132, -0.0173058268, -0.637033641},
+			{-2.70905447, 1.40332532, -0.955382884},
+			{-2.02012229, 1.18688262, -0.943184257},
+			{-2.42252111, 0.0620364472, -0.474374950},
+			{-1.56289959, 1.40382648, -0.681109011},
+			{-2.39032364, -1.25889599, 0.0907319933},
+			{-2.19775391, 0.174560547, 0.274658203},
+			{-0.306748420, 1.84567726, -1.35950983},
+			{0.250524580, 2.27636194, -0.823237658},
+			{-1.23819447, 1.49723995, 0.237386853},
+			{-2.13792968, -0.616698205, 0.937073231},
+			{-0.347599745, 1.71002686, 0.470531255},
+			{-1.66023159, -2.55999684, -0.452154934},
+			{-1.58301616, -2.24312091, 0.468077749},
+			{-1.67812920, 0.0274793506, 0.998333335},
+			{-1.36489558, -1.84878838, 0.930600643},
+			{0.255698651, -2.16175246, 0.109909050},
+			{-0.787146091, -1.07865918, 0.862207949},
+			{-0.293874621, -2.31530595, 0.560731053},
+			{0.162645504, 1.18506360, 0.511415780},
+			{0.535789251, -1.17199028, 0.772845566},
+			{1.08959603, -1.89062190, 0.519282222},
+			{-0.329488188, -2.73174953, -0.953910708},
+			{3.05201745, -1.60201991, -0.293554634},
+			{1.98701501, -1.01330709, 0.433782279},
+			{4.13116980, -1.43739676, -0.0357451625},
+			{2.42845988, 0.716189802, 0.176060840},
+			{5.01461887, -0.924357474, -0.0785766020},
+			{2.14807916, 1.47517538, 0.204042286},
+			{8.51385689, -0.236161187, -0.669285059},
+			{1.14428675, 2.08710790, -0.255287379},
+			{6.16604090, 0.561538637, -0.544156194},
+			{6.23778296, 0.760041595, -0.706411839},
+			{4.30203772, 1.36588299, -1.17367637},
+			{8.15555477, 0.428962231, -1.73520589},
+			{12.0248108, -0.490237832, -1.43530607},
+			{11.9705954, -0.198881984, -1.88936043},
+			{10.0048714, -0.699601173, -3.28313756},
+			{11.2095804, -0.147232190, -2.45466590},
+			{16.6165028, -0.696181834, -3.11862493},
+			{11.7848577, -1.81683028, -3.73632550},
+			{11.8666601, -1.22696412, -3.14981008},
+			{15.7040558, -0.663861752, -3.65228081},
+			{14.0951519, -1.07241631, -3.61164665},
+			{13.9181585, -2.03203130, -3.63557220},
+			{14.9848728, -1.53837466, -3.98249745},
+			{15.0845509, -3.46569848, -3.95276761},
+			{16.0242348, -2.97810125, -4.17965460},
+			{17.6243019, -4.46606350, -4.79625034},
+			{17.4456787, -3.35461426, -5.32739258},
+			{18.7262955, -3.89710188, -5.70676517},
+			{16.9230785, -2.57194042, -5.40591431},
+			{22.7467842, -3.61460471, -6.97978115},
+			{18.4024391, -1.30696046, -6.12640095},
+			{24.8947697, -0.243050858, -8.32624245},
+			{20.7879944, -1.08460391, -6.77528811},
+			{21.7145271, -0.556654394, -6.90972757},
+			{17.4861660, -1.25891256, -5.46996546},
+			{24.8378277, 0.196830407, -7.98421001},
+			{18.6861229, -0.568430483, -5.20232296},
+			{23.9673500, -0.0726047829, -6.57525635},
+			{20.4216003, -0.520928204, -4.92496729},
+			{23.2731094, -0.885929525, -5.62632990},
+			{17.6102352, -1.26811695, -3.06619239},
+			{18.1110725, -2.89222097, -2.71354604},
+			{15.2493439, -1.51500702, -1.99904561},
+			{23.4795017, -5.85999155, -5.31351185},
+			{23.3095360, -6.35326576, -5.25977182},
+			{17.5398540, -4.66228580, -2.38013840},
+			{23.1760674, -6.62616539, -5.57958126},
+			{19.5313797, -6.43315840, -3.60291767},
+			{15.4807367, -3.66881418, -1.85560131},
+			{14.3018312, -4.70082331, -1.46323049},
+			{12.8267784, -4.37319517, -1.09716678},
+			{15.7109251, -6.05440664, -2.37732816},
+			{12.3741608, -1.41286337, -1.26915634},
+			{13.8603306, -5.49541807, -1.90867615},
+			{17.6683712, -6.81585407, -3.37876749},
+			{12.2477732, -4.72466660, -1.29786050},
+			{11.1379576, -3.74680114, -0.753773510},
+			{9.36099339, -2.47074175, -0.534102321},
+			{10.0452557, -3.30079317, -0.685146451},
+			{6.68738413, -2.30865097, -0.461158574},
+			{7.36793947, -3.03040981, -1.47094512},
+			{11.8176231, -4.84084797, -1.82727480},
+			{4.59196663, -2.28209496, -1.41259789},
+			{7.10861683, -2.85069633, -2.06923938},
+			{3.39536858, -2.24798036, -1.65714610},
+			{0.172218561, -2.71166372, -1.40673709},
+			{2.89844537, -1.88502884, -2.27524662},
+			{-2.25631189, -2.96263766, -1.67407131},
+			{4.27257490, -1.76102304, -2.35111189},
+			{-1.88132656, -2.04832077, -2.39434218},
+			{-6.50054932, -3.11684728, -1.45549214},
+			{-6.90812969, -1.50060546, -3.38022733},
+			{-2.73349977, -0.821462870, -4.29114628},
+			{1.80500829, -1.06030691, -3.12149286},
+			{4.53264713, -1.14842916, -3.26621342},
+			{7.26582527, -2.38612676, -2.95009947},
+			{11.0593281, -3.88886213, -3.55232120},
+			{13.7664843, -5.74423409, -2.92658687},
+			{17.2798290, -6.74407482, -4.16304970},
+			{22.1492062, -6.65115118, -6.29792738},
+			{ 26.2703266, -5.87738276, -8.47085190 }
+		});
+
+		TArray<TVector<int32, 3>> Indices(
+		{
+			{1, 0, 2},
+			{3, 0, 1},
+			{0, 3, 4},
+			{5, 4, 3},
+			{4, 6, 0},
+			{7, 0, 6},
+			{7, 6, 8},
+			{4, 5, 9},
+			{5, 10, 9},
+			{4, 11, 6},
+			{4, 9, 11},
+			{12, 8, 6},
+			{6, 11, 12},
+			{9, 10, 13},
+			{10, 14, 13},
+			{9, 15, 11},
+			{9, 13, 15},
+			{13, 14, 16},
+			{17, 16, 14},
+			{13, 18, 15},
+			{18, 13, 16},
+			{19, 11, 15},
+			{19, 15, 18},
+			{11, 19, 12},
+			{16, 17, 20},
+			{21, 20, 17},
+			{16, 22, 18},
+			{22, 16, 20},
+			{23, 20, 21},
+			{24, 23, 21},
+			{23, 25, 20},
+			{20, 25, 22},
+			{26, 23, 24},
+			{26, 25, 23},
+			{27, 26, 24},
+			{25, 28, 22},
+			{28, 25, 26},
+			{26, 27, 29},
+			{29, 27, 30},
+			{31, 29, 30},
+			{31, 32, 29},
+			{33, 26, 29},
+			{28, 26, 34},
+			{33, 34, 26},
+			{35, 33, 29},
+			{35, 29, 32},
+			{34, 33, 36},
+			{37, 33, 35},
+			{33, 37, 36},
+			{38, 35, 32},
+			{34, 36, 39},
+			{35, 38, 40},
+			{40, 37, 35},
+			{36, 41, 39},
+			{37, 41, 36},
+			{42, 34, 39},
+			{34, 42, 28},
+			{43, 39, 41},
+			{43, 42, 39},
+			{44, 41, 37},
+			{44, 43, 41},
+			{37, 45, 44},
+			{37, 40, 45},
+			{44, 46, 43},
+			{42, 43, 46},
+			{45, 47, 44},
+			{46, 44, 47},
+			{48, 45, 40},
+			{47, 45, 48},
+			{48, 40, 49},
+			{38, 49, 40},
+			{48, 50, 47},
+			{51, 48, 49},
+			{50, 48, 51},
+			{49, 38, 52},
+			{50, 53, 47},
+			{53, 46, 47},
+			{52, 54, 49},
+			{49, 54, 51},
+			{38, 55, 52},
+			{32, 55, 38},
+			{55, 32, 31},
+			{55, 31, 30},
+			{30, 56, 55},
+			{52, 55, 56},
+			{30, 27, 56},
+			{24, 56, 27},
+			{56, 24, 21},
+			{56, 57, 52},
+			{52, 57, 54},
+			{14, 21, 17},
+			{57, 58, 54},
+			{51, 54, 58},
+			{21, 14, 59},
+			{59, 56, 21},
+			{59, 57, 56},
+			{59, 58, 57},
+			{10, 59, 14},
+			{59, 10, 60},
+			{5, 60, 10},
+			{3, 60, 5},
+			{59, 61, 58},
+			{60, 61, 59},
+			{60, 3, 62},
+			{1, 62, 3},
+			{2, 62, 1},
+			{62, 2, 63},
+			{62, 63, 64},
+			{60, 62, 65},
+			{62, 64, 65},
+			{60, 65, 61},
+			{66, 65, 64},
+			{66, 67, 65},
+			{67, 68, 65},
+			{61, 65, 68},
+			{68, 67, 69},
+			{61, 68, 70},
+			{58, 61, 70},
+			{69, 71, 68},
+			{70, 68, 71},
+			{69, 72, 71},
+			{70, 71, 72},
+			{72, 69, 73},
+			{74, 72, 73},
+			{58, 70, 75},
+			{70, 72, 75},
+			{58, 75, 51},
+			{74, 76, 72},
+			{75, 77, 51},
+			{50, 51, 77},
+			{72, 78, 75},
+			{78, 77, 75},
+			{78, 72, 76},
+			{79, 50, 77},
+			{79, 77, 78},
+			{79, 53, 50},
+			{80, 78, 76},
+			{80, 79, 78},
+			{76, 74, 81},
+			{81, 80, 76},
+			{82, 81, 74},
+			{82, 83, 81},
+			{80, 84, 79},
+			{84, 53, 79},
+			{83, 85, 81},
+			{85, 80, 81},
+			{83, 86, 85},
+			{87, 84, 80},
+			{80, 85, 87},
+			{88, 53, 84},
+			{87, 88, 84},
+			{89, 85, 86},
+			{89, 86, 90},
+			{90, 8, 89},
+			{12, 89, 8},
+			{91, 85, 89},
+			{89, 12, 91},
+			{85, 92, 87},
+			{92, 85, 91},
+			{91, 12, 93},
+			{93, 92, 91},
+			{12, 19, 93},
+			{87, 92, 94},
+			{92, 93, 94},
+			{93, 19, 94},
+			{87, 94, 88},
+			{94, 19, 95},
+			{19, 18, 95},
+			{88, 94, 96},
+			{96, 94, 95},
+			{53, 88, 96},
+			{18, 97, 95},
+			{97, 18, 22},
+			{95, 98, 96},
+			{98, 53, 96},
+			{95, 97, 98},
+			{53, 98, 46},
+			{46, 98, 42},
+			{97, 22, 99},
+			{98, 97, 99},
+			{98, 99, 42},
+			{28, 99, 22},
+			{99, 28, 42},
+			{100, 63, 2},
+			{100, 64, 63},
+			{64, 100, 101},
+			{102, 101, 100},
+			{103, 102, 100},
+			{104, 64, 101},
+			{64, 104, 66},
+			{101, 102, 105},
+			{105, 104, 101},
+			{102, 103, 106},
+			{107, 106, 103},
+			{102, 108, 105},
+			{108, 102, 106},
+			{109, 106, 107},
+			{110, 109, 107},
+			{111, 108, 106},
+			{109, 111, 106},
+			{112, 105, 108},
+			{111, 112, 108},
+			{113, 109, 110},
+			{113, 111, 109},
+			{114, 113, 110},
+			{115, 114, 110},
+			{115, 116, 114},
+			{117, 111, 113},
+			{114, 117, 113},
+			{116, 118, 114},
+			{119, 118, 116},
+			{114, 118, 120},
+			{117, 114, 120},
+			{121, 118, 119},
+			{118, 121, 120},
+			{119, 122, 121},
+			{123, 121, 122},
+			{124, 121, 123},
+			{121, 124, 125},
+			{126, 125, 124},
+			{127, 125, 126},
+			{128, 127, 126},
+			{128, 129, 127},
+			{130, 121, 125},
+			{125, 127, 131},
+			{130, 125, 131},
+			{132, 127, 129},
+			{127, 132, 131},
+			{121, 130, 133},
+			{120, 121, 133},
+			{130, 131, 134},
+			{130, 135, 133},
+			{134, 135, 130},
+			{133, 136, 120},
+			{134, 131, 137},
+			{132, 137, 131},
+			{133, 138, 136},
+			{135, 138, 133},
+			{139, 137, 132},
+			{139, 140, 137},
+			{137, 140, 141},
+			{137, 141, 134},
+			{132, 142, 139},
+			{141, 143, 134},
+			{142, 132, 144},
+			{132, 129, 144},
+			{142, 144, 145},
+			{129, 146, 144},
+			{144, 147, 145},
+			{148, 144, 146},
+			{147, 144, 148},
+			{149, 134, 143},
+			{149, 135, 134},
+			{143, 150, 149},
+			{143, 141, 150},
+			{135, 149, 151},
+			{141, 152, 150},
+			{149, 150, 153},
+			{151, 149, 153},
+			{150, 152, 154},
+			{153, 150, 154},
+			{152, 141, 155},
+			{155, 141, 140},
+			{139, 155, 140},
+			{156, 152, 155},
+			{154, 152, 156},
+			{157, 155, 139},
+			{157, 139, 142},
+			{157, 142, 145},
+			{157, 158, 155},
+			{158, 156, 155},
+			{158, 157, 159},
+			{158, 159, 156},
+			{157, 145, 159},
+			{159, 160, 156},
+			{156, 160, 154},
+			{161, 159, 145},
+			{147, 161, 145},
+			{162, 161, 147},
+			{163, 159, 161},
+			{163, 161, 162},
+			{159, 163, 160},
+			{163, 162, 164},
+			{165, 164, 162},
+			{166, 164, 165},
+			{164, 166, 167},
+			{163, 164, 168},
+			{168, 164, 167},
+			{160, 163, 169},
+			{163, 168, 169},
+			{160, 169, 170},
+			{154, 160, 170},
+			{168, 171, 169},
+			{169, 171, 170},
+			{154, 170, 172},
+			{170, 171, 172},
+			{154, 172, 153},
+			{173, 153, 172},
+			{172, 171, 173},
+			{173, 151, 153},
+			{173, 171, 174},
+			{151, 173, 175},
+			{174, 175, 173},
+			{171, 176, 174},
+			{171, 168, 176},
+			{175, 174, 177},
+			{177, 174, 176},
+			{178, 151, 175},
+			{135, 151, 178},
+			{138, 135, 178},
+			{175, 179, 178},
+			{179, 138, 178},
+			{175, 177, 180},
+			{180, 179, 175},
+			{177, 181, 180},
+			{177, 176, 181},
+			{182, 179, 180},
+			{168, 183, 176},
+			{168, 167, 183},
+			{184, 181, 176},
+			{183, 184, 176},
+			{180, 181, 185},
+			{180, 185, 182},
+			{186, 181, 184},
+			{186, 185, 181},
+			{184, 183, 187},
+			{185, 186, 188},
+			{186, 184, 189},
+			{189, 184, 187},
+			{190, 185, 188},
+			{182, 185, 190},
+			{186, 191, 188},
+			{191, 186, 189},
+			{188, 191, 190},
+			{187, 192, 189},
+			{191, 189, 192},
+			{187, 183, 193},
+			{183, 167, 193},
+			{194, 192, 187},
+			{193, 194, 187},
+			{195, 191, 192},
+			{196, 192, 194},
+			{192, 196, 195},
+			{197, 191, 195},
+			{190, 191, 197},
+			{196, 198, 195},
+			{198, 197, 195},
+			{199, 190, 197},
+			{190, 199, 182},
+			{197, 198, 200},
+			{201, 182, 199},
+			{182, 201, 179},
+			{197, 202, 199},
+			{202, 197, 200},
+			{203, 199, 202},
+			{199, 203, 201},
+			{202, 200, 203},
+			{204, 179, 201},
+			{201, 203, 204},
+			{179, 204, 138},
+			{205, 138, 204},
+			{204, 203, 205},
+			{138, 205, 136},
+			{200, 206, 203},
+			{203, 207, 205},
+			{206, 207, 203},
+			{136, 205, 208},
+			{205, 207, 209},
+			{205, 209, 208},
+			{206, 210, 207},
+			{211, 136, 208},
+			{120, 136, 211},
+			{120, 211, 117},
+			{208, 209, 212},
+			{211, 208, 212},
+			{207, 213, 209},
+			{213, 207, 210},
+			{214, 212, 209},
+			{214, 209, 213},
+			{211, 212, 215},
+			{212, 214, 216},
+			{216, 215, 212},
+			{211, 215, 217},
+			{117, 211, 217},
+			{216, 218, 215},
+			{215, 218, 217},
+			{117, 217, 219},
+			{218, 219, 217},
+			{219, 111, 117},
+			{111, 219, 112},
+			{219, 218, 220},
+			{219, 221, 112},
+			{220, 221, 219},
+			{218, 222, 220},
+			{221, 220, 222},
+			{222, 218, 216},
+			{112, 221, 223},
+			{223, 105, 112},
+			{223, 104, 105},
+			{224, 221, 222},
+			{223, 225, 104},
+			{104, 225, 66},
+			{226, 223, 221},
+			{226, 225, 223},
+			{226, 221, 224},
+			{66, 225, 67},
+			{226, 224, 227},
+			{227, 225, 226},
+			{222, 228, 224},
+			{224, 228, 227},
+			{216, 228, 222},
+			{216, 214, 228},
+			{214, 213, 228},
+			{67, 225, 229},
+			{229, 225, 227},
+			{229, 69, 67},
+			{229, 73, 69},
+			{74, 73, 229},
+			{227, 228, 230},
+			{229, 227, 230},
+			{213, 230, 228},
+			{231, 74, 229},
+			{230, 231, 229},
+			{232, 230, 213},
+			{230, 232, 231},
+			{210, 232, 213},
+			{74, 231, 233},
+			{233, 231, 232},
+			{233, 82, 74},
+			{233, 83, 82},
+			{234, 232, 210},
+			{232, 234, 233},
+			{235, 83, 233},
+			{235, 233, 234},
+			{234, 210, 236},
+			{236, 235, 234},
+			{210, 206, 236},
+			{206, 200, 236},
+			{235, 237, 83},
+			{86, 83, 237},
+			{237, 90, 86},
+			{90, 237, 238},
+			{238, 8, 90},
+			{237, 235, 239},
+			{8, 238, 240},
+			{240, 7, 8},
+			{241, 238, 237},
+			{240, 238, 241},
+			{237, 239, 241},
+			{239, 235, 242},
+			{235, 236, 242},
+			{239, 242, 243},
+			{242, 236, 243},
+			{241, 239, 243},
+			{236, 244, 243},
+			{241, 243, 245},
+			{244, 236, 246},
+			{246, 236, 200},
+			{245, 243, 247},
+			{247, 243, 244},
+			{245, 248, 241},
+			{245, 247, 248},
+			{240, 241, 248},
+			{247, 244, 249},
+			{246, 250, 244},
+			{244, 250, 249},
+			{200, 251, 246},
+			{250, 246, 251},
+			{251, 200, 198},
+			{251, 198, 196},
+			{252, 250, 251},
+			{252, 249, 250},
+			{196, 253, 251},
+			{251, 253, 252},
+			{249, 252, 253},
+			{253, 196, 194},
+			{254, 249, 253},
+			{253, 194, 254},
+			{249, 255, 247},
+			{249, 254, 255},
+			{255, 248, 247},
+			{256, 254, 194},
+			{194, 193, 256},
+			{255, 254, 257},
+			{256, 257, 254},
+			{256, 193, 258},
+			{257, 256, 258},
+			{258, 193, 259},
+			{167, 259, 193},
+			{260, 258, 259},
+			{259, 167, 261},
+			{260, 259, 261},
+			{167, 166, 261},
+			{166, 165, 261},
+			{262, 258, 260},
+			{262, 257, 258},
+			{260, 261, 263},
+			{261, 165, 264},
+			{261, 264, 263},
+			{264, 165, 162},
+			{147, 264, 162},
+			{148, 264, 147},
+			{146, 264, 148},
+			{264, 146, 265},
+			{146, 129, 265},
+			{129, 128, 265},
+			{126, 265, 128},
+			{266, 264, 265},
+			{265, 126, 266},
+			{266, 263, 264},
+			{126, 124, 123},
+			{123, 266, 126},
+			{266, 267, 263},
+			{267, 266, 123},
+			{263, 267, 260},
+			{267, 123, 122},
+			{267, 262, 260},
+			{122, 268, 267},
+			{262, 267, 268},
+			{268, 122, 119},
+			{268, 119, 116},
+			{268, 116, 115},
+			{268, 269, 262},
+			{115, 269, 268},
+			{257, 262, 269},
+			{270, 269, 115},
+			{257, 269, 270},
+			{270, 115, 110},
+			{270, 271, 257},
+			{270, 110, 271},
+			{271, 255, 257},
+			{271, 248, 255},
+			{110, 272, 271},
+			{248, 271, 272},
+			{272, 110, 107},
+			{107, 103, 272},
+			{272, 273, 248},
+			{273, 272, 103},
+			{240, 248, 273},
+			{100, 273, 103},
+			{273, 274, 240},
+			{273, 100, 274},
+			{240, 274, 7},
+			{2, 274, 100},
+			{0, 7, 274},
+			{274, 2, 0}
+		});
+
+		TArray<uint16> Materials;
+		for (int32 i = 0; i < Indices.Num(); ++i)
+		{
+			Materials.Emplace(0);
+		}
+
+		TUniquePtr<FTriangleMeshImplicitObject> TriangleMesh = MakeUnique<FTriangleMeshImplicitObject>(MoveTemp(TrimeshParticles), MoveTemp(Indices), MoveTemp(Materials));
+		TImplicitObjectScaled<FTriangleMeshImplicitObject> ScaledTriangleMesh = TImplicitObjectScaled<FTriangleMeshImplicitObject>(MakeSerializable(TriangleMesh), FVec3(50,50,50));
+
+		const TVector<FReal, 3> X1 = { 0,0,-19.45 };
+		const TVector<FReal, 3> X2 = X1 + TVector<FReal, 3>(0, 0, 38.9);
+		const FReal Radius = 25.895;
+		const TCapsule<FReal> Capsule = TCapsule<FReal>(X1, X2, Radius);
+
+		const FVec3 CapsuleToTrimeshTranslation = { 1818.55884, 27.8377075, -630.160645 };
+		const TRigidTransform<FReal, 3> CapsuleToTrimesh(CapsuleToTrimeshTranslation, FQuat::Identity);
+
+		const FVec3 TrimeshTranslation = { -1040.00000, 700.000000, 992.000000 };
+		const TRigidTransform<FReal, 3> TrimeshTransform(TrimeshTranslation, FQuat::Identity);
+
+		const FVec3 Dir(0, 0, -1);
+		const FReal Length = 159.100098;
+
+		FReal OutTime = -1;
+		FVec3 Normal(0.0f);
+		FVec3 Position(0.0f);
+		int32 FaceIndex = -1;
+		bool bResult = ScaledTriangleMesh.LowLevelSweepGeom(Capsule, CapsuleToTrimesh, Dir, Length, OutTime, Position, Normal, FaceIndex, 0.0f, true);
+		FVec3 WorldPosition = TrimeshTransform.TransformPositionNoScale(Position);
+
+		EXPECT_EQ(bResult, true);
+		EXPECT_EQ(FaceIndex, 158);
+		EXPECT_NEAR(WorldPosition.X, 763.85413, KINDA_SMALL_NUMBER);
+		EXPECT_NEAR(WorldPosition.Y, 728.80212, KINDA_SMALL_NUMBER);
+		EXPECT_NEAR(WorldPosition.Z, 303.77856, KINDA_SMALL_NUMBER);
+	}
+
+	template void CapsuleSweepAgainstTriMeshReal<float>();
+}
\ No newline at end of file
diff --git a/Engine/Source/Programs/HeadlessChaos/Public/HeadlessChaosTestSweep.h b/Engine/Source/Programs/HeadlessChaos/Public/HeadlessChaosTestSweep.h
new file mode 100644
--- /dev/null
+++ b/Engine/Source/Programs/HeadlessChaos/Public/HeadlessChaosTestSweep.h
@@ -0,0 +1,13 @@
+// Copyright Epic Games, Inc. All Rights Reserved.
+
+
+#pragma once
+
+#include "HeadlessChaosTestUtility.h"
+
+namespace ChaosTest
+{
+
+	template<typename T>
+	void CapsuleSweepAgainstTriMeshReal();
+}
\ No newline at end of file
