Fix runtime crash in PBDCollisionConstraints... Removed check(false) for condition which is actually valid for destruction. re-added early-out when a shape has no material assigned. Updated tests to have empty unique material arrays.

diff --git a/Engine/Source/Programs/HeadlessChaos/Private/HeadlessChaosTestCollisions.cpp b/Engine/Source/Programs/HeadlessChaos/Private/HeadlessChaosTestCollisions.cpp
--- a/Engine/Source/Programs/HeadlessChaos/Private/HeadlessChaosTestCollisions.cpp
+++ b/Engine/Source/Programs/HeadlessChaos/Private/HeadlessChaosTestCollisions.cpp
@@ -36,10 +36,12 @@ namespace ChaosTest {
 		PhysicsMaterial->Friction = (T)0;
 		PhysicsMaterial->Restitution = (T)0;
 		TArrayCollectionArray<TSerializablePtr<FChaosPhysicsMaterial>> PhysicsMaterials;
+		TArrayCollectionArray<TUniquePtr<FChaosPhysicsMaterial>> PerParticlePhysicsMaterials;
 		
 		TPBDRigidsSOAs<T, 3> Particles;
 		Particles.GetParticleHandles().AddArray(&Collided);
 		Particles.GetParticleHandles().AddArray(&PhysicsMaterials);
+		Particles.GetParticleHandles().AddArray(&PerParticlePhysicsMaterials);
 
 		auto Box1 = AppendDynamicParticleBox<T>(Particles);
 		Box1->X() = TVector<T, 3>(1.f);
@@ -55,7 +57,7 @@ namespace ChaosTest {
 		Box2->Q() = Box2->R();
 		Box2->AuxilaryValue(PhysicsMaterials) = MakeSerializable(PhysicsMaterial);
 
-		FPBDCollisionConstraintAccessor Collisions(Particles, Collided, PhysicsMaterials, 1, 1, T(0));
+		FPBDCollisionConstraintAccessor Collisions(Particles, Collided, PhysicsMaterials, PerParticlePhysicsMaterials, 1, 1, T(0));
 		Collisions.ComputeConstraints(0.f);
 		EXPECT_EQ(Collisions.NumConstraints(), 1);
 
@@ -83,10 +85,12 @@ namespace ChaosTest {
 		PhysicsMaterial->Friction = (T)0;
 		PhysicsMaterial->Restitution = (T)0;
 		TArrayCollectionArray<TSerializablePtr<FChaosPhysicsMaterial>> PhysicsMaterials;
+		TArrayCollectionArray<TUniquePtr<FChaosPhysicsMaterial>> PerParticlePhysicsMaterials;
 
 		TPBDRigidsSOAs<T, 3> Particles;
 		Particles.GetParticleHandles().AddArray(&Collided);
 		Particles.GetParticleHandles().AddArray(&PhysicsMaterials);
+		Particles.GetParticleHandles().AddArray(&PerParticlePhysicsMaterials);
 
 		auto Box1 = AppendDynamicParticleConvexBox<T>(Particles, TVector<T, 3>(1.f) );
 		Box1->X() = TVector<T, 3>(0.f);
@@ -102,7 +106,7 @@ namespace ChaosTest {
 		Box2->Q() = Box2->R();
 		Box2->AuxilaryValue(PhysicsMaterials) = MakeSerializable(PhysicsMaterial);
 
-		FPBDCollisionConstraintAccessor Collisions(Particles, Collided, PhysicsMaterials, 1, 1, T(0));
+		FPBDCollisionConstraintAccessor Collisions(Particles, Collided, PhysicsMaterials, PerParticlePhysicsMaterials, 1, 1, T(0));
 		Collisions.ComputeConstraints(0.f);
 		EXPECT_EQ(Collisions.NumConstraints(), 1);
 
@@ -125,10 +129,12 @@ namespace ChaosTest {
 		PhysicsMaterial->Friction = (T)0;
 		PhysicsMaterial->Restitution = (T)1;
 		TArrayCollectionArray<TSerializablePtr<FChaosPhysicsMaterial>> PhysicsMaterials;
+		TArrayCollectionArray<TUniquePtr<FChaosPhysicsMaterial>> PerParticlePhysicsMaterials;
 
 		TPBDRigidsSOAs<T, 3> Particles;
 		Particles.GetParticleHandles().AddArray(&Collided);
 		Particles.GetParticleHandles().AddArray(&PhysicsMaterials);
+		Particles.GetParticleHandles().AddArray(&PerParticlePhysicsMaterials);
 
 		auto Floor = AppendStaticAnalyticFloor<T>(Particles);
 		auto Box = AppendDynamicParticleBox<T>(Particles);
@@ -141,7 +147,7 @@ namespace ChaosTest {
 
 		const float Dt = 1 / 24.;
 
-		FPBDCollisionConstraintAccessor Collisions(Particles, Collided, PhysicsMaterials, 2, 5, T(0));
+		FPBDCollisionConstraintAccessor Collisions(Particles, Collided, PhysicsMaterials, PerParticlePhysicsMaterials, 2, 5, T(0));
 
 		Collisions.ComputeConstraints(Dt);
 		EXPECT_EQ(Collisions.NumConstraints(), 1);
@@ -195,10 +201,12 @@ namespace ChaosTest {
 		PhysicsMaterial->Friction = (T)0;
 		PhysicsMaterial->Restitution = (T)0;
 		TArrayCollectionArray<TSerializablePtr<FChaosPhysicsMaterial>> PhysicsMaterials;
+		TArrayCollectionArray<TUniquePtr<FChaosPhysicsMaterial>> PerParticlePhysicsMaterials;
 
 		TPBDRigidsSOAs<T, 3> Particles;
 		Particles.GetParticleHandles().AddArray(&Collided);
 		Particles.GetParticleHandles().AddArray(&PhysicsMaterials);
+		Particles.GetParticleHandles().AddArray(&PerParticlePhysicsMaterials);
 
 		auto Floor = AppendStaticConvexFloor<T>(Particles);
 		auto Box = AppendDynamicParticleConvexBox<T>( Particles, TVector<T, 3>(50) );
@@ -211,7 +219,7 @@ namespace ChaosTest {
 
 		const float Dt = 1 / 24.;
 
-		FPBDCollisionConstraintAccessor Collisions(Particles, Collided, PhysicsMaterials, 2, 5, T(0));
+		FPBDCollisionConstraintAccessor Collisions(Particles, Collided, PhysicsMaterials, PerParticlePhysicsMaterials, 2, 5, T(0));
 
 		Collisions.ComputeConstraints(Dt);
 		EXPECT_EQ(Collisions.NumConstraints(), 1);
@@ -265,10 +273,12 @@ namespace ChaosTest {
 		PhysicsMaterial->Friction = (T)0;
 		PhysicsMaterial->Restitution = (T)0;
 		TArrayCollectionArray<TSerializablePtr<FChaosPhysicsMaterial>> PhysicsMaterials;
+		TArrayCollectionArray<TUniquePtr<FChaosPhysicsMaterial>> PerParticlePhysicsMaterials;
 		
 		TPBDRigidsSOAs<T, 3> Particles;
 		Particles.GetParticleHandles().AddArray(&Collided);
 		Particles.GetParticleHandles().AddArray(&PhysicsMaterials);
+		Particles.GetParticleHandles().AddArray(&PerParticlePhysicsMaterials);
 
 		auto Floor = AppendStaticAnalyticFloor<T>(Particles);
 		auto Box = AppendDynamicParticleBox<T>(Particles);
@@ -281,7 +291,7 @@ namespace ChaosTest {
 
 		const float Dt = 1 / 24.;
 
-		FPBDCollisionConstraintAccessor Collisions(Particles, Collided, PhysicsMaterials, 2, 5, T(0));
+		FPBDCollisionConstraintAccessor Collisions(Particles, Collided, PhysicsMaterials, PerParticlePhysicsMaterials, 2, 5, T(0));
 
 		Collisions.ComputeConstraints(Dt);
 		EXPECT_EQ(Collisions.NumConstraints(), 1);
@@ -334,9 +344,11 @@ namespace ChaosTest {
 		PhysicsMaterial->Friction = (T)0;
 		PhysicsMaterial->Restitution = (T)1;
 		TArrayCollectionArray<TSerializablePtr<FChaosPhysicsMaterial>> PhysicsMaterials;
+		TArrayCollectionArray<TUniquePtr<FChaosPhysicsMaterial>> PerParticlePhysicsMaterials;
 		TPBDRigidsSOAs<T, 3> Particles;
 		Particles.GetParticleHandles().AddArray(&Collided);
 		Particles.GetParticleHandles().AddArray(&PhysicsMaterials);
+		Particles.GetParticleHandles().AddArray(&PerParticlePhysicsMaterials);
 
 		auto Floor = AppendStaticAnalyticFloor<T>(Particles);
 		auto Box = AppendDynamicParticleBox<T>(Particles);
@@ -349,7 +361,7 @@ namespace ChaosTest {
 
 		const float Dt = 1 / 24.;
 
-		FPBDCollisionConstraintAccessor Collisions(Particles, Collided, PhysicsMaterials, 2, 5, T(0));
+		FPBDCollisionConstraintAccessor Collisions(Particles, Collided, PhysicsMaterials, PerParticlePhysicsMaterials, 2, 5, T(0));
 
 		Collisions.ComputeConstraints(Dt);
 		EXPECT_EQ(Collisions.NumConstraints(), 1);
@@ -405,9 +417,11 @@ namespace ChaosTest {
 		PhysicsMaterial->Friction = (T)0;
 		PhysicsMaterial->Restitution = (T)0;
 		TArrayCollectionArray<TSerializablePtr<FChaosPhysicsMaterial>> PhysicsMaterials;
+		TArrayCollectionArray<TUniquePtr<FChaosPhysicsMaterial>> PerParticlePhysicsMaterials;
 		TPBDRigidsSOAs<T, 3> Particles;
 		Particles.GetParticleHandles().AddArray(&Collided);
 		Particles.GetParticleHandles().AddArray(&PhysicsMaterials);
+		Particles.GetParticleHandles().AddArray(&PerParticlePhysicsMaterials);
 
 		auto StaticBox = AppendStaticParticleBox<T>(Particles);
 		StaticBox->X() = TVector<T, 3>(-0.05f, -0.05f, -0.1f);
@@ -425,7 +439,7 @@ namespace ChaosTest {
 
 		float Dt = 1 / 24.;
 
-		FPBDCollisionConstraintAccessor Collisions(Particles, Collided, PhysicsMaterials, 1, 1, T(0));
+		FPBDCollisionConstraintAccessor Collisions(Particles, Collided, PhysicsMaterials, PerParticlePhysicsMaterials, 1, 1, T(0));
 		Collisions.ComputeConstraints(Dt);
 		EXPECT_EQ(Collisions.NumConstraints(), 1);
 
@@ -492,6 +506,7 @@ namespace ChaosTest {
 		PhysicsMaterial->Friction = (T)0;
 		PhysicsMaterial->Restitution = (T)0;
 		TArrayCollectionArray<TSerializablePtr<FChaosPhysicsMaterial>> PhysicsMaterials;
+		TArrayCollectionArray<TUniquePtr<FChaosPhysicsMaterial>> PerParticlePhysicsMaterials;
 		Chaos::TPBDRigidParticles<T, 3> Particles;
 		Particles.AddArray(&Collided);
 		Particles.AddArray(&PhysicsMaterials);
@@ -548,6 +563,7 @@ namespace ChaosTest {
 		PhysicsMaterial->Friction = (T)0;
 		PhysicsMaterial->Restitution = (T)0;
 		TArrayCollectionArray<TSerializablePtr<FChaosPhysicsMaterial>> PhysicsMaterials;
+		TArrayCollectionArray<TUniquePtr<FChaosPhysicsMaterial>> PerParticlePhysicsMaterials;
 		Chaos::TPBDRigidParticles<T, 3> Particles;
 		Particles.AddArray(&Collided);
 		Particles.AddArray(&PhysicsMaterials);
diff --git a/Engine/Source/Programs/HeadlessChaos/Private/HeadlessChaosTestMinEvolution.cpp b/Engine/Source/Programs/HeadlessChaos/Private/HeadlessChaosTestMinEvolution.cpp
--- a/Engine/Source/Programs/HeadlessChaos/Private/HeadlessChaosTestMinEvolution.cpp
+++ b/Engine/Source/Programs/HeadlessChaos/Private/HeadlessChaosTestMinEvolution.cpp
@@ -40,7 +40,7 @@ namespace ChaosTest
 		TArrayCollectionArray<TUniquePtr<Chaos::FChaosPhysicsMaterial>> PerParticleMaterials;
 		TArrayCollectionArray<FVec3> ParticlePrevXs;
 		TArrayCollectionArray<FRotation3> ParticlePrevRs;
-		FCollisionConstraints Collisions(ParticlesContainer, CollidedParticles, ParticleMaterials);
+		FCollisionConstraints Collisions(ParticlesContainer, CollidedParticles, ParticleMaterials, PerParticleMaterials);
 		FParticlePairBroadPhase BroadPhase(ActivePotentiallyCollidingPairs, 0);
 		FNarrowPhase NarrowPhase;
 		FCollisionDetector CollisionDetector(BroadPhase, NarrowPhase, Collisions);
diff --git a/Engine/Source/Programs/HeadlessChaos/Public/HeadlessChaosCollisionConstraints.h b/Engine/Source/Programs/HeadlessChaos/Public/HeadlessChaosCollisionConstraints.h
--- a/Engine/Source/Programs/HeadlessChaos/Public/HeadlessChaosCollisionConstraints.h
+++ b/Engine/Source/Programs/HeadlessChaos/Public/HeadlessChaosCollisionConstraints.h
@@ -37,15 +37,15 @@ public:
 	FPBDCollisionConstraintAccessor()
 		: SpatialAcceleration(EmptyParticles.GetNonDisabledView())
 		, BroadPhase(EmptyParticles, (FReal)1, (FReal)0)
-		, CollisionConstraints(EmptyParticles, EmptyCollided, EmptyPhysicsMaterials, 1, 1)
+		, CollisionConstraints(EmptyParticles, EmptyCollided, EmptyPhysicsMaterials, EmptyUniquePhysicsMaterials, 1, 1)
 		, CollisionDetector(BroadPhase, NarrowPhase, CollisionConstraints)
 	{}
 
-	FPBDCollisionConstraintAccessor(const TPBDRigidsSOAs<FReal, 3>& InParticles, TArrayCollectionArray<bool>& Collided, const TArrayCollectionArray<TSerializablePtr<FChaosPhysicsMaterial>>& PerParticleMaterials,
+	FPBDCollisionConstraintAccessor(const TPBDRigidsSOAs<FReal, 3>& InParticles, TArrayCollectionArray<bool>& Collided, const TArrayCollectionArray<TSerializablePtr<FChaosPhysicsMaterial>>& PerParticleMaterials, const TArrayCollectionArray<TUniquePtr<FChaosPhysicsMaterial>>& PerParticleUniqueMaterials,
 	const int32 PushOutIterations, const int32 PushOutPairIterations, const FReal Thickness) 
 		: SpatialAcceleration(InParticles.GetNonDisabledView())
 		, BroadPhase(InParticles, (FReal)1, (FReal)0)
-		, CollisionConstraints(InParticles, Collided, PerParticleMaterials, 1, 1, Thickness)
+		, CollisionConstraints(InParticles, Collided, PerParticleMaterials, PerParticleUniqueMaterials, 1, 1, Thickness)
 		, CollisionDetector(BroadPhase, NarrowPhase, CollisionConstraints)
 	{}
 
@@ -122,6 +122,7 @@ public:
 	TPBDRigidsSOAs<FReal, 3> EmptyParticles;
 	TArrayCollectionArray<bool> EmptyCollided;
 	TArrayCollectionArray<TSerializablePtr<FChaosPhysicsMaterial>> EmptyPhysicsMaterials;
+	TArrayCollectionArray<TUniquePtr<FChaosPhysicsMaterial>> EmptyUniquePhysicsMaterials;
 
 	FAccelerationStructure SpatialAcceleration;
 	FSpatialAccelerationBroadPhase BroadPhase;
diff --git a/Engine/Source/Runtime/Experimental/Chaos/Private/Chaos/PBDCollisionConstraints.cpp b/Engine/Source/Runtime/Experimental/Chaos/Private/Chaos/PBDCollisionConstraints.cpp
--- a/Engine/Source/Runtime/Experimental/Chaos/Private/Chaos/PBDCollisionConstraints.cpp
+++ b/Engine/Source/Runtime/Experimental/Chaos/Private/Chaos/PBDCollisionConstraints.cpp
@@ -171,11 +171,16 @@ namespace Chaos
 				{
 					return ShapeData->GetMaterials()[0].Get();
 				}
+				else
+				{
+					// This shape doesn't have a material assigned
+					return nullptr;
+				}
 			}
 		}
 
-		check(false); // The geometry used for this particle does not belong to the particle..
-
+		// The geometry used for this particle does not belong to the particle.
+		// This can happen in the case of fracture.
 		return nullptr;
 	}
 
