Fix some edge cases in resim of slave objects that are sleeping

diff --git a/Engine/Source/Programs/HeadlessChaos/Private/HeadlessChaosTestRewind.cpp b/Engine/Source/Programs/HeadlessChaos/Private/HeadlessChaosTestRewind.cpp
--- a/Engine/Source/Programs/HeadlessChaos/Private/HeadlessChaosTestRewind.cpp
+++ b/Engine/Source/Programs/HeadlessChaos/Private/HeadlessChaosTestRewind.cpp
@@ -547,6 +547,7 @@ namespace ChaosTest {
 		auto* Solver = Module->CreateSolver<TypeParam>(nullptr,ESolverFlags::Standalone);
 		Solver->SetEnabled(true);
 
+		//note: this 5 is just a suggestion, there could be more frames saved than that
 		Solver->EnableRewindCapture(5);
 
 
@@ -569,7 +570,7 @@ namespace ChaosTest {
 		Particle->SetGravityEnabled(false);
 		Particle->SetV(FVec3(0));
 
-		for(int Step = 0; Step < 10; ++Step)
+		for(int Step = 0; Step < 11; ++Step)
 		{
 			TickSolverHelper(Module,Solver);
 		}
@@ -1162,4 +1163,145 @@ namespace ChaosTest {
 
 		Module->DestroySolver(Solver);
 	}
+
+	TYPED_TEST(AllTraits,RewindTest_ResimAsSlaveWokenUp)
+	{
+		auto Box = TSharedPtr<FImplicitObject,ESPMode::ThreadSafe>(new TBox<FReal,3>(FVec3(-10,-10,-10),FVec3(10,10,10)));
+
+		FChaosSolversModule* Module = FChaosSolversModule::GetModule();
+		Module->ChangeThreadingMode(EChaosThreadingMode::SingleThread);
+
+		// Make a solver
+		auto* Solver = Module->CreateSolver<TypeParam>(nullptr,ESolverFlags::Standalone);
+		Solver->SetEnabled(true);
+		Solver->EnableRewindCapture(7);
+
+		// Make particles
+		auto ImpulsedObj = TPBDRigidParticle<float,3>::CreateParticle();
+		auto HitObj = TPBDRigidParticle<float,3>::CreateParticle();
+
+		ImpulsedObj->SetGeometry(Box);
+		ImpulsedObj->SetGravityEnabled(false);
+		Solver->RegisterObject(ImpulsedObj.Get());
+
+		HitObj->SetGeometry(Box);
+		HitObj->SetGravityEnabled(false);
+		Solver->RegisterObject(HitObj.Get());
+
+		ImpulsedObj->SetX(FVec3(0,0,20));
+		ImpulsedObj->SetM(1);
+		ImpulsedObj->SetInvM(1);
+		ImpulsedObj->SetResimType(EResimType::ResimAsSlave);
+		ImpulsedObj->SetObjectState(EObjectStateType::Sleeping);
+
+		HitObj->SetX(FVec3(0,0,0));
+		HitObj->SetM(1);
+		HitObj->SetInvM(1);
+		HitObj->SetResimType(EResimType::ResimAsSlave);
+		HitObj->SetObjectState(EObjectStateType::Sleeping);
+
+
+		ChaosTest::SetParticleSimDataToCollide({ImpulsedObj.Get(),HitObj.Get()});
+
+		const int32 ApplyImpulseStep = 8;
+		const int32 LastStep = 11;
+
+		TArray<FVec3> Xs;
+
+		for(int Step = 0; Step <= LastStep; ++Step)
+		{
+			if(ApplyImpulseStep == Step)
+			{
+				ImpulsedObj->SetLinearImpulse(FVec3(0,0,-10));
+			}
+
+			TickSolverHelper(Module,Solver);
+			Xs.Add(HitObj->X());
+		}
+
+		const int RewindStep = 5;
+
+		FRewindData* RewindData = Solver->GetRewindData();
+		EXPECT_TRUE(RewindData->RewindToFrame(RewindStep));
+
+		for(int Step = RewindStep; Step <= LastStep; ++Step)
+		{
+			TickSolverHelper(Module,Solver);
+
+			EXPECT_VECTOR_FLOAT_EQ(HitObj->X(),Xs[Step]);
+		}
+
+		Module->DestroySolver(Solver);
+	}
+
+	TYPED_TEST(AllTraits,RewindTest_ResimAsSlaveWokenUpNoHistory)
+	{
+		auto Box = TSharedPtr<FImplicitObject,ESPMode::ThreadSafe>(new TBox<FReal,3>(FVec3(-10,-10,-10),FVec3(10,10,10)));
+
+		FChaosSolversModule* Module = FChaosSolversModule::GetModule();
+		Module->ChangeThreadingMode(EChaosThreadingMode::SingleThread);
+
+		// Make a solver
+		auto* Solver = Module->CreateSolver<TypeParam>(nullptr,ESolverFlags::Standalone);
+		Solver->SetEnabled(true);
+		Solver->EnableRewindCapture(7);
+
+		// Make particles
+		auto ImpulsedObj = TPBDRigidParticle<float,3>::CreateParticle();
+		auto HitObj = TPBDRigidParticle<float,3>::CreateParticle();
+
+		ImpulsedObj->SetGeometry(Box);
+		ImpulsedObj->SetGravityEnabled(false);
+		Solver->RegisterObject(ImpulsedObj.Get());
+
+		HitObj->SetGeometry(Box);
+		HitObj->SetGravityEnabled(false);
+		Solver->RegisterObject(HitObj.Get());
+
+		ImpulsedObj->SetX(FVec3(0,0,20));
+		ImpulsedObj->SetM(1);
+		ImpulsedObj->SetInvM(1);
+		ImpulsedObj->SetObjectState(EObjectStateType::Sleeping);
+
+		HitObj->SetX(FVec3(0,0,0));
+		HitObj->SetM(1);
+		HitObj->SetInvM(1);
+		HitObj->SetResimType(EResimType::ResimAsSlave);
+		HitObj->SetObjectState(EObjectStateType::Sleeping);
+
+
+		ChaosTest::SetParticleSimDataToCollide({ImpulsedObj.Get(),HitObj.Get()});
+
+		const int32 ApplyImpulseStep = 97;
+		const int32 LastStep = 100;
+
+		TArray<FVec3> Xs;
+
+		for(int Step = 0; Step <= LastStep; ++Step)
+		{
+			TickSolverHelper(Module,Solver);
+			Xs.Add(HitObj->X());	//not a full re-sim so we should end up with exact same result
+		}
+
+		const int RewindStep = 95;
+
+		FRewindData* RewindData = Solver->GetRewindData();
+		EXPECT_TRUE(RewindData->RewindToFrame(RewindStep));
+
+		for(int Step = RewindStep; Step <= LastStep; ++Step)
+		{
+			//during resim apply correction impulse
+			if(ApplyImpulseStep == Step)
+			{
+				ImpulsedObj->SetLinearImpulse(FVec3(0,0,-10));
+			}
+
+			TickSolverHelper(Module,Solver);
+
+			//even though there's now a different collision in the sim, the final result of slave is the same as before
+			EXPECT_VECTOR_FLOAT_EQ(HitObj->X(),Xs[Step]);
+		}
+
+		Module->DestroySolver(Solver);
+	}
 }
diff --git a/Engine/Source/Runtime/Experimental/ChaosSolvers/Public/RewindData.h b/Engine/Source/Runtime/Experimental/ChaosSolvers/Public/RewindData.h
--- a/Engine/Source/Runtime/Experimental/ChaosSolvers/Public/RewindData.h
+++ b/Engine/Source/Runtime/Experimental/ChaosSolvers/Public/RewindData.h
@@ -722,7 +722,7 @@ public:
 	FRewindData(int32 NumFrames)
 	: Managers(NumFrames+1)	//give 1 extra for saving at head
 	, CurFrame(0)
-	, LatestFrame(0)
+	, LatestFrame(-1)
 	, CurWave(1)
 	, FramesSaved(0)
 	, DataIdxOffset(0)
@@ -868,8 +868,6 @@ public:
 		QUICK_SCOPE_CYCLE_COUNTER(RewindDataAdvance);
 		Managers[CurFrame].DeltaTime = DeltaTime;
 
-		++CurFrame;
-		LatestFrame = FMath::Max(LatestFrame,CurFrame);
 		FramesSaved = FMath::Min(FramesSaved+1,static_cast<int32>(Managers.Capacity()));
 		
 		const int32 EarliestFrame = CurFrame - 1 - FramesSaved;
@@ -891,9 +889,9 @@ public:
 
 				//During a resim it's possible the user will not dirty a particle that was previously dirty.
 				//If this happens we need to mark the particle as desynced
-				if(ResimType == EResimType::FullResim && !Info.bDesync && Info.GTDirtyOnFrame[CurFrame-1].MissingWrite(CurFrame-1, CurWave))
+				if(ResimType == EResimType::FullResim && !Info.bDesync && Info.GTDirtyOnFrame[CurFrame].MissingWrite(CurFrame, CurWave))
 				{
-					Info.Desync(CurFrame-1, LatestFrame);
+					Info.Desync(CurFrame, LatestFrame);
 				}
 
 				
@@ -904,6 +902,9 @@ public:
 	void FinishFrame()
 	{
 		QUICK_SCOPE_CYCLE_COUNTER(RewindDataFinishFrame);
+		LatestFrame = FMath::Max(LatestFrame,CurFrame);
+		++CurFrame;
+
 		if(IsResim())
 		{
 			ensure(IsInGameThread());
@@ -930,7 +931,7 @@ public:
 
 	bool IsResim() const
 	{
-		return CurFrame < LatestFrame;
+		return CurFrame <= LatestFrame;
 	}
 
 	//Number of particles that we're currently storing history for
@@ -954,9 +955,8 @@ public:
 		bNeedsSave = true;
 
 		//If manager already exists for previous frame, use it
-		const int32 PrevFrame = CurFrame - 1;
-		FFrameManagerInfo& Info = Managers[PrevFrame];
-		ensure(Info.Manager && Info.FrameCreatedFor == (PrevFrame));
+		FFrameManagerInfo& Info = Managers[CurFrame];
+		ensure(Info.Manager && Info.FrameCreatedFor == (CurFrame));
 
 		DataIdxOffset = Info.Manager->GetNumParticles();
 		Info.Manager->SetNumParticles(DataIdxOffset + NumActiveParticles);
@@ -1063,28 +1063,36 @@ public:
 	{
 		const int32 DestDataIdx = SrcDataIdx + DataIdxOffset;
 
-		if(bResim && Rigid.ResimType() == EResimType::ResimAsSlave)
+		//during resim only full resim objects should modify future
+		if(bResim && Rigid.ResimType() != EResimType::FullResim)
 		{
-			//resim not allowed for this particle so don't modify
+			if(FindParticle(Rigid.UniqueIdx()) == nullptr)
+			{
+				//no history but collision moved/woke us up so snap back manually (if history exists we'll snap in FinishFrame
+				Rigid.SetP(Rigid.X());
+				Rigid.SetQ(Rigid.R());
+				Rigid.SetV(Rigid.PreV());
+				Rigid.SetW(Rigid.PreW());
+			}
+			
+			return;
 		}
-		else
+
+		//todo: is this check needed? why do we pass sleeping rigids into this function?
+		if(SimWritablePropsMayChange(Rigid))
 		{
-			//todo: is this check needed? why do we pass sleeping rigids into this function?
-			if(SimWritablePropsMayChange(Rigid))
-			{
-				FDirtyParticleInfo& Info = FindOrAddParticle(Rigid);
-				Info.LastDirtyFrame = CurFrame-1;
+			FDirtyParticleInfo& Info = FindOrAddParticle(Rigid);
+			Info.LastDirtyFrame = CurFrame;
 
-				//User called PrepareManagerForFrame (or PrepareFrameForPTDirty) for the previous frame, so use it
-				FDirtyPropData DestManagerWrapper(Managers[CurFrame-1].Manager.Get(), DestDataIdx);
+			//User called PrepareManagerForFrame (or PrepareFrameForPTDirty) for the previous frame, so use it
+			FDirtyPropData DestManagerWrapper(Managers[CurFrame].Manager.Get(),DestDataIdx);
 
-				//sim-writable properties changed at head, so we must write down what they were
-				FGeometryParticleStateBase& LatestState = Info.AddFrame(CurFrame-1);
-				LatestState.SyncSimWritablePropsFromSim(DestManagerWrapper,Rigid);
+			//sim-writable properties changed at head, so we must write down what they were
+			FGeometryParticleStateBase& LatestState = Info.AddFrame(CurFrame);
+			LatestState.SyncSimWritablePropsFromSim(DestManagerWrapper,Rigid);
 
-				//update any previous frames that were pointing at head
-				CoalesceBack(Info.Frames, CurFrame-1);
-			}
+			//update any previous frames that were pointing at head
+			CoalesceBack(Info.Frames,CurFrame);
 		}
 	}
 
