Fix another edge case with resim as slave

diff --git a/Engine/Source/Programs/HeadlessChaos/Private/HeadlessChaosTestRewind.cpp b/Engine/Source/Programs/HeadlessChaos/Private/HeadlessChaosTestRewind.cpp
--- a/Engine/Source/Programs/HeadlessChaos/Private/HeadlessChaosTestRewind.cpp
+++ b/Engine/Source/Programs/HeadlessChaos/Private/HeadlessChaosTestRewind.cpp
@@ -570,7 +570,7 @@ namespace ChaosTest {
 		Particle->SetGravityEnabled(false);
 		Particle->SetV(FVec3(0));
 
-		for(int Step = 0; Step < 11; ++Step)
+		for(int Step = 0; Step < 40; ++Step)
 		{
 			TickSolverHelper(Module,Solver);
 		}
@@ -1101,6 +1101,138 @@ namespace ChaosTest {
 		Module->DestroySolver(Solver);
 	}
 
+	TYPED_TEST(AllTraits,RewindTest_FullResimFallSeeCollisionCorrection)
+	{
+		auto Sphere = TSharedPtr<FImplicitObject,ESPMode::ThreadSafe>(new TSphere<FReal,3>(TVector<float,3>(0),10));
+		auto Box = TSharedPtr<FImplicitObject,ESPMode::ThreadSafe>(new TBox<FReal,3>(FVec3(-100,-100,-100),FVec3(100,100,0)));
+
+		FChaosSolversModule* Module = FChaosSolversModule::GetModule();
+		Module->ChangeThreadingMode(EChaosThreadingMode::SingleThread);
+
+		// Make a solver
+		auto* Solver = Module->CreateSolver<TypeParam>(nullptr,ESolverFlags::Standalone);
+		Solver->SetEnabled(true);
+		Solver->EnableRewindCapture(100);
+
+		// Make particles
+		auto Dynamic = TPBDRigidParticle<float,3>::CreateParticle();
+		auto Kinematic = TKinematicGeometryParticle<float,3>::CreateParticle();
+
+		Dynamic->SetGeometry(Sphere);
+		Dynamic->SetGravityEnabled(true);
+		Solver->RegisterObject(Dynamic.Get());
+
+		Kinematic->SetGeometry(Box);
+		Solver->RegisterObject(Kinematic.Get());
+
+		Dynamic->SetX(FVec3(0,0,17));
+		Dynamic->SetGravityEnabled(false);
+		Dynamic->SetV(FVec3(0,0,-1));
+		Dynamic->SetObjectState(EObjectStateType::Dynamic);
+
+		Kinematic->SetX(FVec3(0,0,-1000));
+
+		ChaosTest::SetParticleSimDataToCollide({Dynamic.Get(),Kinematic.Get()});
+
+		const int32 LastStep = 11;
+
+		TArray<FVec3> Xs;
+
+		for(int Step = 0; Step <= LastStep; ++Step)
+		{
+			TickSolverHelper(Module,Solver);
+			Xs.Add(Dynamic->X());
+		}
+
+		EXPECT_FLOAT_EQ(Dynamic->X()[2],5);
+
+		const int RewindStep = 0;
+
+		FRewindData* RewindData = Solver->GetRewindData();
+		EXPECT_TRUE(RewindData->RewindToFrame(RewindStep));
+
+		//force collision
+		Kinematic->SetX(FVec3(0,0,0));
+
+		for(int Step = RewindStep; Step <= LastStep; ++Step)
+		{
+			//Resim ignores collision since it's ResimAsSlave
+			TickSolverHelper(Module,Solver);
+			EXPECT_GE(Dynamic->X()[2],10);
+		}
+
+		EXPECT_FLOAT_EQ(Dynamic->X()[2],10);
+
+		Module->DestroySolver(Solver);
+	}
+
+	TYPED_TEST(AllTraits,RewindTest_ResimAsSlaveFallIgnoreCollision)
+	{
+		auto Sphere = TSharedPtr<FImplicitObject,ESPMode::ThreadSafe>(new TSphere<FReal,3>(TVector<float,3>(0),10));
+		auto Box = TSharedPtr<FImplicitObject,ESPMode::ThreadSafe>(new TBox<FReal,3>(FVec3(-100,-100,-100),FVec3(100,100,0)));
+
+		FChaosSolversModule* Module = FChaosSolversModule::GetModule();
+		Module->ChangeThreadingMode(EChaosThreadingMode::SingleThread);
+
+		// Make a solver
+		auto* Solver = Module->CreateSolver<TypeParam>(nullptr,ESolverFlags::Standalone);
+		Solver->SetEnabled(true);
+		Solver->EnableRewindCapture(100);
+
+		// Make particles
+		auto Dynamic = TPBDRigidParticle<float,3>::CreateParticle();
+		auto Kinematic = TKinematicGeometryParticle<float,3>::CreateParticle();
+
+		Dynamic->SetGeometry(Sphere);
+		Dynamic->SetGravityEnabled(true);
+		Solver->RegisterObject(Dynamic.Get());
+
+		Kinematic->SetGeometry(Box);
+		Solver->RegisterObject(Kinematic.Get());
+
+		Dynamic->SetX(FVec3(0,0,17));
+		Dynamic->SetGravityEnabled(false);
+		Dynamic->SetV(FVec3(0,0,-1));
+		Dynamic->SetObjectState(EObjectStateType::Dynamic);
+		Dynamic->SetResimType(EResimType::ResimAsSlave);
+
+		Kinematic->SetX(FVec3(0,0,-1000));
+
+		ChaosTest::SetParticleSimDataToCollide({Dynamic.Get(),Kinematic.Get()});
+
+		const int32 LastStep = 11;
+
+		TArray<FVec3> Xs;
+
+		for(int Step = 0; Step <= LastStep; ++Step)
+		{
+			TickSolverHelper(Module,Solver);
+			Xs.Add(Dynamic->X());
+		}
+
+		EXPECT_FLOAT_EQ(Dynamic->X()[2],5);
+
+		const int RewindStep = 0;
+
+		FRewindData* RewindData = Solver->GetRewindData();
+		EXPECT_TRUE(RewindData->RewindToFrame(RewindStep));
+
+		//force collision
+		Kinematic->SetX(FVec3(0,0,0));
+
+		for(int Step = RewindStep; Step <= LastStep; ++Step)
+		{
+			//Resim ignores collision since it's ResimAsSlave
+			TickSolverHelper(Module,Solver);
+
+			EXPECT_VECTOR_FLOAT_EQ(Dynamic->X(),Xs[Step]);
+		}
+
+		EXPECT_FLOAT_EQ(Dynamic->X()[2],5);
+
+		Module->DestroySolver(Solver);
+	}
+
 	TYPED_TEST(AllTraits,RewindTest_ResimAsSlaveWithForces)
 	{
 		auto Box = TSharedPtr<FImplicitObject,ESPMode::ThreadSafe>(new TBox<FReal,3>(FVec3(-10,-10,-10),FVec3(10,10,10)));
diff --git a/Engine/Source/Runtime/Experimental/ChaosSolvers/Public/RewindData.h b/Engine/Source/Runtime/Experimental/ChaosSolvers/Public/RewindData.h
--- a/Engine/Source/Runtime/Experimental/ChaosSolvers/Public/RewindData.h
+++ b/Engine/Source/Runtime/Experimental/ChaosSolvers/Public/RewindData.h
@@ -902,9 +902,7 @@ public:
 	void FinishFrame()
 	{
 		QUICK_SCOPE_CYCLE_COUNTER(RewindDataFinishFrame);
-		LatestFrame = FMath::Max(LatestFrame,CurFrame);
-		++CurFrame;
-
+		
 		if(IsResim())
 		{
 			ensure(IsInGameThread());
@@ -917,7 +915,7 @@ public:
 					if(Rigid->ResimType() == EResimType::ResimAsSlave)
 					{
 						ensure(!Info.bDesync);
-						const FGeometryParticleStateBase* State = GetStateAtFrameImp(Info,CurFrame);
+						const FGeometryParticleStateBase* State = GetStateAtFrameImp(Info,CurFrame+1);
 						if(ensure(State != nullptr))
 						{
 							State->SyncToParticle(*Rigid);
@@ -926,12 +924,14 @@ public:
 				}
 			}
 		}
-		
+
+		++CurFrame;
+		LatestFrame = FMath::Max(LatestFrame,CurFrame);
 	}
 
 	bool IsResim() const
 	{
-		return CurFrame <= LatestFrame;
+		return CurFrame < LatestFrame;
 	}
 
 	//Number of particles that we're currently storing history for
