Chaos - fix headlesschaos compile

diff --git a/Engine/Source/Programs/HeadlessChaos/Private/HeadlessChaosTestJointSolver.cpp b/Engine/Source/Programs/HeadlessChaos/Private/HeadlessChaosTestJointSolver.cpp
--- a/Engine/Source/Programs/HeadlessChaos/Private/HeadlessChaosTestJointSolver.cpp
+++ b/Engine/Source/Programs/HeadlessChaos/Private/HeadlessChaosTestJointSolver.cpp
@@ -398,7 +398,6 @@ namespace ChaosTest
 		SolverTestA.JointSettings.bAngularSLerpVelocityDriveEnabled = true;
 		SolverTestA.JointSettings.AngularDriveStiffness = Stiffness;
 		SolverTestA.JointSettings.AngularDriveDamping = Damping;
-		SolverTestA.JointSettings.AngularDriveTargetAngles = FVec3(0, 0, 0);
 		SolverTestA.JointSettings.AngularDrivePositionTarget = FRotation3::FromIdentity();
 		SolverTestA.JointSettings.AngularDriveForceMode = ForceMode;
 
@@ -580,7 +579,6 @@ namespace ChaosTest
 		SolverTestA.JointSettings.bAngularSLerpVelocityDriveEnabled = true;
 		SolverTestA.JointSettings.AngularDriveStiffness = 80.0f;
 		SolverTestA.JointSettings.AngularDriveDamping = 1.0f;
-		SolverTestA.JointSettings.AngularDriveTargetAngles = FVec3(0, 0, 0);
 		SolverTestA.JointSettings.AngularDrivePositionTarget = FRotation3::FromIdentity();
 
 		// Particle 0 is Kinematic
