Chaos - linear breakable joint unit test

diff --git a/Engine/Source/Programs/HeadlessChaos/Private/HeadlessChaosTestJointBreak.cpp b/Engine/Source/Programs/HeadlessChaos/Private/HeadlessChaosTestJointBreak.cpp
new file mode 100644
--- /dev/null
+++ b/Engine/Source/Programs/HeadlessChaos/Private/HeadlessChaosTestJointBreak.cpp
@@ -0,0 +1,343 @@
+// Copyright Epic Games, Inc. All Rights Reserved.
+
+#include "HeadlessChaos.h"
+#include "HeadlessChaosTestConstraints.h"
+
+#include "Modules/ModuleManager.h"
+#include "Chaos/ParticleHandle.h"
+#include "Chaos/PBDConstraintRule.h"
+#include "Chaos/PBDRigidsEvolution.h"
+#include "Chaos/PBDRigidsEvolutionGBF.h"
+#include "Chaos/PBDRigidParticles.h"
+#include "Chaos/PBDJointConstraints.h"
+#include "Chaos/Box.h"
+#include "Chaos/Utilities.h"
+
+namespace ChaosTest {
+
+	using namespace Chaos;
+
+	/**
+	 * Base class for joint break tests.
+	 */
+	template <typename TEvolution>
+	class FJointConstraintBreakTest : public FConstraintsTest<TEvolution>
+	{
+	public:
+		using Base = FConstraintsTest<TEvolution>;
+		using Base::Evolution;
+		using Base::AddParticleBox;
+
+		FJointConstraintBreakTest(const int32 NumIterations, const FReal Gravity)
+			: Base(NumIterations, Gravity)
+			, JointsRule(Joints)
+		{
+			Evolution.AddConstraintRule(&JointsRule);
+		}
+
+		FPBDJointConstraintHandle* AddJoint(const TVector<TGeometryParticleHandle<FReal, 3>*, 2>& InConstrainedParticleIndices, const int32 JointIndex)
+		{
+			FPBDJointConstraintHandle* Joint = Joints.AddConstraint(InConstrainedParticleIndices, FRigidTransform3(JointPositions[JointIndex], FRotation3::FromIdentity()));
+
+			FPBDJointSettings Settings = Joint->GetSettings();
+			Settings.LinearBreakForce = (JointIndex < JointLinearBreakForces.Num()) ? JointLinearBreakForces[JointIndex] : 0.0f;
+			Settings.AngularBreakTorque = (JointIndex < JointAngularBreakTorques.Num()) ? JointAngularBreakTorques[JointIndex] : 0.0f;
+			Joint->SetSettings(Settings);
+
+			return Joint;
+		}
+
+		void Create()
+		{
+			for (int32 ParticleIndex = 0; ParticleIndex < ParticlePositions.Num(); ++ParticleIndex)
+			{
+				AddParticleBox(ParticlePositions[ParticleIndex], FRotation3::MakeFromEuler(FVec3(0.f, 0.f, 0.f)).GetNormalized(), ParticleSizes[ParticleIndex], ParticleMasses[ParticleIndex]);
+			}
+
+			for (int32 JointIndex = 0; JointIndex < JointPositions.Num(); ++JointIndex)
+			{
+				const TVector<TGeometryParticleHandle<FReal, 3>*, 2> ConstraintedParticleIds(GetParticle(JointParticleIndices[JointIndex][0]), GetParticle(JointParticleIndices[JointIndex][1]));
+				AddJoint(ConstraintedParticleIds, JointIndex);
+			}
+		}
+
+		void InitVerticalChain(int32 NumParticles)
+		{
+			for (int32 ParticleIndex = 0; ParticleIndex < NumParticles; ++ParticleIndex)
+			{
+				FReal Z = (NumParticles - ParticleIndex - 1) * 100.0f;
+				FReal M = (ParticleIndex == 0) ? 0.0f : 100.0f;
+				ParticlePositions.Add(FVec3(0, 0, Z));
+				ParticleSizes.Add(FVec3(10, 10, 10));
+				ParticleMasses.Add(M);
+			}
+
+			for (int32 JointIndex = 0; JointIndex < NumParticles - 1; ++JointIndex)
+			{
+				int32 ParticleIndex0 = JointIndex;
+				int32 ParticleIndex1 = JointIndex + 1;
+				FReal Z = (NumParticles - JointIndex - 1) * 100.0f;
+				JointPositions.Add(FVec3(0, 0, Z));
+				JointParticleIndices.Add(TVector<int32, 2>(ParticleIndex0, ParticleIndex1));
+			}
+		}
+
+		// Initial particles setup
+		TArray<FVec3> ParticlePositions;
+		TArray<FVec3> ParticleSizes;
+		TArray<FReal> ParticleMasses;
+
+		// Initial joints setup
+		TArray<FVec3> JointPositions;
+		TArray<FReal> JointLinearBreakForces;
+		TArray<FReal> JointAngularBreakTorques;
+		TArray<TVector<int32, 2>> JointParticleIndices;
+
+		// Solver state
+		FPBDJointConstraints Joints;
+		TPBDConstraintIslandRule<FPBDJointConstraints> JointsRule;
+	};
+
+	// Set up a test with a non-breakable joint, then manually break it.
+	// Verify that the break callback is called and the joint is disabled.
+	template <typename TEvolution>
+	void JointBreak_ManualBreak()
+	{
+		const int32 NumIterations = 1;
+		const FReal Gravity = 980;
+		const FReal Dt = 0.01f;
+		int32 NumIts = 10;
+
+		FJointConstraintBreakTest<TEvolution> Test(NumIterations, Gravity);
+		Test.InitVerticalChain(2);
+		Test.Create();
+
+		bool bBrokenCallbackCalled = false;
+		Test.Joints.SetBreakCallback([&bBrokenCallbackCalled](FPBDJointConstraintHandle* Constraint)
+			{
+				bBrokenCallbackCalled = true;
+			});
+
+		// Run the sim - nothing should move
+		for (int32 i = 0; i < NumIts; ++i)
+		{
+			Test.Evolution.AdvanceOneTimeStep(Dt);
+			Test.Evolution.EndFrame(Dt);
+		}
+		EXPECT_NEAR(Test.GetParticle(1)->X().Z, Test.ParticlePositions[1].Z, 1.0f);
+
+		// Nothing should have broken
+		EXPECT_FALSE(bBrokenCallbackCalled);
+		EXPECT_TRUE(Test.Joints.IsConstraintEnabled(0));
+
+		// Manually break the constraints
+		Test.Joints.BreakConstraint(0);
+
+		// Check that it worked
+		EXPECT_TRUE(bBrokenCallbackCalled);
+		EXPECT_FALSE(Test.Joints.IsConstraintEnabled(0));
+
+		// Run the sim - body should fall
+		for (int32 i = 0; i < NumIts; ++i)
+		{
+			Test.Evolution.AdvanceOneTimeStep(Dt);
+			Test.Evolution.EndFrame(Dt);
+		}
+		FReal ExpectedZ = -0.5f * Gravity * (NumIts * Dt) * (NumIts * Dt);
+		EXPECT_NEAR(Test.GetParticle(1)->X().Z, ExpectedZ, 1.0f);
+	}
+
+	TYPED_TEST(AllEvolutions, JointBreakTests_TestManualBreak)
+	{
+		JointBreak_ManualBreak<TypeParam>();
+	}
+
+	// 1 Kinematic Body with 1 Dynamic body hanging from it by a breakable constraint.
+	// Constraint break force is larger than M x G, so joint should not break.
+	template <typename TEvolution>
+	void JointBreak_UnderLinearThreshold()
+	{
+		const int32 NumIterations = 1;
+		const FReal Gravity = 980;
+		const FReal Dt = 0.01f;
+		int32 NumIts = 10;
+
+		FJointConstraintBreakTest<TEvolution> Test(NumIterations, Gravity);
+		Test.InitVerticalChain(2);
+
+		// Joint should break only if Threshold < MG
+		// So not in this test
+		Test.JointLinearBreakForces =
+		{
+			1.1f * Test.ParticleMasses[1] * Gravity
+		};
+
+		Test.Create();
+
+		bool bBrokenCallbackCalled = false;
+		Test.Joints.SetBreakCallback([&bBrokenCallbackCalled](FPBDJointConstraintHandle* Constraint)
+			{
+				bBrokenCallbackCalled = true;
+			});
+
+		// Run the sim
+		for (int32 i = 0; i < NumIts; ++i)
+		{
+			Test.Evolution.AdvanceOneTimeStep(Dt);
+			Test.Evolution.EndFrame(Dt);
+		}
+
+		// Nothing should have broken
+		EXPECT_FALSE(bBrokenCallbackCalled);
+		EXPECT_TRUE(Test.Joints.IsConstraintEnabled(0));
+	}
+
+	TYPED_TEST(AllEvolutions, JointBreakTests_TestUnderLinearThreshold)
+	{
+		JointBreak_UnderLinearThreshold<TypeParam>();
+	}
+
+	// 1 Kinematic Body with 2 Dynamic bodies hanging from it by a breakable constraint.
+	// Constraint break forces are larger than M x G, so joint should not break.
+	template <typename TEvolution>
+	void JointBreak_UnderLinearThreshold2()
+	{
+		const int32 NumIterations = 1;
+		const FReal Gravity = 980;
+		const FReal Dt = 0.01f;
+		int32 NumIts = 10;
+
+		FJointConstraintBreakTest<TEvolution> Test(NumIterations, Gravity);
+		Test.InitVerticalChain(3);
+
+		// Joint should break only if Threshold < MG
+		// So not in this test
+		// NOTE: internal forces reach almost 50% over MG
+		Test.JointLinearBreakForces =
+		{
+			1.5f * (Test.ParticleMasses[1] + Test.ParticleMasses[2]) * Gravity,
+			1.5f * Test.ParticleMasses[2] * Gravity,
+		};
+
+		Test.Create();
+
+		bool bBrokenCallbackCalled = false;
+		Test.Joints.SetBreakCallback([&bBrokenCallbackCalled](FPBDJointConstraintHandle* Constraint)
+			{
+				bBrokenCallbackCalled = true;
+			});
+
+		// Run the sim
+		for (int32 i = 0; i < NumIts; ++i)
+		{
+			Test.Evolution.AdvanceOneTimeStep(Dt);
+			Test.Evolution.EndFrame(Dt);
+		}
+
+		// Nothing should have broken
+		EXPECT_FALSE(bBrokenCallbackCalled);
+		EXPECT_TRUE(Test.Joints.IsConstraintEnabled(0));
+		EXPECT_TRUE(Test.Joints.IsConstraintEnabled(1));
+	}
+
+	TYPED_TEST(AllEvolutions, JointBreakTests_TestUnderLinearThreshold2)
+	{
+		JointBreak_UnderLinearThreshold2<TypeParam>();
+	}
+
+	// 1 Kinematic Body with 1 Dynamic body hanging from it by a breakable constraint.
+	// Constraint break force is less than M x G, so joint should break.
+	template <typename TEvolution>
+	void JointBreak_OverLinearThreshold()
+	{
+		const int32 NumIterations = 1;
+		const FReal Gravity = 980;
+		const FReal Dt = 0.01f;
+		int32 NumIts = 10;
+
+		FJointConstraintBreakTest<TEvolution> Test(NumIterations, Gravity);
+		Test.InitVerticalChain(2);
+
+		// Joint should break only if Threshold < MG
+		// So yes in this test
+		Test.JointLinearBreakForces =
+		{
+			0.9f * Test.ParticleMasses[1] * Gravity
+		};
+
+		Test.Create();
+
+		bool bBrokenCallbackCalled = false;
+		Test.Joints.SetBreakCallback([&bBrokenCallbackCalled](FPBDJointConstraintHandle* Constraint)
+			{
+				bBrokenCallbackCalled = true;
+			});
+
+		// Run the sim
+		for (int32 i = 0; i < NumIts; ++i)
+		{
+			Test.Evolution.AdvanceOneTimeStep(Dt);
+			Test.Evolution.EndFrame(Dt);
+		}
+
+		// Constraint should have broken
+		EXPECT_TRUE(bBrokenCallbackCalled);
+		EXPECT_FALSE(Test.Joints.IsConstraintEnabled(0));
+	}
+
+	TYPED_TEST(AllEvolutions, JointBreakTests_TestOverLinearThreshold)
+	{
+		JointBreak_OverLinearThreshold<TypeParam>();
+	}
+
+
+	// 1 Kinematic Body with 2 Dynamic bodies hanging from it by a breakable constraint.
+	// Constraint break force is less than M x G, so joint should break.
+	template <typename TEvolution>
+	void JointBreak_OverLinearThreshold2()
+	{
+		const int32 NumIterations = 1;
+		const FReal Gravity = 980;
+		const FReal Dt = 0.01f;
+		int32 NumIts = 10;
+
+		FJointConstraintBreakTest<TEvolution> Test(NumIterations, Gravity);
+		Test.InitVerticalChain(3);
+
+		// Joint should break only if Threshold < MG
+		// So yes in this test
+		// NOTE: internal forces reach almost 50% over MG
+		Test.JointLinearBreakForces =
+		{
+			1.2f * (Test.ParticleMasses[1] + Test.ParticleMasses[2]) * Gravity,
+			1.2f * Test.ParticleMasses[2] * Gravity,
+		};
+
+		Test.Create();
+
+		bool bBrokenCallbackCalled = false;
+		Test.Joints.SetBreakCallback([&bBrokenCallbackCalled](FPBDJointConstraintHandle* Constraint)
+			{
+				bBrokenCallbackCalled = true;
+			});
+
+		// Run the sim
+		for (int32 i = 0; i < NumIts; ++i)
+		{
+			Test.Evolution.AdvanceOneTimeStep(Dt);
+			Test.Evolution.EndFrame(Dt);
+		}
+
+		// Constraint should have broken
+		EXPECT_TRUE(bBrokenCallbackCalled);
+		EXPECT_FALSE(Test.Joints.IsConstraintEnabled(0));
+		EXPECT_FALSE(Test.Joints.IsConstraintEnabled(1));
+	}
+
+	TYPED_TEST(AllEvolutions, JointBreakTests_TestOverLinearThreshold2)
+	{
+		JointBreak_OverLinearThreshold2<TypeParam>();
+	}
+
+
+}
\ No newline at end of file
