Merge geometry collection into headless chaos.

diff --git a/Engine/Source/Programs/HeadlessChaos/HeadlessChaos.Build.cs b/Engine/Source/Programs/HeadlessChaos/HeadlessChaos.Build.cs
--- a/Engine/Source/Programs/HeadlessChaos/HeadlessChaos.Build.cs
+++ b/Engine/Source/Programs/HeadlessChaos/HeadlessChaos.Build.cs
@@ -19,8 +19,14 @@ public class HeadlessChaos : ModuleRules
                 "GoogleTest",
 				"Chaos",
 				"ChaosSolvers",
-            }
-        );
+				"GeometricObjects",
+				"GeometryCollectionCore",
+				"GeometryCollectionSimulationCore",
+				"FieldSystemCore",
+				"FieldSystemSimulationCore"
+
+			}
+		);
 
 		if (Target.Platform == UnrealTargetPlatform.Win32 || Target.Platform == UnrealTargetPlatform.Win64 || Target.Platform == UnrealTargetPlatform.XboxOne)
 		{
diff --git a/Engine/Source/Programs/HeadlessChaos/Private/GeometryCollection/GeometryCollectionTestBoneHierarchy.cpp b/Engine/Source/Programs/HeadlessChaos/Private/GeometryCollection/GeometryCollectionTestBoneHierarchy.cpp
new file mode 100644
--- /dev/null
+++ b/Engine/Source/Programs/HeadlessChaos/Private/GeometryCollection/GeometryCollectionTestBoneHierarchy.cpp
@@ -0,0 +1,147 @@
+// Copyright Epic Games, Inc. All Rights Reserved.
+
+#include "GeometryCollection/GeometryCollectionTestBoneHierarchy.h"
+
+#include "PhysicsProxy/AnalyticImplicitGroup.h"
+#include "BoneHierarchy.h"
+
+#include "Chaos/Sphere.h"
+
+namespace GeometryCollectionTest
+{
+	template <class TImplicitShape>
+	void AllOnSurface(const TImplicitShape *Shape, const TArray<Chaos::TVector<float, 3>> &Points, const float Tolerance = KINDA_SMALL_NUMBER)
+	{
+		for (auto &Pt : Points)
+		{
+			const float Phi = Shape->SignedDistance(Pt);
+			EXPECT_NEAR(Phi, 0, Tolerance);
+		}
+	}
+
+	void TestSphere(const Chaos::TVector<float, 3>& Center, const float Radius, const FTransform &BoneRelXf)
+	{
+		Chaos::TSphere<float, 3>* Sphere = new Chaos::TSphere<float, 3>(Center, Radius);
+		TUniquePtr<FAnalyticImplicitGroup> Group1(new FAnalyticImplicitGroup("Root", 0));
+		Group1->Init(1);
+		Group1->SetParentBoneIndex(INDEX_NONE);
+		Group1->Add(BoneRelXf, Sphere);
+		const TArray<Chaos::TVector<float,3>> *Points = Group1->BuildSamplePoints(1.0, 1, 1000);
+		check(Points);
+		check(Points->Num());
+		Chaos::FImplicitObject *Implicit = Group1->BuildSimImplicitObject();
+		AllOnSurface(Implicit, *Points);
+	}
+	void TestSphere2(
+		const Chaos::TVector<float, 3>& Center1, 
+		const Chaos::TVector<float, 3>& Center2, 
+		const float Radius1, 
+		const float Radius2, 
+		const FTransform &BoneRelXf1,
+		const FTransform &BoneRelXf2)
+	{
+		Chaos::TSphere<float, 3>* Sphere1 = new Chaos::TSphere<float, 3>(Center1, Radius1);
+		Chaos::TSphere<float, 3>* Sphere2 = new Chaos::TSphere<float, 3>(Center2, Radius2);
+		TUniquePtr<FAnalyticImplicitGroup> Group1(new FAnalyticImplicitGroup("Root", 0));
+		Group1->Init(2);
+		Group1->SetParentBoneIndex(INDEX_NONE);
+		Group1->Add(BoneRelXf1, Sphere1);
+		Group1->Add(BoneRelXf2, Sphere2);
+		const TArray<Chaos::TVector<float,3>> *Points = Group1->BuildSamplePoints(1.0, 1, 1000);
+		check(Points);
+		check(Points->Num());
+		Chaos::FImplicitObject *Implicit = Group1->BuildSimImplicitObject();
+		AllOnSurface(Implicit, *Points);
+	}
+
+	void RunAnalyticImplicitGroupTest()
+	{
+		TestSphere(Chaos::TVector<float, 3>(0, 0, 0), 1.f, FTransform::Identity);
+		TestSphere(Chaos::TVector<float, 3>(0, 0, 0), 1.f, FTransform(FVector(1,0,0)));
+		TestSphere(Chaos::TVector<float, 3>(0, 0, 0), 1.f, FTransform(FVector(1,1,0)));
+		TestSphere(Chaos::TVector<float, 3>(0, 0, 0), 1.f, FTransform(FVector(1,1,1)));
+
+		//This test fails in Dev-Physics for spheres of larger radii
+		//TestSphere(R, Chaos::TVector<float, 3>(0, 0, 0), 10.f, FTransform::Identity);
+		//TestSphere(R, Chaos::TVector<float, 3>(0, 0, 0), 10.f, FTransform(FVector(1,0,0)));
+		//TestSphere(R, Chaos::TVector<float, 3>(0, 0, 0), 10.f, FTransform(FVector(1,1,0)));
+		//TestSphere(R, Chaos::TVector<float, 3>(0, 0, 0), 10.f, FTransform(FVector(1,1,1)));
+
+		TestSphere2( 
+			Chaos::TVector<float, 3>(0, 0, 0), 
+			Chaos::TVector<float, 3>(0, 0, 0), 
+			1.f, 
+			1.f, 
+			FTransform(FVector(-2,0,0)),
+			FTransform(FVector(2,0,0)));
+	}
+
+	void RunBoneHierarchyTest()
+	{
+		FBoneHierarchy Hierarchy;
+		Hierarchy.InitPreAdd(2);
+		{
+			TUniquePtr<FAnalyticImplicitGroup> Group1(
+				new FAnalyticImplicitGroup("Root", 0));
+			Group1->SetParentBoneIndex(INDEX_NONE);
+			Group1->Add(FTransform::Identity, new Chaos::TSphere<float, 3>(Chaos::TVector<float, 3>(0), 1.));
+			Hierarchy.Add(MoveTemp(Group1));
+
+			TUniquePtr<FAnalyticImplicitGroup> Group2(
+				new FAnalyticImplicitGroup("Bone1", 1));
+			Group2->SetParentBoneIndex(0);
+			Group2->Add(FTransform::Identity, new Chaos::TSphere<float, 3>(Chaos::TVector<float, 3>(0), 1.));
+			Hierarchy.Add(MoveTemp(Group2));
+
+			Hierarchy.InitPostAdd();
+		}
+
+		Hierarchy.PrepareForUpdate();
+		Hierarchy.SetAnimLocalSpaceTransform(0,
+			FTransform(FQuat::MakeFromEuler(FVector(0, 0, 0)), 
+			FVector(1,0,0)));
+		Hierarchy.SetAnimLocalSpaceTransform(1,
+			FTransform(FQuat::MakeFromEuler(FVector(0, 0, 0)), 
+			FVector(0,1,0)));
+		Hierarchy.SetActorWorldSpaceTransform(
+			FTransform(FQuat::MakeFromEuler(FVector(0, 0, 0)), 
+			FVector(0,0,1)));
+		Hierarchy.PrepareAnimWorldSpaceTransforms();
+		{
+			const FTransform* Xf = Hierarchy.GetAnimWorldSpaceTransformsForBone(1);
+			check(Xf);
+			const FTransform XfExpected(
+				FQuat::MakeFromEuler(FVector(0, 0, 0)), 
+				FVector(1, 1, 1));
+			check(Xf->Equals(XfExpected));
+		}
+
+		Hierarchy.PrepareForUpdate();
+		Hierarchy.SetAnimLocalSpaceTransform(0,
+			FTransform(FQuat::MakeFromEuler(FVector(10, 0, 0)), 
+			FVector(0,0,0)));
+		Hierarchy.SetAnimLocalSpaceTransform(1,
+			FTransform(FQuat::MakeFromEuler(FVector(10, 0, 0)), 
+			FVector(0,0,0)));
+		Hierarchy.SetActorWorldSpaceTransform(
+			FTransform(FQuat::MakeFromEuler(FVector(10, 0, 0)), 
+			FVector(0,0,0)));
+		Hierarchy.PrepareAnimWorldSpaceTransforms();
+		{
+			const FTransform* Xf = Hierarchy.GetAnimWorldSpaceTransformsForBone(1);
+			check(Xf);
+			const FTransform XfExpected(
+				FQuat::MakeFromEuler(FVector(30, 0, 0)), 
+				FVector(0, 0, 0));
+			check(Xf->Equals(XfExpected));
+		}
+	}
+
+	template<class T>
+	void TestImplicitBoneHierarchy()
+	{
+		RunAnalyticImplicitGroupTest();
+		RunBoneHierarchyTest();
+	}
+	template void TestImplicitBoneHierarchy<float>();
+} // namespace GeometryCollectionTest
