Merge geometry collection into headless chaos.

diff --git a/Engine/Source/Programs/HeadlessChaos/Private/GeometryCollection/GeometryCollectionTestClustering.cpp b/Engine/Source/Programs/HeadlessChaos/Private/GeometryCollection/GeometryCollectionTestClustering.cpp
new file mode 100644
--- /dev/null
+++ b/Engine/Source/Programs/HeadlessChaos/Private/GeometryCollection/GeometryCollectionTestClustering.cpp
@@ -0,0 +1,2106 @@
+// Copyright Epic Games, Inc. All Rights Reserved.
+
+#include "GeometryCollection/GeometryCollectionTestClustering.h"
+#include "GeometryCollection/GeometryCollectionTestFramework.h"
+#include "GeometryCollection/GeometryCollectionTestUtility.h"
+
+#include "GeometryCollection/GeometryCollection.h"
+#include "GeometryCollection/GeometryCollectionUtility.h"
+#include "GeometryCollection/GeometryCollectionAlgo.h"
+#include "GeometryCollection/GeometryCollectionClusteringUtility.h"
+
+#include "GeometryCollection/TransformCollection.h"
+#include "UObject/Package.h"
+#include "UObject/UObjectGlobals.h"
+#include "Field/FieldSystem.h"
+#include "Field/FieldSystemNodes.h"
+
+
+#include "GeometryCollectionProxyData.h"
+#include "PhysicsProxy/PhysicsProxies.h"
+#include "Chaos/ErrorReporter.h"
+#include "ChaosSolversModule.h"
+#include "PhysicsSolver.h"
+#include "Chaos/PBDRigidClustering.h"
+
+#include "HAL/IConsoleManager.h"
+#include "HeadlessChaosTestUtility.h"
+
+bool bSkipGeomFailing = true;
+
+
+DEFINE_LOG_CATEGORY_STATIC(GCTCL_Log, Verbose, All);
+
+// #TODO Lots of duplication in here, anyone making solver or object changes
+// has to go and fix up so many callsites here and they're all pretty much
+// Identical. The similar code should be pulled out
+
+namespace GeometryCollectionTest
+{
+	using namespace ChaosTest;
+	bool ClusterMapContains(const Chaos::TPBDRigidClustering<FPBDRigidsEvolution, FPBDCollisionConstraints, float, 3>::FClusterMap& ClusterMap, TPBDRigidParticleHandle<float, 3>* Key, TArray<TPBDRigidParticleHandle<float, 3>*> Elements)
+	{
+		if (ClusterMap.Num())
+			if (ClusterMap.Contains(Key))
+			{
+				if (ClusterMap[Key].Num() == Elements.Num())
+				{
+					for (TPBDRigidParticleHandle<float, 3>* Element : Elements)
+						if (!ClusterMap[Key].Contains(Element))
+							return false;
+					return true;
+				}
+			}
+		return false;
+	}
+
+
+	TYPED_TEST(AllTraits,GeometryCollection_RigidBodies_ClusterTest_SingleLevelNonBreaking)
+	{
+		if(bSkipGeomFailing)
+		{
+			return;
+		}
+
+		using Traits = TypeParam;
+		TFramework<Traits> UnitTest;
+
+		RigidBodyWrapper* Floor = TNewSimulationObject<GeometryType::RigidFloor>::Init<Traits>()->As<RigidBodyWrapper>();
+		UnitTest.AddSimulationObject(Floor);
+
+		TSharedPtr<FGeometryCollection> RestCollection = GeometryCollection::MakeCubeElement(FTransform(FQuat::MakeFromEuler(FVector(0, 0, 0.)), FVector(0, -10, 10)),FVector(1.0));
+		RestCollection->AppendGeometry(*GeometryCollection::MakeCubeElement(FTransform(FQuat::MakeFromEuler(FVector(0, 0, 0.)), FVector(0, 10, 10)), FVector(1.0)));
+		EXPECT_EQ(RestCollection->Transform.Num(), 2);
+
+		FGeometryCollectionClusteringUtility::ClusterAllBonesUnderNewRoot(RestCollection.Get());
+		EXPECT_EQ(RestCollection->Transform.Num(), 3);
+		RestCollection->Transform[2] = FTransform(FQuat::MakeFromEuler(FVector(90.f, 0, 0.)), FVector(0, 0, 40));
+
+		//GeometryCollectionAlgo::PrintParentHierarchy(RestCollection.Get());
+
+		CreationParameters Params;
+		Params.RestCollection = RestCollection;
+		Params.DynamicState = EObjectStateTypeEnum::Chaos_Object_Dynamic;
+		Params.CollisionType = ECollisionTypeEnum::Chaos_Surface_Volumetric;
+		Params.Simulating = true;
+		Params.EnableClustering = true;
+		Params.DamageThreshold = { 1000.f };
+		TGeometryCollectionWrapper<Traits>* Collection = TNewSimulationObject<GeometryType::GeometryCollectionWithSuppliedRestCollection>::Init<Traits>(Params)->As<TGeometryCollectionWrapper<Traits>>();
+
+		UnitTest.AddSimulationObject(Collection);
+		UnitTest.Initialize();
+
+		TManagedArray<FTransform>& Transform = Collection->DynamicCollection->Transform;
+		float StartingRigidDistance = (Transform[1].GetTranslation() - Transform[0].GetTranslation()).Size(), CurrentRigidDistance = 0.f;
+
+		TManagedArray<bool>& Active = Collection->DynamicCollection->Active;
+
+		EXPECT_TRUE(Active[0]);
+		EXPECT_TRUE(Active[1]);
+		EXPECT_TRUE(Active[2]);
+		UnitTest.Advance();
+		EXPECT_FALSE(Active[0]);
+		EXPECT_FALSE(Active[1]);
+		EXPECT_TRUE(Active[2]);
+
+		auto& Clustering = UnitTest.Solver->GetEvolution()->GetRigidClustering();
+		const auto & ClusterMap = Clustering.GetChildrenMap();
+		EXPECT_TRUE(ClusterMapContains(ClusterMap, Collection->PhysObject->GetSolverClusterHandles()[0], 
+			{ Collection->PhysObject->GetSolverParticleHandles()[0],Collection->PhysObject->GetSolverParticleHandles()[1] }));
+
+		float InitialZ = Collection->RestCollection->Transform[2].GetTranslation().Z;
+		for (int Frame = 1; Frame < 10; Frame++)
+		{			
+			UnitTest.Advance();
+
+			EXPECT_FALSE(Active[0]);
+			EXPECT_FALSE(Active[1]);
+			EXPECT_TRUE(Active[2]);
+
+			CurrentRigidDistance = (Transform[1].GetTranslation() - Transform[0].GetTranslation()).Size();
+			EXPECT_LT(FMath::Abs(CurrentRigidDistance - StartingRigidDistance), SMALL_NUMBER); // two bodies under cluster maintain distance
+			EXPECT_LT(Collection->DynamicCollection->Transform[2].GetTranslation().Z, InitialZ); // body should be falling and decreasing in Z			
+		}
+
+		EXPECT_TRUE(ClusterMapContains(ClusterMap, Collection->PhysObject->GetSolverClusterHandles()[0],
+			{ Collection->PhysObject->GetSolverParticleHandles()[0],Collection->PhysObject->GetSolverParticleHandles()[1] }));
+
+	}
+
+
+	
+	TYPED_TEST(AllTraits, GeometryCollection_RigidBodies_ClusterTest_DeactivateClusterParticle)
+	{
+		if(bSkipGeomFailing)
+		{
+			return;
+		}
+
+		using Traits = TypeParam;
+		TFramework<Traits> UnitTest;
+
+		// 5 cube leaf nodes
+		TSharedPtr<FGeometryCollection> RestCollection = GeometryCollection::MakeCubeElement(FTransform(FQuat::MakeFromEuler(FVector(0.f)), FVector(20.f)),FVector(1.0));
+		RestCollection->AppendGeometry(*GeometryCollection::MakeCubeElement(FTransform(FQuat::MakeFromEuler(FVector(0.f)), FVector(30.f)), FVector(1.0)));
+		RestCollection->AppendGeometry(*GeometryCollection::MakeCubeElement(FTransform(FQuat::MakeFromEuler(FVector(0.f)), FVector(40.f)), FVector(1.0)));
+		RestCollection->AppendGeometry(*GeometryCollection::MakeCubeElement(FTransform(FQuat::MakeFromEuler(FVector(0.f)), FVector(50.f)), FVector(1.0)));
+		RestCollection->AppendGeometry(*GeometryCollection::MakeCubeElement(FTransform(FQuat::MakeFromEuler(FVector(0.f)), FVector(60.f)), FVector(1.0)));
+
+		// 4 mid-level cluster parents
+		RestCollection->AddElements(4, FGeometryCollection::TransformGroup);
+		// @todo(ClusteringUtils) This is a bad assumption, the state flags should be initialized to zero.
+		(RestCollection->SimulationType)[5] = FGeometryCollection::ESimulationTypes::FST_Clustered;
+		(RestCollection->SimulationType)[6] = FGeometryCollection::ESimulationTypes::FST_Clustered;
+		(RestCollection->SimulationType)[7] = FGeometryCollection::ESimulationTypes::FST_Clustered;
+		(RestCollection->SimulationType)[8] = FGeometryCollection::ESimulationTypes::FST_Clustered;
+
+		// Build a binary tree cluster parent hierarchy
+		GeometryCollectionAlgo::ParentTransforms(RestCollection.Get(), 5, { 4,3 }); // Transform index 5 is parent to 4 and 3
+		GeometryCollectionAlgo::ParentTransforms(RestCollection.Get(), 6, { 5,2 });
+		GeometryCollectionAlgo::ParentTransforms(RestCollection.Get(), 7, { 6,1 });
+		GeometryCollectionAlgo::ParentTransforms(RestCollection.Get(), 8, { 7,0 });
+		
+		CreationParameters Params;
+		Params.RestCollection = RestCollection;
+		Params.DynamicState = EObjectStateTypeEnum::Chaos_Object_Dynamic;
+		Params.CollisionType = ECollisionTypeEnum::Chaos_Surface_Volumetric;
+		Params.ImplicitType = EImplicitTypeEnum::Chaos_Implicit_Box;
+		Params.Simulating = true;
+		Params.EnableClustering = true;
+		Params.DamageThreshold = { 50.0, 50.0, 50.0, FLT_MAX };
+		Params.MaxClusterLevel = 1;	
+
+		TGeometryCollectionWrapper<Traits>* Collection = TNewSimulationObject<GeometryType::GeometryCollectionWithSuppliedRestCollection>::Init<Traits>(Params)->As<TGeometryCollectionWrapper<Traits>>();
+
+		UnitTest.AddSimulationObject(Collection);
+		UnitTest.Initialize();		
+
+		TManagedArray<FTransform>& Transform = Collection->DynamicCollection->Transform;
+		float StartingRigidDistance = (Transform[1].GetTranslation() - Transform[0].GetTranslation()).Size(), CurrentRigidDistance = 0.f;
+
+		UnitTest.Advance();
+
+		TArray<Chaos::TPBDRigidClusteredParticleHandle<float, 3>*> &ParticleHandles = Collection->PhysObject->GetSolverParticleHandles();		
+
+		auto& Clustering = UnitTest.Solver->GetEvolution()->GetRigidClustering();
+		const auto & ClusterMap = Clustering.GetChildrenMap();
+
+		EXPECT_TRUE(ClusterMapContains(ClusterMap, ParticleHandles[5], { ParticleHandles[4], ParticleHandles[3] }));
+		EXPECT_TRUE(ClusterMapContains(ClusterMap, ParticleHandles[6], { ParticleHandles[5], ParticleHandles[2] }));
+		EXPECT_TRUE(ClusterMapContains(ClusterMap, ParticleHandles[7], { ParticleHandles[6], ParticleHandles[1] }));
+		EXPECT_TRUE(ClusterMapContains(ClusterMap, ParticleHandles[8], { ParticleHandles[7], ParticleHandles[0] }));
+
+
+		TArray<bool> Conditions = { false,false };
+		TArray<bool> DisabledFlags;
+
+		for (int Frame = 1; Frame < 4; Frame++)
+		{
+			UnitTest.Advance();
+			
+			if (Frame == 2)
+			{
+				Clustering.DeactivateClusterParticle(ParticleHandles[8]);
+			}
+
+			DisabledFlags.Reset();
+			for (const auto* Handle : ParticleHandles)
+			{
+				DisabledFlags.Add(Handle->Disabled());
+			}
+
+			//UE_LOG(GCTCL_Log, Verbose, TEXT("FRAME : %d"), Frame);
+			//for (int32 rdx = 0; rdx < (int32)Particles.Size(); rdx++)
+			//{
+			//  UE_LOG(GCTCL_Log, Verbose, TEXT("... ... ...Disabled[%d] : %d"), rdx, Particles.Disabled(rdx));
+			//  UE_LOG(GCTCL_Log, Verbose, TEXT("... ... ...    InvM[%d] : %f"), rdx, Particles.InvM(rdx));
+			//}
+
+			if (Frame == 1)
+			{
+				if (DisabledFlags[0] == true &&
+					DisabledFlags[1] == true &&
+					DisabledFlags[2] == true &&
+					DisabledFlags[3] == true &&
+					DisabledFlags[4] == true &&
+					DisabledFlags[5] == true &&
+					DisabledFlags[6] == true &&
+					DisabledFlags[7] == true &&
+					DisabledFlags[8] == false)
+				{
+					Conditions[0] = true;
+				}
+			}
+			else if (Frame == 2 || Frame == 3)
+			{
+				if (Conditions[0] == true)
+				{
+					if (DisabledFlags[0] == false &&
+						DisabledFlags[1] == true &&
+						DisabledFlags[2] == true &&
+						DisabledFlags[3] == true &&
+						DisabledFlags[4] == true &&
+						DisabledFlags[5] == true &&
+						DisabledFlags[6] == true &&
+						DisabledFlags[7] == false &&
+						DisabledFlags[8] == true)
+					{
+						Conditions[1] = true;
+
+						EXPECT_TRUE(!ClusterMap.Contains(ParticleHandles[8]));
+						EXPECT_TRUE(ClusterMapContains(ClusterMap, ParticleHandles[7], { ParticleHandles[6], ParticleHandles[1] }));
+						EXPECT_TRUE(ClusterMapContains(ClusterMap, ParticleHandles[6], { ParticleHandles[5], ParticleHandles[2] }));
+						EXPECT_TRUE(ClusterMapContains(ClusterMap, ParticleHandles[5], { ParticleHandles[4], ParticleHandles[3] }));
+					}
+				}
+			}
+	
+		}
+	
+		for (int i = 0; i < Conditions.Num(); i++)
+		{
+			EXPECT_TRUE(Conditions[i]);
+		}	
+	}
+	
+	
+	TYPED_TEST(AllTraits, GeometryCollection_RigidBodies_ClusterTest_BreakClusterParticle)
+	{
+		if(bSkipGeomFailing)
+		{
+			return;
+		}
+
+		using Traits = TypeParam;
+		TFramework<Traits> UnitTest;
+
+		// 5 cube leaf nodes
+		TSharedPtr<FGeometryCollection> RestCollection = GeometryCollection::MakeCubeElement(FTransform(FQuat::MakeFromEuler(FVector(0.f)), FVector(20.f)), FVector(1.0));
+		RestCollection->AppendGeometry(*GeometryCollection::MakeCubeElement(FTransform(FQuat::MakeFromEuler(FVector(0.f)), FVector(30.f)), FVector(1.0)));
+		RestCollection->AppendGeometry(*GeometryCollection::MakeCubeElement(FTransform(FQuat::MakeFromEuler(FVector(0.f)), FVector(40.f)), FVector(1.0)));
+		RestCollection->AppendGeometry(*GeometryCollection::MakeCubeElement(FTransform(FQuat::MakeFromEuler(FVector(0.f)), FVector(50.f)), FVector(1.0)));
+		RestCollection->AppendGeometry(*GeometryCollection::MakeCubeElement(FTransform(FQuat::MakeFromEuler(FVector(0.f)), FVector(60.f)), FVector(1.0)));
+
+		// 4 mid-level cluster parents
+		RestCollection->AddElements(4, FGeometryCollection::TransformGroup);
+		// @todo(ClusteringUtils) This is a bad assumption, the state flags should be initialized to zero.
+		(RestCollection->SimulationType)[5] = FGeometryCollection::ESimulationTypes::FST_Clustered;
+		(RestCollection->SimulationType)[6] = FGeometryCollection::ESimulationTypes::FST_Clustered;
+		(RestCollection->SimulationType)[7] = FGeometryCollection::ESimulationTypes::FST_Clustered;
+		(RestCollection->SimulationType)[8] = FGeometryCollection::ESimulationTypes::FST_Clustered;
+
+		// Build a binary tree cluster parent hierarchy
+		GeometryCollectionAlgo::ParentTransforms(RestCollection.Get(), 5, { 4,3 }); // Transform index 5 is parent to 4 and 3
+		GeometryCollectionAlgo::ParentTransforms(RestCollection.Get(), 6, { 5,2 });
+		GeometryCollectionAlgo::ParentTransforms(RestCollection.Get(), 7, { 6,1 });
+		GeometryCollectionAlgo::ParentTransforms(RestCollection.Get(), 8, { 7,0 });
+
+		CreationParameters Params;
+		Params.RestCollection = RestCollection;
+		Params.DynamicState = EObjectStateTypeEnum::Chaos_Object_Dynamic;
+		Params.CollisionType = ECollisionTypeEnum::Chaos_Surface_Volumetric;
+		Params.ImplicitType = EImplicitTypeEnum::Chaos_Implicit_Box;
+		Params.Simulating = true;
+		Params.EnableClustering = true;
+		Params.DamageThreshold = { 50.0, 50.0, 50.0, FLT_MAX };
+		Params.MaxClusterLevel = 1;
+
+		TGeometryCollectionWrapper<Traits>* Collection = TNewSimulationObject<GeometryType::GeometryCollectionWithSuppliedRestCollection>::Init<Traits>(Params)->As<TGeometryCollectionWrapper<Traits>>();
+
+		UnitTest.AddSimulationObject(Collection);
+		UnitTest.Initialize();
+
+		TManagedArray<FTransform>& Transform = Collection->DynamicCollection->Transform;
+		float StartingRigidDistance = (Transform[1].GetTranslation() - Transform[0].GetTranslation()).Size(), CurrentRigidDistance = 0.f;
+
+		UnitTest.Advance();
+
+		TArray<Chaos::TPBDRigidClusteredParticleHandle<float, 3>*>& ParticleHandles = Collection->PhysObject->GetSolverParticleHandles();
+
+		auto& Clustering = UnitTest.Solver->GetEvolution()->GetRigidClustering();
+		const auto& ClusterMap = Clustering.GetChildrenMap();
+
+		EXPECT_TRUE(ClusterMapContains(ClusterMap, ParticleHandles[5], { ParticleHandles[4], ParticleHandles[3] }));
+		EXPECT_TRUE(ClusterMapContains(ClusterMap, ParticleHandles[6], { ParticleHandles[5], ParticleHandles[2] }));
+		EXPECT_TRUE(ClusterMapContains(ClusterMap, ParticleHandles[7], { ParticleHandles[6], ParticleHandles[1] }));
+		EXPECT_TRUE(ClusterMapContains(ClusterMap, ParticleHandles[8], { ParticleHandles[7], ParticleHandles[0] }));
+
+
+		TArray<bool> Conditions = { false,false };
+		TArray<bool> DisabledFlags;
+
+		for (int Frame = 1; Frame < 4; Frame++)
+		{
+			UnitTest.Advance();
+
+			if (Frame == 2)
+			{
+				TMap<TGeometryParticleHandle<float, 3>*, float> ExternalStrains = { {ParticleHandles[0], -1.0f} };
+				Clustering.BreakingModel(&ExternalStrains);
+			}
+
+			DisabledFlags.Reset();
+			for (const auto* Handle : ParticleHandles)
+			{
+				DisabledFlags.Add(Handle->Disabled());
+			}
+
+			//UE_LOG(GCTCL_Log, Verbose, TEXT("FRAME : %d"), Frame);
+			//for (int32 rdx = 0; rdx < (int32)Particles.Size(); rdx++)
+			//{
+			//  UE_LOG(GCTCL_Log, Verbose, TEXT("... ... ...Disabled[%d] : %d"), rdx, Particles.Disabled(rdx));
+			//  UE_LOG(GCTCL_Log, Verbose, TEXT("... ... ...    InvM[%d] : %f"), rdx, Particles.InvM(rdx));
+			//}
+
+			if (Frame == 1)
+			{
+				if (DisabledFlags[0] == true &&
+					DisabledFlags[1] == true &&
+					DisabledFlags[2] == true &&
+					DisabledFlags[3] == true &&
+					DisabledFlags[4] == true &&
+					DisabledFlags[5] == true &&
+					DisabledFlags[6] == true &&
+					DisabledFlags[7] == true &&
+					DisabledFlags[8] == false)
+				{
+					Conditions[0] = true;
+				}
+			}
+			else if (Frame == 2 || Frame == 3)
+			{
+				if (Conditions[0] == true)
+				{
+					if (DisabledFlags[0] == false &&
+						DisabledFlags[1] == true &&
+						DisabledFlags[2] == true &&
+						DisabledFlags[3] == true &&
+						DisabledFlags[4] == true &&
+						DisabledFlags[5] == true &&
+						DisabledFlags[6] == true &&
+						DisabledFlags[7] == false &&
+						DisabledFlags[8] == true)
+					{
+						Conditions[1] = true;
+
+						EXPECT_TRUE(!ClusterMap.Contains(ParticleHandles[8]));
+						EXPECT_TRUE(ClusterMapContains(ClusterMap, ParticleHandles[7], { ParticleHandles[6], ParticleHandles[1] }));
+						EXPECT_TRUE(ClusterMapContains(ClusterMap, ParticleHandles[6], { ParticleHandles[5], ParticleHandles[2] }));
+						EXPECT_TRUE(ClusterMapContains(ClusterMap, ParticleHandles[5], { ParticleHandles[4], ParticleHandles[3] }));
+					}
+				}
+			}
+
+		}
+
+		for (int i = 0; i < Conditions.Num(); i++)
+		{
+			EXPECT_TRUE(Conditions[i]);
+		}
+	}
+
+	
+	TYPED_TEST(AllTraits, GeometryCollection_RigidBodies_ClusterTest_SingleLevelBreaking)
+	{
+		if(bSkipGeomFailing)
+		{
+			return;
+		}
+		//
+		// Test overview:
+		// Create two 1cm cubes in a cluster arranged vertically and 20cm apart.
+		// Position the cluster above the ground.
+		// Wait until the cluster hits the ground.
+		// Ensure that the cluster breaks and that the children have the correct states from then on.
+		//
+
+		using Traits = TypeParam;
+		TFramework<Traits> UnitTest;
+
+		UnitTest.AddSimulationObject(TNewSimulationObject<GeometryType::RigidFloor>::Init<Traits>()->As<RigidBodyWrapper>());
+
+
+		TSharedPtr<FGeometryCollection> RestCollection = CreateClusteredBody(FVector::ZeroVector);
+		RestCollection->Transform[2] = FTransform(FQuat::MakeFromEuler(FVector(0., 90.f, 0.)), FVector(0, 0, 17));
+		CreationParameters Params;
+		Params.RestCollection = RestCollection;
+		Params.DynamicState = EObjectStateTypeEnum::Chaos_Object_Dynamic;
+		Params.CollisionType = ECollisionTypeEnum::Chaos_Volumetric;
+		Params.ImplicitType = EImplicitTypeEnum::Chaos_Implicit_Sphere;
+		Params.Simulating = true;
+		Params.EnableClustering = true;
+		Params.DamageThreshold = { 0.1f };
+		Params.ClusterGroupIndex = 1;
+		TGeometryCollectionWrapper<Traits>* Collection = TNewSimulationObject<GeometryType::GeometryCollectionWithSuppliedRestCollection>::Init<Traits>(Params)->As<TGeometryCollectionWrapper<Traits>>();
+
+
+		UnitTest.AddSimulationObject(Collection);
+		UnitTest.Initialize();
+
+		Collection->PhysObject->SetCollisionParticlesPerObjectFraction(1.0);
+
+		TManagedArray<FTransform>& Transform = Collection->DynamicCollection->Transform;
+		float StartingRigidDistance = (Transform[1].GetTranslation() - Transform[0].GetTranslation()).Size(), CurrentRigidDistance = 0.f;
+
+		UnitTest.Advance();
+
+		auto& Clustering = UnitTest.Solver->GetEvolution()->GetRigidClustering();
+		const auto& ClusterMap = Clustering.GetChildrenMap();
+
+		EXPECT_TRUE(ClusterMapContains(ClusterMap, Collection->PhysObject->GetSolverClusterHandles()[0],
+			{ Collection->PhysObject->GetSolverParticleHandles()[0],Collection->PhysObject->GetSolverParticleHandles()[1] }));
+
+		TArray<Chaos::TPBDRigidClusteredParticleHandle<float, 3>*>& ParticleHandles = Collection->PhysObject->GetSolverParticleHandles();
+
+		// Particles array contains the following:		
+		// 0: Box1 (top)
+		// 1: Box2 (bottom)
+		TArray<bool> Conditions = { false,false,false,false,false };
+
+		// 2: Box1+Box2 Cluster
+		for (int Frame = 1; Frame < 20; Frame++)
+		{
+			UnitTest.Advance();
+
+			CurrentRigidDistance = (Transform[1].GetTranslation() - Transform[0].GetTranslation()).Size();
+
+			if (!Conditions[0])
+			{
+				// The two boxes are dropping to the ground as a cluster
+				EXPECT_TRUE(ParticleHandles[0]->Disabled());
+				EXPECT_TRUE(ParticleHandles[1]->Disabled());
+				EXPECT_FALSE(ParticleHandles[2]->Disabled());
+
+				EXPECT_LT(FMath::Abs(CurrentRigidDistance - StartingRigidDistance), 1e-4);
+				Conditions[0] = true;
+			}
+
+			if (Conditions[0] && !Conditions[1])
+			{
+				// The cluster has just hit the ground and should have broken.
+				// The boxes are still separated by StartingRigidDistance (when Rewind is disabled).
+				// All children should have zero velocity.
+
+				if (!ParticleHandles[0]->Disabled()
+					&& !ParticleHandles[1]->Disabled()
+					&& ParticleHandles[2]->Disabled()
+					&& ClusterMap.Num() == 0)
+				{
+					Conditions[1] = true;
+				}
+			}
+			if (Conditions[1] && !Conditions[2])
+			{
+				if (ParticleHandles[0]->V().Size() <= 1.e-4
+					&& ParticleHandles[1]->V().Size() <= 1.e-4
+					&& FMath::Abs(CurrentRigidDistance - StartingRigidDistance) <= 1e-4)
+				{
+					Conditions[2] = true;
+				}
+			}
+			if (Conditions[2] && !Conditions[3])
+			{
+				// The boxes are now moving independently, but they had zero velocity 
+				// last frame, so they should still be separated by StartingRigidDistance. 
+				if (!ParticleHandles[0]->Disabled()
+					&& !ParticleHandles[1]->Disabled()
+					&& ParticleHandles[2]->Disabled()
+					&& ClusterMap.Num() == 0
+					&& FMath::Abs(CurrentRigidDistance - StartingRigidDistance) <= 1e-4)
+				{
+					Conditions[3] = true;
+				}
+			}
+			if (Conditions[3])
+			{
+				// The boxes are now moving independently - the bottom one is on the ground and should be stopped.
+				// The top one is still falling, so they should be closer together	
+				EXPECT_FALSE(ParticleHandles[0]->Disabled());
+				EXPECT_FALSE(ParticleHandles[1]->Disabled());
+				EXPECT_TRUE(ParticleHandles[2]->Disabled());
+				EXPECT_EQ(ClusterMap.Num(), 0);
+				Conditions[4] = true;
+			}
+		}
+
+		for (int i = 0; i < Conditions.Num(); i++)
+		{
+			EXPECT_TRUE(Conditions[i]);
+		}
+
+		EXPECT_GT(FMath::Abs(CurrentRigidDistance - StartingRigidDistance), 1e-4);
+	}
+
+
+	
+	TYPED_TEST(AllTraits, GeometryCollection_RigidBodies_ClusterTest_NestedCluster)
+	{
+		if(bSkipGeomFailing)
+		{
+			return;
+		}
+		using Traits = TypeParam;
+		TFramework<Traits> UnitTest;
+
+		RigidBodyWrapper* Floor = TNewSimulationObject<GeometryType::RigidFloor>::Init<Traits>()->As<RigidBodyWrapper>();
+		UnitTest.AddSimulationObject(Floor);
+
+		TSharedPtr<FGeometryCollection> RestCollection = GeometryCollection::MakeCubeElement(FTransform(FQuat::MakeFromEuler(FVector(0, 0, 0.)), FVector(0, -10, 10)), FVector(1.0));
+		RestCollection->AppendGeometry(*GeometryCollection::MakeCubeElement(FTransform(FQuat::MakeFromEuler(FVector(0, 0, 0.)), FVector(0, 10, 10)), FVector(1.0)));
+		EXPECT_EQ(RestCollection->Transform.Num(), 2);
+
+		FGeometryCollectionClusteringUtility::ClusterAllBonesUnderNewRoot(RestCollection.Get());
+		EXPECT_EQ(RestCollection->Transform.Num(), 3);
+		RestCollection->Transform[2] = FTransform(FQuat::MakeFromEuler(FVector(90.f, 0, 0.)), FVector(0, 0, 40));
+
+		FGeometryCollectionClusteringUtility::ClusterBonesUnderNewNode(RestCollection.Get(), 3, { 2 }, true);
+		EXPECT_EQ(RestCollection->Transform.Num(), 4);
+		RestCollection->Transform[3] = FTransform(FQuat::MakeFromEuler(FVector(0.f, 0, 0.)), FVector(0, 0, 10));
+
+		//GeometryCollectionAlgo::PrintParentHierarchy(RestCollection.Get());
+
+		CreationParameters Params;
+		Params.RestCollection = RestCollection;
+		Params.DynamicState = EObjectStateTypeEnum::Chaos_Object_Dynamic;
+		Params.ImplicitType = EImplicitTypeEnum::Chaos_Implicit_Sphere;
+		Params.CollisionType = ECollisionTypeEnum::Chaos_Volumetric;
+		Params.Simulating = true;
+		Params.EnableClustering = true;
+		Params.DamageThreshold = { 0.1f };
+		TGeometryCollectionWrapper<Traits>* Collection = TNewSimulationObject<GeometryType::GeometryCollectionWithSuppliedRestCollection>::Init<Traits>(Params)->As<TGeometryCollectionWrapper<Traits>>();
+
+		UnitTest.AddSimulationObject(Collection);
+		UnitTest.Initialize();
+		
+		TManagedArray<FTransform>& Transform = Collection->DynamicCollection->Transform;
+		float StartingRigidDistance = (Transform[1].GetTranslation() - Transform[0].GetTranslation()).Size(), CurrentRigidDistance = 0.f;
+
+		UnitTest.Advance();
+
+		TArray<Chaos::TPBDRigidClusteredParticleHandle<float, 3>*>& ParticleHandles = Collection->PhysObject->GetSolverParticleHandles();		
+
+		auto& Clustering = UnitTest.Solver->GetEvolution()->GetRigidClustering();
+		const auto& ClusterMap = Clustering.GetChildrenMap();
+
+		EXPECT_TRUE(ClusterMapContains(ClusterMap, ParticleHandles[2], { ParticleHandles[0],ParticleHandles[1] }));
+		EXPECT_TRUE(ClusterMapContains(ClusterMap, ParticleHandles[3], { ParticleHandles[2] }));
+		
+		TArray<bool> Conditions = {false,false,false};
+
+		for (int Frame = 1; Frame < 100; Frame++)
+		{
+			UnitTest.Advance();
+
+			CurrentRigidDistance = (Transform[1].GetTranslation() - Transform[0].GetTranslation()).Size();
+
+			if (Conditions[0]==false)
+			{
+				if (
+					ParticleHandles[0]->Disabled() == true &&
+					ParticleHandles[1]->Disabled() == true &&
+					ParticleHandles[2]->Disabled() == true &&
+					ParticleHandles[3]->Disabled() == false) 
+				{
+					Conditions[0] = true;
+				}
+			}
+			else if (Conditions[0]==true && Conditions[1] == false)
+			{
+				if (
+					ParticleHandles[0]->Disabled() == true &&
+					ParticleHandles[1]->Disabled() == true &&
+					ParticleHandles[2]->Disabled() == false &&
+					ParticleHandles[3]->Disabled() == true)
+				{
+					Conditions[1] = true;
+					EXPECT_TRUE(ClusterMapContains(ClusterMap, ParticleHandles[2], { ParticleHandles[0],ParticleHandles[1] }));
+					EXPECT_EQ(ClusterMap.Num(), 1);
+					EXPECT_TRUE(!ClusterMap.Contains(ParticleHandles[3]));
+				}
+			}
+			else if (Conditions[1] == true && Conditions[2] == false)
+			{
+				if (
+					ParticleHandles[0]->Disabled() == false &&
+					ParticleHandles[1]->Disabled() == false &&
+					ParticleHandles[2]->Disabled() == true &&
+					ParticleHandles[3]->Disabled() == true)
+				{
+					Conditions[2] = true;
+					EXPECT_EQ(ClusterMap.Num(), 0);
+				}
+			}
+		}
+		
+
+		for (int i = 0; i < Conditions.Num(); i++)
+		{
+			EXPECT_TRUE(Conditions[i]);
+		}
+		
+	}
+
+	
+	TYPED_TEST(AllTraits, GeometryCollection_RigidBodies_ClusterTest_NestedCluster_NonIdentityMassToLocal)
+	{
+		if(bSkipGeomFailing)
+		{
+			return;
+		}
+		// Advance and release each cluster, everything is kinematic, so the output transforms should never change.
+		// This tests the transforms in BufferPhysicsResults, validating that MassToLocal, ChildToParent, and X,P
+		// will properly map back into the GeometryCollections animation transform hierarchy. 
+		using Traits = TypeParam;
+		TFramework<Traits> UnitTest;
+
+		TSharedPtr<FGeometryCollection> RestCollection = CreateClusteredBody_TwoParents_TwoBodiesB(FVector(0, 0, 20));
+		CreationParameters Params;
+		Params.RestCollection = RestCollection;
+		Params.DynamicState = EObjectStateTypeEnum::Chaos_Object_Dynamic;
+		Params.ImplicitType = EImplicitTypeEnum::Chaos_Implicit_Box;
+		Params.CollisionType = ECollisionTypeEnum::Chaos_Surface_Volumetric;
+		Params.Simulating = true;
+		Params.EnableClustering = true;
+		Params.DamageThreshold = { FLT_MAX };
+		Params.MaxClusterLevel = 1;
+		Params.ClusterGroupIndex = 0;
+		TGeometryCollectionWrapper<Traits>* Collection1 = TNewSimulationObject<GeometryType::GeometryCollectionWithSuppliedRestCollection>::Init<Traits>(Params)->As<TGeometryCollectionWrapper<Traits>>();
+		TSharedPtr<FGeometryDynamicCollection> DynamicCollection1 = Collection1->DynamicCollection;
+		DynamicCollection1->GetAttribute<int32>("DynamicState", FGeometryCollection::TransformGroup)[1] = (uint8)EObjectStateTypeEnum::Chaos_Object_Kinematic;
+		DynamicCollection1->GetAttribute<int32>("DynamicState", FGeometryCollection::TransformGroup)[0] = (uint8)EObjectStateTypeEnum::Chaos_Object_Kinematic;
+
+		UnitTest.AddSimulationObject(Collection1);
+
+		UnitTest.Initialize();
+
+		auto& Clustering = UnitTest.Solver->GetEvolution()->GetRigidClustering();
+		const auto& ClusterMap = Clustering.GetChildrenMap();
+		TArray<FTransform> Collection1_InitialTM; GeometryCollectionAlgo::GlobalMatrices(Collection1->RestCollection->Transform, Collection1->RestCollection->Parent, Collection1_InitialTM);
+		TArray<Chaos::TPBDRigidClusteredParticleHandle<float, 3>*>& Collection1Handles = Collection1->PhysObject->GetSolverParticleHandles();
+		const auto& SovlerParticleHandles = UnitTest.Solver->GetParticles().GetParticleHandles();
+
+		EXPECT_EQ(SovlerParticleHandles.Size(), 4);
+		EXPECT_EQ(ClusterMap.Num(), 2);
+		EXPECT_TRUE(ClusterMapContains(ClusterMap, Collection1Handles[2], { Collection1Handles[1],Collection1Handles[0] }));
+		EXPECT_TRUE(ClusterMapContains(ClusterMap, Collection1Handles[3], { Collection1Handles[2] }));
+
+		UnitTest.Advance();
+
+		EXPECT_EQ(SovlerParticleHandles.Size(), 4);
+		EXPECT_EQ(ClusterMap.Num(), 2);
+		EXPECT_TRUE(ClusterMapContains(ClusterMap, Collection1Handles[2], { Collection1Handles[1],Collection1Handles[0] }));
+		EXPECT_TRUE(ClusterMapContains(ClusterMap, Collection1Handles[3], { Collection1Handles[2] }));
+		TArray<FTransform> Collection1_PreReleaseTM; GeometryCollectionAlgo::GlobalMatrices(DynamicCollection1->Transform, DynamicCollection1->Parent, Collection1_PreReleaseTM);
+		for (int Idx = 0; Idx < Collection1_PreReleaseTM.Num(); Idx++) {
+			EXPECT_TRUE( (Collection1_PreReleaseTM[Idx].GetTranslation()-Collection1_InitialTM[Idx].GetTranslation()).Size()<KINDA_SMALL_NUMBER);
+		}
+
+		UnitTest.Solver->GetEvolution()->GetRigidClustering().DeactivateClusterParticle({ Collection1Handles[3] });
+		UnitTest.Advance();
+
+		EXPECT_EQ(SovlerParticleHandles.Size(), 4);
+		EXPECT_EQ(ClusterMap.Num(), 1);
+		EXPECT_TRUE(ClusterMapContains(ClusterMap, Collection1Handles[2], { Collection1Handles[1],Collection1Handles[0] }));
+		TArray<FTransform> Collection1_PostReleaseTM; GeometryCollectionAlgo::GlobalMatrices(DynamicCollection1->Transform, DynamicCollection1->Parent, Collection1_PostReleaseTM);
+		for (int Idx = 0; Idx < Collection1_PostReleaseTM.Num(); Idx++) {
+			EXPECT_TRUE((Collection1_PostReleaseTM[Idx].GetTranslation() - Collection1_InitialTM[Idx].GetTranslation()).Size() < KINDA_SMALL_NUMBER);
+		}
+
+		UnitTest.Solver->GetEvolution()->GetRigidClustering().DeactivateClusterParticle({ Collection1Handles[2] });
+		UnitTest.Advance();
+
+		EXPECT_EQ(SovlerParticleHandles.Size(), 4);
+		EXPECT_EQ(ClusterMap.Num(), 0);
+		TArray<FTransform> Collection1_PostRelease2TM; GeometryCollectionAlgo::GlobalMatrices(DynamicCollection1->Transform, DynamicCollection1->Parent, Collection1_PostRelease2TM);
+		for (int Idx = 0; Idx < Collection1_PostRelease2TM.Num(); Idx++) {
+			EXPECT_TRUE((Collection1_PostRelease2TM[Idx].GetTranslation() - Collection1_InitialTM[Idx].GetTranslation()).Size() < KINDA_SMALL_NUMBER);
+		}
+
+	}
+
+
+	
+	TYPED_TEST(AllTraits, GeometryCollection_RigidBodies_ClusterTest_NestedCluster_MultiStrain)
+	{
+		if(bSkipGeomFailing)
+		{
+			return;
+		}
+		using Traits = TypeParam;
+		TFramework<Traits> UnitTest;
+		
+		RigidBodyWrapper* Floor = TNewSimulationObject<GeometryType::RigidFloor>::Init<Traits>()->As<RigidBodyWrapper>();
+		UnitTest.AddSimulationObject(Floor);
+
+		TSharedPtr<FGeometryCollection> RestCollection = GeometryCollection::MakeCubeElement(FTransform(FQuat::MakeFromEuler(FVector(0.f)), FVector(20.f)),FVector(1.0));
+		RestCollection->AppendGeometry(*GeometryCollection::MakeCubeElement(FTransform(FQuat::MakeFromEuler(FVector(0.f)), FVector(30.f)), FVector(1.0)));
+		RestCollection->AppendGeometry(*GeometryCollection::MakeCubeElement(FTransform(FQuat::MakeFromEuler(FVector(0.f)), FVector(40.f)), FVector(1.0)));
+		RestCollection->AppendGeometry(*GeometryCollection::MakeCubeElement(FTransform(FQuat::MakeFromEuler(FVector(0.f)), FVector(50.f)), FVector(1.0)));
+		RestCollection->AppendGeometry(*GeometryCollection::MakeCubeElement(FTransform(FQuat::MakeFromEuler(FVector(0.f)), FVector(60.f)), FVector(1.0)));
+
+		RestCollection->AddElements(4, FGeometryCollection::TransformGroup);
+		// @todo(ClusteringUtils) This is a bad assumption, the state flags should be initialized to zero.
+		(RestCollection->SimulationType)[5] = FGeometryCollection::ESimulationTypes::FST_Clustered;
+		(RestCollection->SimulationType)[6] = FGeometryCollection::ESimulationTypes::FST_Clustered;
+		(RestCollection->SimulationType)[7] = FGeometryCollection::ESimulationTypes::FST_Clustered;
+		(RestCollection->SimulationType)[8] = FGeometryCollection::ESimulationTypes::FST_Clustered;
+
+		GeometryCollectionAlgo::ParentTransforms(RestCollection.Get(), 5, { 4,3 });
+		GeometryCollectionAlgo::ParentTransforms(RestCollection.Get(), 6, { 5,2 });
+		GeometryCollectionAlgo::ParentTransforms(RestCollection.Get(), 7, { 6,1 });
+		GeometryCollectionAlgo::ParentTransforms(RestCollection.Get(), 8, { 7,0 });
+
+		// @todo(brice->Bill.Henderson) Why did this not work? I needed to build my own parenting and level initilization. 
+		//FGeometryCollectionClusteringUtility::ClusterAllBonesUnderNewRoot(RestCollection.Get());
+		//FGeometryCollectionClusteringUtility::ClusterBonesUnderNewNode(RestCollection.Get(), 4, { 0, 1 }, true);
+		//FGeometryCollectionClusteringUtility::ClusterBonesUnderNewNode(RestCollection.Get(), 4, { 2, 3 }, true);
+
+		CreationParameters Params;
+		Params.RestCollection = RestCollection;
+		Params.DynamicState = EObjectStateTypeEnum::Chaos_Object_Dynamic;
+		Params.ImplicitType = EImplicitTypeEnum::Chaos_Implicit_Sphere;
+		Params.CollisionType = ECollisionTypeEnum::Chaos_Volumetric;
+		Params.Simulating = true;
+		Params.EnableClustering = true;
+		Params.DamageThreshold = { 50.0, 50.0, 50.0, FLT_MAX };
+		TGeometryCollectionWrapper<Traits>* Collection = TNewSimulationObject<GeometryType::GeometryCollectionWithSuppliedRestCollection>::Init<Traits>(Params)->As<TGeometryCollectionWrapper<Traits>>();
+
+		UnitTest.AddSimulationObject(Collection);
+		UnitTest.Initialize();
+
+		TManagedArray<FTransform>& Transform = Collection->DynamicCollection->Transform;
+		float StartingRigidDistance = (Transform[1].GetTranslation() - Transform[0].GetTranslation()).Size(), CurrentRigidDistance = 0.f;
+
+		TArray<bool> Conditions = { false,false,false,false };
+
+		UnitTest.Advance();
+
+		TArray<Chaos::TPBDRigidClusteredParticleHandle<float, 3>*>& ParticleHandles = Collection->PhysObject->GetSolverParticleHandles();		
+
+		auto& Clustering = UnitTest.Solver->GetEvolution()->GetRigidClustering();
+		const auto& ClusterMap = Clustering.GetChildrenMap();
+
+		EXPECT_TRUE(ClusterMapContains(ClusterMap, ParticleHandles[5], { ParticleHandles[4],ParticleHandles[3] }));
+		EXPECT_TRUE(ClusterMapContains(ClusterMap, ParticleHandles[6], { ParticleHandles[5],ParticleHandles[2] }));
+		EXPECT_TRUE(ClusterMapContains(ClusterMap, ParticleHandles[7], { ParticleHandles[6],ParticleHandles[1] }));
+		EXPECT_TRUE(ClusterMapContains(ClusterMap, ParticleHandles[8], { ParticleHandles[7],ParticleHandles[0] }));
+
+
+		for (int Frame = 1; Frame < 40; Frame++)
+		{
+			UnitTest.Advance();			
+
+			CurrentRigidDistance = (Transform[1].GetTranslation() - Transform[0].GetTranslation()).Size();
+			
+			if (Conditions[0] == false)
+			{
+				if (ParticleHandles[0]->Disabled() == true &&
+					ParticleHandles[1]->Disabled() == true &&
+					ParticleHandles[2]->Disabled() == true &&
+					ParticleHandles[3]->Disabled() == true &&
+					ParticleHandles[4]->Disabled() == true &&
+					ParticleHandles[5]->Disabled() == true &&
+					ParticleHandles[6]->Disabled() == true &&
+					ParticleHandles[7]->Disabled() == true &&
+					ParticleHandles[8]->Disabled() == false)
+				{
+					Conditions[0] = true;
+				}
+			}
+			else if (Conditions[0] == true && Conditions[1] == false)
+			{
+				if (
+					ParticleHandles[0]->Disabled() == false &&
+					ParticleHandles[1]->Disabled() == true &&
+					ParticleHandles[2]->Disabled() == true &&
+					ParticleHandles[3]->Disabled() == true &&
+					ParticleHandles[4]->Disabled() == true &&
+					ParticleHandles[5]->Disabled() == true &&
+					ParticleHandles[6]->Disabled() == true &&
+					ParticleHandles[7]->Disabled() == false &&
+					ParticleHandles[8]->Disabled() == true)
+				{
+					Conditions[1] = true;
+
+					EXPECT_EQ(ClusterMap.Num(), 3);
+					EXPECT_TRUE(ClusterMapContains(ClusterMap, ParticleHandles[5], { ParticleHandles[4],ParticleHandles[3] }));
+					EXPECT_TRUE(ClusterMapContains(ClusterMap, ParticleHandles[6], { ParticleHandles[5],ParticleHandles[2] }));
+					EXPECT_TRUE(ClusterMapContains(ClusterMap, ParticleHandles[7], { ParticleHandles[6],ParticleHandles[1] }));
+				}
+			}
+			else if (Conditions[1] == true && Conditions[2] == false)
+			{
+				if (
+					ParticleHandles[0]->Disabled() == false &&
+					ParticleHandles[1]->Disabled() == false &&
+					ParticleHandles[2]->Disabled() == true &&
+					ParticleHandles[3]->Disabled() == true &&
+					ParticleHandles[4]->Disabled() == true &&
+					ParticleHandles[5]->Disabled() == true &&
+					ParticleHandles[6]->Disabled() == false &&
+					ParticleHandles[7]->Disabled() == true &&
+					ParticleHandles[8]->Disabled() == true)
+				{
+					Conditions[2] = true;
+
+					EXPECT_EQ(ClusterMap.Num(), 2);
+					EXPECT_TRUE(ClusterMapContains(ClusterMap, ParticleHandles[5], { ParticleHandles[4],ParticleHandles[3] }));
+					EXPECT_TRUE(ClusterMapContains(ClusterMap, ParticleHandles[6], { ParticleHandles[5],ParticleHandles[2] }));
+				}
+			}
+			else if (Conditions[2] == true && Conditions[3] == false)
+			{
+				if (
+					ParticleHandles[0]->Disabled() == false &&
+					ParticleHandles[1]->Disabled() == false &&
+					ParticleHandles[2]->Disabled() == false &&
+					ParticleHandles[3]->Disabled() == true &&
+					ParticleHandles[4]->Disabled() == true &&
+					ParticleHandles[5]->Disabled() == false &&
+					ParticleHandles[6]->Disabled() == true &&
+					ParticleHandles[7]->Disabled() == true &&
+					ParticleHandles[8]->Disabled() == true)
+				{
+					Conditions[3] = true;
+
+					EXPECT_EQ(ClusterMap.Num(), 1);
+					EXPECT_TRUE(ClusterMapContains(ClusterMap, ParticleHandles[5], { ParticleHandles[4],ParticleHandles[3] }));
+				}
+			}
+			else if (Conditions[3] == true)
+			{
+				// fLT_MAX strain so last cluster should never break. 
+				EXPECT_TRUE(ParticleHandles[0]->Disabled() == false);
+				EXPECT_TRUE(ParticleHandles[1]->Disabled() == false);
+				EXPECT_TRUE(ParticleHandles[2]->Disabled() == false);
+				EXPECT_TRUE(ParticleHandles[3]->Disabled() == true);
+				EXPECT_TRUE(ParticleHandles[4]->Disabled() == true);
+				EXPECT_TRUE(ParticleHandles[5]->Disabled() == false);
+				EXPECT_TRUE(ParticleHandles[6]->Disabled() == true);
+				EXPECT_TRUE(ParticleHandles[7]->Disabled() == true);
+				EXPECT_TRUE(ParticleHandles[8]->Disabled() == true);
+				EXPECT_EQ(ClusterMap.Num(), 1);
+				EXPECT_TRUE(ClusterMapContains(ClusterMap, ParticleHandles[5], { ParticleHandles[4],ParticleHandles[3] }));
+			}
+		}
+		for (int i = 0; i < Conditions.Num(); i++)
+		{
+			EXPECT_TRUE(Conditions[i]);
+		}
+
+	}
+
+
+
+	
+	TYPED_TEST(AllTraits, GeometryCollection_RigidBodies_ClusterTest_KinematicAnchor)
+	{
+		if(bSkipGeomFailing)
+		{
+			return;
+		}
+		// Test : Set one element kinematic. When the cluster breaks the elements that do not contain the kinematic
+		//        rigid body should be dynamic, while the clusters that contain the kinematic body should remain 
+		//        kinematic.
+
+		using Traits = TypeParam;
+		TFramework<Traits> UnitTest;
+
+		TSharedPtr<FGeometryCollection> RestCollection = GeometryCollection::MakeCubeElement(FTransform(FQuat::MakeFromEuler(FVector(0.f)), FVector(20.f)), FVector(1.0));
+		RestCollection->AppendGeometry(*GeometryCollection::MakeCubeElement(FTransform(FQuat::MakeFromEuler(FVector(0.f)), FVector(30.f)), FVector(1.0)));
+		RestCollection->AppendGeometry(*GeometryCollection::MakeCubeElement(FTransform(FQuat::MakeFromEuler(FVector(0.f)), FVector(40.f)), FVector(1.0)));
+		RestCollection->AppendGeometry(*GeometryCollection::MakeCubeElement(FTransform(FQuat::MakeFromEuler(FVector(0.f)), FVector(50.f)), FVector(1.0)));
+		RestCollection->AppendGeometry(*GeometryCollection::MakeCubeElement(FTransform(FQuat::MakeFromEuler(FVector(0.f)), FVector(60.f)), FVector(1.0)));
+
+		RestCollection->AddElements(4, FGeometryCollection::TransformGroup);
+		// @todo(ClusteringUtils) This is a bad assumption, the state flags should be initialized to zero.
+		(RestCollection->SimulationType)[5] = FGeometryCollection::ESimulationTypes::FST_Clustered;
+		(RestCollection->SimulationType)[6] = FGeometryCollection::ESimulationTypes::FST_Clustered;
+		(RestCollection->SimulationType)[7] = FGeometryCollection::ESimulationTypes::FST_Clustered;
+		(RestCollection->SimulationType)[8] = FGeometryCollection::ESimulationTypes::FST_Clustered;
+
+		GeometryCollectionAlgo::ParentTransforms(RestCollection.Get(), 5, { 4,3 });
+		GeometryCollectionAlgo::ParentTransforms(RestCollection.Get(), 6, { 5,2 });
+		GeometryCollectionAlgo::ParentTransforms(RestCollection.Get(), 7, { 6,1 });
+		GeometryCollectionAlgo::ParentTransforms(RestCollection.Get(), 8, { 7,0 });
+
+		CreationParameters Params;
+		Params.RestCollection = RestCollection;
+		Params.DynamicState = EObjectStateTypeEnum::Chaos_Object_Dynamic;
+		Params.ImplicitType = EImplicitTypeEnum::Chaos_Implicit_Box;
+		Params.CollisionType = ECollisionTypeEnum::Chaos_Surface_Volumetric;
+		Params.Simulating = true;
+		Params.EnableClustering = true;
+		Params.DamageThreshold = { 50.0, 50.0, 50.0, FLT_MAX };
+		Params.MaxClusterLevel = 1;
+		TGeometryCollectionWrapper<Traits>* Collection = TNewSimulationObject<GeometryType::GeometryCollectionWithSuppliedRestCollection>::Init<Traits>(Params)->As<TGeometryCollectionWrapper<Traits>>();
+
+		Collection->DynamicCollection->GetAttribute<int32>("DynamicState", FGeometryCollection::TransformGroup)[1] = (uint8)EObjectStateTypeEnum::Chaos_Object_Kinematic;
+
+		UnitTest.AddSimulationObject(Collection);
+		UnitTest.Initialize();
+
+		TManagedArray<FTransform>& Transform = Collection->DynamicCollection->Transform;
+		float StartingRigidDistance = (Transform[1].GetTranslation() - Transform[0].GetTranslation()).Size(), CurrentRigidDistance = 0.f;
+
+		TArray<bool> Conditions = { false,false,false,false };
+
+		UnitTest.Advance();
+
+		TArray<Chaos::TPBDRigidClusteredParticleHandle<float, 3>*>& ParticleHandles = Collection->PhysObject->GetSolverParticleHandles();
+		TArray<Chaos::TPBDRigidClusteredParticleHandle<float, 3>*>& ClusterHandles = Collection->PhysObject->GetSolverClusterHandles();
+
+		auto& Clustering = UnitTest.Solver->GetEvolution()->GetRigidClustering();
+		const auto& ClusterMap = Clustering.GetChildrenMap();
+
+		EXPECT_TRUE(ClusterMapContains(ClusterMap, ParticleHandles[5], { ParticleHandles[4],ParticleHandles[3] }));
+		EXPECT_TRUE(ClusterMapContains(ClusterMap, ParticleHandles[6], { ParticleHandles[5],ParticleHandles[2] }));
+		EXPECT_TRUE(ClusterMapContains(ClusterMap, ParticleHandles[7], { ParticleHandles[6],ParticleHandles[1] }));
+		EXPECT_TRUE(ClusterMapContains(ClusterMap, ParticleHandles[8], { ParticleHandles[7],ParticleHandles[0] }));
+
+		
+		for (int Frame = 1; Frame < 10; Frame++)
+		{
+			UnitTest.Advance();
+
+			/*
+			if (Frame == 2)
+			{
+				Solver->GetRigidClustering().DeactivateClusterParticle(9);
+			}
+			if (Frame == 4)
+			{
+				Solver->GetRigidClustering().DeactivateClusterParticle(8);
+			}
+			*/
+
+			//UE_LOG(GCTCL_Log, Verbose, TEXT("FRAME : %d"), Frame);
+			//for (int32 rdx = 0; rdx < (int32)Particles.Size(); rdx++)
+			//{
+			//	UE_LOG(GCTCL_Log, Verbose, TEXT("... ... ...Disabled[%d] : %d"), rdx, Particles.Disabled(rdx));
+			//	UE_LOG(GCTCL_Log, Verbose, TEXT("... ... ...    InvM[%d] : %f"), rdx, Particles.InvM(rdx));
+			//}
+
+			
+			EXPECT_NE(ParticleHandles[0]->InvM(), 0.f); // dynamic rigid
+			EXPECT_EQ(ParticleHandles[1]->InvM(), 0.f); // kinematic rigid
+			EXPECT_NE(ParticleHandles[2]->InvM(), 0.f); // dynamic rigid
+			EXPECT_NE(ParticleHandles[3]->InvM(), 0.f); // dynamic rigid
+			EXPECT_NE(ParticleHandles[4]->InvM(), 0.f); // dynamic rigid
+			EXPECT_NE(ParticleHandles[5]->InvM(), 0.f); // dynamic rigid
+			EXPECT_NE(ParticleHandles[6]->InvM(), 0.f); // dynamic cluster
+
+			FVector Ref1, Ref2, Ref7; // RigidBody0(Dynamic), RigidBody1(Kinematic), RigidBody6(Kinematic then Dynamic)
+			if (Conditions[0] == false && Frame == 1)
+			{
+				if (ParticleHandles[0]->Disabled() == true &&
+					ParticleHandles[1]->Disabled() == true &&
+					ParticleHandles[2]->Disabled() == true &&
+					ParticleHandles[3]->Disabled() == true &&
+					ParticleHandles[4]->Disabled() == true &&
+					ParticleHandles[5]->Disabled() == true &&
+					ParticleHandles[6]->Disabled() == true &&
+					ParticleHandles[7]->Disabled() == true &&
+					ParticleHandles[8]->Disabled() == false)
+				{
+					Conditions[0] = true;
+					Ref1 = ParticleHandles[0]->X();
+					Ref2 = ParticleHandles[1]->X();
+					Ref7 = ParticleHandles[6]->X();
+					EXPECT_EQ(ParticleHandles[7]->InvM(), 0.f); // kinematic cluster
+					EXPECT_EQ(ParticleHandles[7]->ObjectState(), Chaos::EObjectStateType::Kinematic); // kinematic cluster
+					EXPECT_EQ(ParticleHandles[8]->InvM(), 0.f); // kinematic cluster
+					EXPECT_EQ(ParticleHandles[8]->ObjectState(), Chaos::EObjectStateType::Kinematic); // kinematic cluster
+				}
+			}
+			else if (Conditions[0] == true && Conditions[1] == false && Frame == 2)
+			{
+				if (ParticleHandles[0]->Disabled() == false &&
+					ParticleHandles[1]->Disabled() == true &&
+					ParticleHandles[2]->Disabled() == true &&
+					ParticleHandles[3]->Disabled() == true &&
+					ParticleHandles[4]->Disabled() == true &&
+					ParticleHandles[5]->Disabled() == true &&
+					ParticleHandles[6]->Disabled() == true &&
+					ParticleHandles[7]->Disabled() == false &&
+					ParticleHandles[8]->Disabled() == true)
+				{
+					Conditions[1] = true;
+					FVector X1 = ParticleHandles[0]->X();
+					FVector X2 = ParticleHandles[1]->X();
+					FVector X7 = ParticleHandles[6]->X();
+
+					EXPECT_NEAR(FMath::Abs(X1.Size() - Ref1.Size()), 0, KINDA_SMALL_NUMBER);// << *FString("Kinematic body1 moved");
+					EXPECT_NEAR(FMath::Abs(X2.Size() - Ref2.Size()), 0, KINDA_SMALL_NUMBER);// << *FString("Kinematic body2 moved");
+					EXPECT_NEAR(FMath::Abs(X7.Size() - Ref7.Size()), 0, KINDA_SMALL_NUMBER);// << *FString("Kinematic body7 moved");
+
+					EXPECT_EQ(ParticleHandles[7]->InvM(), 0.f); // kinematic cluster
+					EXPECT_EQ(ParticleHandles[7]->ObjectState(), Chaos::EObjectStateType::Kinematic); // kinematic cluster
+					EXPECT_EQ(ParticleHandles[8]->InvM(), 0.f); // kinematic cluster
+					EXPECT_EQ(ParticleHandles[8]->ObjectState(), Chaos::EObjectStateType::Kinematic); // kinematic cluster
+
+					EXPECT_EQ(ClusterMap.Num(), 3);
+					EXPECT_TRUE(ClusterMapContains(ClusterMap, ParticleHandles[7], { ParticleHandles[6],ParticleHandles[1] }));
+					EXPECT_TRUE(ClusterMapContains(ClusterMap, ParticleHandles[6], { ParticleHandles[5],ParticleHandles[2] }));
+					EXPECT_TRUE(ClusterMapContains(ClusterMap, ParticleHandles[5], { ParticleHandles[4],ParticleHandles[3] }));					
+				}
+			}
+			else if (Conditions[1] == true && Conditions[2] == false && Frame == 4)
+			{
+				if (ParticleHandles[0]->Disabled() == false &&
+					ParticleHandles[1]->Disabled() == false &&
+					ParticleHandles[2]->Disabled() == true &&
+					ParticleHandles[3]->Disabled() == true &&
+					ParticleHandles[4]->Disabled() == true &&
+					ParticleHandles[5]->Disabled() == true &&
+					ParticleHandles[6]->Disabled() == false &&
+					ParticleHandles[7]->Disabled() == true &&
+					ParticleHandles[8]->Disabled() == true)
+				{
+					Conditions[2] = true;
+					FVector X1 = ParticleHandles[0]->X();
+					FVector X2 = ParticleHandles[1]->X();
+					FVector X7 = ParticleHandles[6]->X();					
+
+					EXPECT_GT(FMath::Abs(X1.Size() - Ref1.Size()), KINDA_SMALL_NUMBER);// << *FString("Dynamic body failed to move");
+					EXPECT_NEAR(FMath::Abs(X2.Size() - Ref2.Size()), 0, KINDA_SMALL_NUMBER);// << *FString("Kinematic body moved");
+					EXPECT_NEAR(FMath::Abs(X7.Size() - Ref7.Size()), 0, KINDA_SMALL_NUMBER);// << FString("Kinematic body moved");
+					EXPECT_EQ(ParticleHandles[7]->InvM(), 0.f); // kinematic cluster
+					EXPECT_EQ(ParticleHandles[7]->ObjectState(), Chaos::EObjectStateType::Kinematic); // kinematic cluster
+
+					EXPECT_EQ(ClusterMap.Num(), 2);
+					EXPECT_TRUE(ClusterMapContains(ClusterMap, ParticleHandles[6], { ParticleHandles[5],ParticleHandles[2] }));
+					EXPECT_TRUE(ClusterMapContains(ClusterMap, ParticleHandles[5], { ParticleHandles[4],ParticleHandles[3] }));
+				}
+			}
+			else if (Conditions[2] == true && Conditions[3] == false && Frame == 6)
+			{				
+				if (ParticleHandles[0]->Disabled() == false &&
+					ParticleHandles[1]->Disabled() == false &&
+					ParticleHandles[2]->Disabled() == true &&
+					ParticleHandles[3]->Disabled() == true &&
+					ParticleHandles[4]->Disabled() == true &&
+					ParticleHandles[5]->Disabled() == true &&
+					ParticleHandles[6]->Disabled() == false &&
+					ParticleHandles[7]->Disabled() == true &&
+					ParticleHandles[8]->Disabled() == true)
+				{
+					Conditions[3] = true;
+					FVector X1 = ParticleHandles[0]->X();
+					FVector X2 = ParticleHandles[1]->X();
+					FVector X7 = ParticleHandles[6]->X();
+
+					EXPECT_GT(FMath::Abs(X1.Size() - Ref1.Size()), KINDA_SMALL_NUMBER);// << *FString("Dynamic body 1 failed to move.");
+					EXPECT_NEAR(FMath::Abs(X2.Size() - Ref2.Size()), 0, KINDA_SMALL_NUMBER);// << *FString("Kinematic body moved");
+					EXPECT_GT(FMath::Abs(X7.Size() - Ref7.Size()), KINDA_SMALL_NUMBER);// << *FString("Dynamic body 7 failed to move");
+					EXPECT_EQ(ParticleHandles[7]->InvM(), 0.f); // kinematic cluster
+					EXPECT_EQ(ParticleHandles[7]->ObjectState(), Chaos::EObjectStateType::Kinematic); // kinematic cluster
+
+
+					EXPECT_EQ(ClusterMap.Num(), 2);
+					EXPECT_TRUE(ClusterMapContains(ClusterMap, ParticleHandles[6], { ParticleHandles[5],ParticleHandles[2] }));
+					EXPECT_TRUE(ClusterMapContains(ClusterMap, ParticleHandles[5], { ParticleHandles[4],ParticleHandles[3] }));
+				}
+			}
+		}
+		for (int i = 0; i < Conditions.Num(); i++)
+		{
+			EXPECT_TRUE(Conditions[i]);
+		}		
+	}
+
+	
+	TYPED_TEST(AllTraits, GeometryCollection_RigidBodies_ClusterTest_StaticAnchor)
+	{
+		if(bSkipGeomFailing)
+		{
+			return;
+		}
+		// Test : Set one element static. When the cluster breaks the elements that do not contain the static
+		//        rigid body should be dynamic, while the clusters that contain the static body should remain 
+		//        static. 
+
+		using Traits = TypeParam;
+		TFramework<Traits> UnitTest;
+
+		TSharedPtr<FGeometryCollection> RestCollection = GeometryCollection::MakeCubeElement(FTransform(FQuat::MakeFromEuler(FVector(0.f)), FVector(20.f)), FVector(1.0));
+		RestCollection->AppendGeometry(*GeometryCollection::MakeCubeElement(FTransform(FQuat::MakeFromEuler(FVector(0.f)), FVector(30.f)), FVector(1.0)));
+		RestCollection->AppendGeometry(*GeometryCollection::MakeCubeElement(FTransform(FQuat::MakeFromEuler(FVector(0.f)), FVector(40.f)), FVector(1.0)));
+		RestCollection->AppendGeometry(*GeometryCollection::MakeCubeElement(FTransform(FQuat::MakeFromEuler(FVector(0.f)), FVector(50.f)), FVector(1.0)));
+		RestCollection->AppendGeometry(*GeometryCollection::MakeCubeElement(FTransform(FQuat::MakeFromEuler(FVector(0.f)), FVector(60.f)), FVector(1.0)));
+
+		RestCollection->AddElements(4, FGeometryCollection::TransformGroup);
+		// @todo(ClusteringUtils) This is a bad assumption, the state flags should be initialized to zero.
+		(RestCollection->SimulationType)[5] = FGeometryCollection::ESimulationTypes::FST_Clustered;
+		(RestCollection->SimulationType)[6] = FGeometryCollection::ESimulationTypes::FST_Clustered;
+		(RestCollection->SimulationType)[7] = FGeometryCollection::ESimulationTypes::FST_Clustered;
+		(RestCollection->SimulationType)[8] = FGeometryCollection::ESimulationTypes::FST_Clustered;
+
+		GeometryCollectionAlgo::ParentTransforms(RestCollection.Get(), 5, { 4,3 });
+		GeometryCollectionAlgo::ParentTransforms(RestCollection.Get(), 6, { 5,2 });
+		GeometryCollectionAlgo::ParentTransforms(RestCollection.Get(), 7, { 6,1 });
+		GeometryCollectionAlgo::ParentTransforms(RestCollection.Get(), 8, { 7,0 });
+
+		CreationParameters Params;
+		Params.RestCollection = RestCollection;
+		Params.DynamicState = EObjectStateTypeEnum::Chaos_Object_Dynamic;
+		Params.ImplicitType = EImplicitTypeEnum::Chaos_Implicit_Box;
+		Params.CollisionType = ECollisionTypeEnum::Chaos_Surface_Volumetric;
+		Params.Simulating = true;
+		Params.EnableClustering = true;
+		Params.DamageThreshold = { 50.0, 50.0, 50.0, FLT_MAX };
+		Params.MaxClusterLevel = 1;
+		TGeometryCollectionWrapper<Traits>* Collection = TNewSimulationObject<GeometryType::GeometryCollectionWithSuppliedRestCollection>::Init<Traits>(Params)->As<TGeometryCollectionWrapper<Traits>>();
+
+		Collection->DynamicCollection->GetAttribute<int32>("DynamicState", FGeometryCollection::TransformGroup)[1] = (uint8)EObjectStateTypeEnum::Chaos_Object_Kinematic;
+
+		UnitTest.AddSimulationObject(Collection);
+		UnitTest.Initialize();
+
+		TManagedArray<FTransform>& Transform = Collection->DynamicCollection->Transform;
+		float StartingRigidDistance = (Transform[1].GetTranslation() - Transform[0].GetTranslation()).Size(), CurrentRigidDistance = 0.f;
+
+		TArray<bool> Conditions = { false,false,false,false };
+
+		UnitTest.Advance();
+
+		TArray<Chaos::TPBDRigidClusteredParticleHandle<float, 3>*>& ParticleHandles = Collection->PhysObject->GetSolverParticleHandles();
+		TArray<Chaos::TPBDRigidClusteredParticleHandle<float, 3>*>& ClusterHandles = Collection->PhysObject->GetSolverClusterHandles();
+
+		auto& Clustering = UnitTest.Solver->GetEvolution()->GetRigidClustering();
+		const auto& ClusterMap = Clustering.GetChildrenMap();
+		
+		for (int Frame = 1; Frame < 10; Frame++)
+		{
+			UnitTest.Advance();
+			/*
+			if (Frame == 2)
+			{
+				Solver->GetRigidClustering().DeactivateClusterParticle(9);
+			}
+			if (Frame == 4)
+			{
+				Solver->GetRigidClustering().DeactivateClusterParticle(8);
+			}
+			*/
+			//UE_LOG(GCTCL_Log, Verbose, TEXT("FRAME : %d"), Frame);
+			//for (int32 rdx = 0; rdx < (int32)Particles.Size(); rdx++)
+			//{
+			//	UE_LOG(GCTCL_Log, Verbose, TEXT("... ... ...Disabled[%d] : %d"), rdx, Particles.Disabled(rdx));
+			//	UE_LOG(GCTCL_Log, Verbose, TEXT("... ... ...    InvM[%d] : %f"), rdx, Particles.InvM(rdx));
+			//}
+
+			EXPECT_NE(ParticleHandles[0]->InvM(), 0.f); // dynamic rigid
+			EXPECT_EQ(ParticleHandles[1]->InvM(), 0.f); // static rigid
+			EXPECT_NE(ParticleHandles[2]->InvM(), 0.f); // dynamic rigid
+			EXPECT_NE(ParticleHandles[3]->InvM(), 0.f); // dynamic rigid
+			EXPECT_NE(ParticleHandles[4]->InvM(), 0.f); // dynamic rigid
+			EXPECT_NE(ParticleHandles[5]->InvM(), 0.f); // dynamic rigid
+			EXPECT_NE(ParticleHandles[6]->InvM(), 0.f); // dynamic cluster
+
+		
+
+
+			FVector Ref1, Ref2, Ref7; // RigidBody0(Dynamic), RigidBody1(static), RigidBody6(static then Dynamic)						
+			if (Conditions[0] == false && Frame == 1)
+			{
+				if (ParticleHandles[0]->Disabled() == true &&
+					ParticleHandles[1]->Disabled() == true &&
+					ParticleHandles[2]->Disabled() == true &&
+					ParticleHandles[3]->Disabled() == true &&
+					ParticleHandles[4]->Disabled() == true &&
+					ParticleHandles[5]->Disabled() == true &&
+					ParticleHandles[6]->Disabled() == true &&
+					ParticleHandles[7]->Disabled() == true &&
+					ParticleHandles[8]->Disabled() == false)
+				{
+					Conditions[0] = true;
+					Ref1 = ParticleHandles[0]->X();
+					Ref2 = ParticleHandles[1]->X();
+					Ref7 = ParticleHandles[6]->X();
+					EXPECT_EQ(ParticleHandles[7]->InvM(), 0.f); // static cluster
+					EXPECT_EQ(ParticleHandles[7]->ObjectState(), Chaos::EObjectStateType::Static); // static cluster
+					EXPECT_EQ(ParticleHandles[8]->InvM(), 0.f); // static cluster
+					EXPECT_EQ(ParticleHandles[8]->ObjectState(), Chaos::EObjectStateType::Static); // static cluster
+				}
+			}		
+			else if (Conditions[0] == true && Conditions[1] == false && Frame == 2)
+			{
+				if (ParticleHandles[0]->Disabled() == false &&
+					ParticleHandles[1]->Disabled() == true &&
+					ParticleHandles[2]->Disabled() == true &&
+					ParticleHandles[3]->Disabled() == true &&
+					ParticleHandles[4]->Disabled() == true &&
+					ParticleHandles[5]->Disabled() == true &&
+					ParticleHandles[6]->Disabled() == true &&
+					ParticleHandles[7]->Disabled() == false &&
+					ParticleHandles[8]->Disabled() == true)
+				{
+					Conditions[1] = true;
+					FVector X1 = ParticleHandles[0]->X();
+					FVector X2 = ParticleHandles[1]->X();
+					FVector X7 = ParticleHandles[6]->X();
+
+					EXPECT_NEAR(FMath::Abs(X1.Size() - Ref1.Size()), 0, KINDA_SMALL_NUMBER);// << *FString("Static body1 moved");
+					EXPECT_NEAR(FMath::Abs(X2.Size() - Ref2.Size()), 0, KINDA_SMALL_NUMBER);// << *FString("Static body2 moved");
+					EXPECT_NEAR(FMath::Abs(X7.Size() - Ref7.Size()), 0, KINDA_SMALL_NUMBER);// << *FString("Static body7 moved");
+
+					EXPECT_EQ(ParticleHandles[7]->InvM(), 0.f); // static cluster
+					EXPECT_EQ(ParticleHandles[7]->ObjectState(), Chaos::EObjectStateType::Static); // static cluster
+					EXPECT_EQ(ParticleHandles[8]->InvM(), 0.f); // static cluster
+					EXPECT_EQ(ParticleHandles[8]->ObjectState(), Chaos::EObjectStateType::Static); // static cluster
+				}
+			}
+			else if (Conditions[1] == true && Conditions[2] == false && Frame == 4)
+			{
+				if (ParticleHandles[0]->Disabled() == false &&
+					ParticleHandles[1]->Disabled() == false &&
+					ParticleHandles[2]->Disabled() == true &&
+					ParticleHandles[3]->Disabled() == true &&
+					ParticleHandles[4]->Disabled() == true &&
+					ParticleHandles[5]->Disabled() == true &&
+					ParticleHandles[6]->Disabled() == false &&
+					ParticleHandles[7]->Disabled() == true &&
+					ParticleHandles[8]->Disabled() == true)
+				{
+					Conditions[2] = true;
+					FVector X1 = ParticleHandles[0]->X();
+					FVector X2 = ParticleHandles[1]->X();
+					FVector X7 = ParticleHandles[6]->X();
+
+					EXPECT_GT(FMath::Abs(X1.Size() - Ref1.Size()), KINDA_SMALL_NUMBER);// << *FString("Dynamic body failed to move");
+					EXPECT_NEAR(X2.Size() - Ref2.Size(), 0, KINDA_SMALL_NUMBER);// << *FString("static body moved");
+					EXPECT_NEAR(X7.Size() - Ref7.Size(), 0, KINDA_SMALL_NUMBER);// << *FString("static body moved");
+					EXPECT_EQ(ParticleHandles[7]->InvM(), 0.f); // kinematic cluster
+					EXPECT_EQ(ParticleHandles[7]->ObjectState(), Chaos::EObjectStateType::Kinematic); // kinematic cluster
+				}
+			}
+			else if (Conditions[2] == true && Conditions[3] == false && Frame == 6)
+			{
+				if (ParticleHandles[0]->Disabled() == false &&
+					ParticleHandles[1]->Disabled() == false &&
+					ParticleHandles[2]->Disabled() == true &&
+					ParticleHandles[3]->Disabled() == true &&
+					ParticleHandles[4]->Disabled() == true &&
+					ParticleHandles[5]->Disabled() == true &&
+					ParticleHandles[6]->Disabled() == false &&
+					ParticleHandles[7]->Disabled() == true &&
+					ParticleHandles[8]->Disabled() == true)
+				{
+					Conditions[3] = true;
+					FVector X1 = ParticleHandles[0]->X();
+					FVector X2 = ParticleHandles[1]->X();
+					FVector X7 = ParticleHandles[6]->X();
+
+					EXPECT_GT(FMath::Abs(X1.Size() - Ref1.Size()), KINDA_SMALL_NUMBER);// << *FString("Dynamic body 1 failed to move.");
+					EXPECT_NEAR(X2.Size() - Ref2.Size(), 0, KINDA_SMALL_NUMBER);// << *FString("static body moved");
+					EXPECT_GT(FMath::Abs(X7.Size() - Ref7.Size()), KINDA_SMALL_NUMBER);// << *FString("Dynamic body 7 failed to move");
+					EXPECT_EQ(ParticleHandles[7]->InvM(), 0.f); // static cluster
+					EXPECT_EQ(ParticleHandles[7]->ObjectState(), Chaos::EObjectStateType::Kinematic); // static cluster
+				}
+			}
+					
+		}
+		for (int i = 0; i < Conditions.Num(); i++)
+		{
+			EXPECT_TRUE(Conditions[i]);
+		}
+
+	}
+
+	
+	TYPED_TEST(AllTraits, GeometryCollection_RigidBodies_ClusterTest_UnionClusters)
+	{
+		if(bSkipGeomFailing)
+		{
+			return;
+		}
+		// Test : Joining collections using the ClusterGroupIndex by a particle dynamically created within the solver. 		
+
+		using Traits = TypeParam;
+		TFramework<Traits> UnitTest;
+
+		TSharedPtr<FGeometryCollection> RestCollection = CreateClusteredBody(FVector(-2, 0, 3));
+		CreationParameters Params;
+		Params.RestCollection = RestCollection;
+		Params.DynamicState = EObjectStateTypeEnum::Chaos_Object_Dynamic;
+		Params.ImplicitType = EImplicitTypeEnum::Chaos_Implicit_Box;
+		Params.CollisionType = ECollisionTypeEnum::Chaos_Surface_Volumetric;
+		Params.Simulating = true;
+		Params.EnableClustering = true;
+		Params.DamageThreshold = { FLT_MAX };
+		Params.ClusterGroupIndex = 1;
+		TGeometryCollectionWrapper<Traits>* Collection = TNewSimulationObject<GeometryType::GeometryCollectionWithSuppliedRestCollection>::Init<Traits>(Params)->As<TGeometryCollectionWrapper<Traits>>();
+
+
+		TSharedPtr<FGeometryCollection> RestCollection2 = CreateClusteredBody(FVector(2, 0, 3));
+		CreationParameters Params2;
+		Params2.RestCollection = RestCollection2;
+		Params2.DynamicState = EObjectStateTypeEnum::Chaos_Object_Kinematic;
+		Params2.ImplicitType = EImplicitTypeEnum::Chaos_Implicit_Box;
+		Params2.CollisionType = ECollisionTypeEnum::Chaos_Surface_Volumetric;
+		Params2.Simulating = true;
+		Params2.EnableClustering = true;
+		Params2.DamageThreshold = { FLT_MAX };
+		Params2.ClusterGroupIndex = 1;
+		TGeometryCollectionWrapper<Traits>* Collection2 = TNewSimulationObject<GeometryType::GeometryCollectionWithSuppliedRestCollection>::Init<Traits>(Params2)->As<TGeometryCollectionWrapper<Traits>>();		
+		
+		UnitTest.AddSimulationObject(Collection);
+		UnitTest.AddSimulationObject(Collection2);
+		UnitTest.Initialize();
+
+		TSharedPtr<FGeometryDynamicCollection> DynamicCollection = Collection->DynamicCollection;
+		TSharedPtr<FGeometryDynamicCollection> DynamicCollection2 = Collection2->DynamicCollection;
+
+
+		TArray<float> Distances;
+		TManagedArray<FTransform>& Transform = DynamicCollection->Transform;
+		TManagedArray<FTransform>& Transform2 = DynamicCollection2->Transform;
+
+
+		auto& Clustering = UnitTest.Solver->GetEvolution()->GetRigidClustering();
+		const auto& ClusterMap = Clustering.GetChildrenMap();
+		const auto& ParticleHandles = UnitTest.Solver->GetParticles().GetParticleHandles();
+
+		EXPECT_EQ(ClusterMap.Num(), 2);
+		EXPECT_TRUE(ClusterMapContains(ClusterMap, ParticleHandles.Handle(2)->CastToRigidParticle(), { ParticleHandles.Handle(1)->CastToRigidParticle(),ParticleHandles.Handle(0)->CastToRigidParticle() }));
+		EXPECT_TRUE(ClusterMapContains(ClusterMap, ParticleHandles.Handle(5)->CastToRigidParticle(), { ParticleHandles.Handle(4)->CastToRigidParticle(),ParticleHandles.Handle(3)->CastToRigidParticle() }));
+
+		for (int Frame = 0; Frame < 100; Frame++)
+		{
+			UnitTest.Advance();
+
+			if (Frame == 0)
+			{
+				TArray<FTransform> GlobalTransform;
+				GeometryCollectionAlgo::GlobalMatrices(DynamicCollection->Transform, DynamicCollection->Parent, GlobalTransform);
+
+				TArray<FTransform> GlobalTransform2;
+				GeometryCollectionAlgo::GlobalMatrices(DynamicCollection2->Transform, DynamicCollection2->Parent, GlobalTransform2);
+
+				// build relative transforms distances
+				for (int32 i = 0; i < (int32)GlobalTransform.Num()-1; i++)
+				{
+					for (int j = 0; j < (int32)GlobalTransform2.Num()-1; j++)
+					{
+						Distances.Add((GlobalTransform[i].GetTranslation() - GlobalTransform2[j].GetTranslation()).Size());
+					}
+				}
+				
+				EXPECT_EQ(ClusterMap.Num(), 1);
+				EXPECT_TRUE(ClusterMapContains(ClusterMap, ParticleHandles.Handle(6)->CastToRigidParticle(), { ParticleHandles.Handle(1)->CastToRigidParticle(),ParticleHandles.Handle(0)->CastToRigidParticle(),ParticleHandles.Handle(3)->CastToRigidParticle(),ParticleHandles.Handle(4)->CastToRigidParticle() }));
+
+			}
+		}
+
+		
+		TArray<FTransform> GlobalTransform;
+		GeometryCollectionAlgo::GlobalMatrices(DynamicCollection->Transform, DynamicCollection->Parent, GlobalTransform);
+
+		TArray<FTransform> GlobalTransform2;
+		GeometryCollectionAlgo::GlobalMatrices(DynamicCollection2->Transform, DynamicCollection2->Parent, GlobalTransform2);
+
+		// build relative transforms distances
+		TArray<float> Distances2;
+		for (int32 i = 0; i < (int32)GlobalTransform.Num() - 1; i++)
+		{
+			for (int j = 0; j < (int32)GlobalTransform2.Num() - 1; j++)
+			{
+				Distances2.Add((GlobalTransform[i].GetTranslation() - GlobalTransform2[j].GetTranslation()).Size());
+			}
+		}
+		for (int i = 0; i < Distances.Num()/2.0; i++)
+		{
+			EXPECT_LT( FMath::Abs(Distances[i] - Distances2[i]), 0.1 );
+		}
+		
+	}
+
+
+	
+	TYPED_TEST(AllTraits, GeometryCollection_RigidBodies_ClusterTest_UnionClustersFalling)
+	{
+		if(bSkipGeomFailing)
+		{
+			return;
+		}
+		// Test : Joining collections using the ClusterGroupIndex by a particle dynamically created within the solver. 		
+		
+		using Traits = TypeParam;
+		TFramework<Traits> UnitTest;
+
+		TSharedPtr<FGeometryCollection> RestCollection = CreateClusteredBody(FVector(-2, 0, 3));
+		CreationParameters Params;
+		Params.RestCollection = RestCollection;
+		Params.DynamicState = EObjectStateTypeEnum::Chaos_Object_Dynamic;
+		Params.ImplicitType = EImplicitTypeEnum::Chaos_Implicit_Box;
+		Params.CollisionType = ECollisionTypeEnum::Chaos_Surface_Volumetric;
+		Params.Simulating = true;
+		Params.EnableClustering = true;
+		Params.DamageThreshold = { FLT_MAX };
+		Params.ClusterGroupIndex = 1;
+		TGeometryCollectionWrapper<Traits>* Collection = TNewSimulationObject<GeometryType::GeometryCollectionWithSuppliedRestCollection>::Init<Traits>(Params)->As<TGeometryCollectionWrapper<Traits>>();
+
+
+		TSharedPtr<FGeometryCollection> RestCollection2 = CreateClusteredBody(FVector(2, 0, 3));
+		CreationParameters Params2;
+		Params2.RestCollection = RestCollection2;
+		Params2.DynamicState = EObjectStateTypeEnum::Chaos_Object_Dynamic;
+		Params2.ImplicitType = EImplicitTypeEnum::Chaos_Implicit_Box;
+		Params2.CollisionType = ECollisionTypeEnum::Chaos_Surface_Volumetric;
+		Params2.Simulating = true;
+		Params2.EnableClustering = true;
+		Params2.DamageThreshold = { FLT_MAX };
+		Params2.ClusterGroupIndex = 1;
+		TGeometryCollectionWrapper<Traits>* Collection2 = TNewSimulationObject<GeometryType::GeometryCollectionWithSuppliedRestCollection>::Init<Traits>(Params2)->As<TGeometryCollectionWrapper<Traits>>();
+
+		UnitTest.AddSimulationObject(Collection);
+		UnitTest.AddSimulationObject(Collection2);
+		UnitTest.Initialize();
+
+		TSharedPtr<FGeometryDynamicCollection> DynamicCollection = Collection->DynamicCollection;
+		TSharedPtr<FGeometryDynamicCollection> DynamicCollection2 = Collection2->DynamicCollection;
+
+
+		TArray<float> Distances;
+		TManagedArray<FTransform>& Transform = DynamicCollection->Transform;
+		TManagedArray<FTransform>& Transform2 = DynamicCollection2->Transform;
+
+
+		auto& Clustering = UnitTest.Solver->GetEvolution()->GetRigidClustering();
+		const auto& ClusterMap = Clustering.GetChildrenMap();
+		const auto& ParticleHandles = UnitTest.Solver->GetParticles().GetParticleHandles();
+
+		EXPECT_EQ(ClusterMap.Num(), 2);
+		EXPECT_TRUE(ClusterMapContains(ClusterMap, ParticleHandles.Handle(2)->CastToRigidParticle(), { ParticleHandles.Handle(1)->CastToRigidParticle(),ParticleHandles.Handle(0)->CastToRigidParticle() }));
+		EXPECT_TRUE(ClusterMapContains(ClusterMap, ParticleHandles.Handle(5)->CastToRigidParticle(), { ParticleHandles.Handle(4)->CastToRigidParticle(),ParticleHandles.Handle(3)->CastToRigidParticle() }));
+
+		TArray<FTransform> PrevGlobalTransform;
+		GeometryCollectionAlgo::GlobalMatrices(DynamicCollection->Transform, DynamicCollection->Parent, PrevGlobalTransform);
+
+		TArray<FTransform> PrevGlobalTransform2;
+		GeometryCollectionAlgo::GlobalMatrices(DynamicCollection2->Transform, DynamicCollection2->Parent, PrevGlobalTransform2);
+
+
+		for (int Frame = 0; Frame < 100; Frame++)
+		{
+			UnitTest.Advance();
+
+			TArray<FTransform> GlobalTransform;
+			GeometryCollectionAlgo::GlobalMatrices(DynamicCollection->Transform, DynamicCollection->Parent, GlobalTransform);
+
+			TArray<FTransform> GlobalTransform2;
+			GeometryCollectionAlgo::GlobalMatrices(DynamicCollection2->Transform, DynamicCollection2->Parent, GlobalTransform2);
+
+			EXPECT_EQ(ClusterMap.Num(), 1);
+			EXPECT_TRUE(ClusterMapContains(ClusterMap, ParticleHandles.Handle(6)->CastToRigidParticle(), { ParticleHandles.Handle(1)->CastToRigidParticle(),ParticleHandles.Handle(0)->CastToRigidParticle(),ParticleHandles.Handle(3)->CastToRigidParticle(),ParticleHandles.Handle(4)->CastToRigidParticle() }));
+
+			EXPECT_TRUE(DynamicCollection->Parent[0] == INDEX_NONE);
+			EXPECT_TRUE(DynamicCollection->Parent[1] == INDEX_NONE);
+			EXPECT_TRUE(DynamicCollection->Parent[2] == INDEX_NONE);
+
+			EXPECT_TRUE(DynamicCollection2->Parent[0] == INDEX_NONE);
+			EXPECT_TRUE(DynamicCollection2->Parent[1] == INDEX_NONE);
+			EXPECT_TRUE(DynamicCollection2->Parent[2] == INDEX_NONE);
+
+			EXPECT_TRUE(GlobalTransform[0].GetTranslation().X == PrevGlobalTransform[0].GetTranslation().X);
+			EXPECT_TRUE(GlobalTransform[1].GetTranslation().X == PrevGlobalTransform[1].GetTranslation().X);
+			EXPECT_TRUE(GlobalTransform[0].GetTranslation().Y == PrevGlobalTransform[0].GetTranslation().Y);
+			EXPECT_TRUE(GlobalTransform[1].GetTranslation().Y == PrevGlobalTransform[1].GetTranslation().Y);
+			EXPECT_TRUE(GlobalTransform[0].GetTranslation().Z < PrevGlobalTransform[0].GetTranslation().Z);
+			EXPECT_TRUE(GlobalTransform[1].GetTranslation().Z < PrevGlobalTransform[1].GetTranslation().Z);
+
+			EXPECT_TRUE(GlobalTransform2[0].GetTranslation().X == PrevGlobalTransform2[0].GetTranslation().X);
+			EXPECT_TRUE(GlobalTransform2[1].GetTranslation().X == PrevGlobalTransform2[1].GetTranslation().X);
+			EXPECT_TRUE(GlobalTransform2[0].GetTranslation().Y == PrevGlobalTransform2[0].GetTranslation().Y);
+			EXPECT_TRUE(GlobalTransform2[1].GetTranslation().Y == PrevGlobalTransform2[1].GetTranslation().Y);
+			EXPECT_TRUE(GlobalTransform2[0].GetTranslation().Z < PrevGlobalTransform2[0].GetTranslation().Z);
+			EXPECT_TRUE(GlobalTransform2[1].GetTranslation().Z < PrevGlobalTransform2[1].GetTranslation().Z);
+
+			PrevGlobalTransform = GlobalTransform;
+			PrevGlobalTransform2 = GlobalTransform2;
+		}
+	}
+
+
+	
+	TYPED_TEST(AllTraits, GeometryCollection_RigidBodies_ClusterTest_UnionClusterCollisions)
+	{
+		if(bSkipGeomFailing)
+		{
+			return;
+		}
+		// Test : Joining collections using the ClusterGroupIndex by a particle dynamically created within the solver. 		
+
+		using Traits = TypeParam;
+		TFramework<Traits> UnitTest;
+
+		TSharedPtr<FGeometryCollection> RestCollection = CreateClusteredBody(FVector(-2, 0, 3));
+		CreationParameters Params;
+		Params.RestCollection = RestCollection;
+		Params.DynamicState = EObjectStateTypeEnum::Chaos_Object_Dynamic;
+		Params.ImplicitType = EImplicitTypeEnum::Chaos_Implicit_LevelSet;
+		Params.CollisionType = ECollisionTypeEnum::Chaos_Surface_Volumetric;
+		Params.Simulating = true;
+		Params.EnableClustering = true;
+		Params.DamageThreshold = { FLT_MAX };
+		Params.ClusterGroupIndex = 1;
+		TGeometryCollectionWrapper<Traits>* Collection = TNewSimulationObject<GeometryType::GeometryCollectionWithSuppliedRestCollection>::Init<Traits>(Params)->As<TGeometryCollectionWrapper<Traits>>();
+
+
+		TSharedPtr<FGeometryCollection> RestCollection2 = CreateClusteredBody(FVector(2, 0, 3));
+		CreationParameters Params2;
+		Params2.RestCollection = RestCollection2;
+		Params2.DynamicState = EObjectStateTypeEnum::Chaos_Object_Dynamic;
+		Params2.ImplicitType = EImplicitTypeEnum::Chaos_Implicit_LevelSet;
+		Params2.CollisionType = ECollisionTypeEnum::Chaos_Surface_Volumetric;
+		Params2.Simulating = true;
+		Params2.EnableClustering = true;
+		Params2.DamageThreshold = { FLT_MAX };
+		Params2.ClusterGroupIndex = 1;
+		TGeometryCollectionWrapper<Traits>* Collection2 = TNewSimulationObject<GeometryType::GeometryCollectionWithSuppliedRestCollection>::Init<Traits>(Params2)->As<TGeometryCollectionWrapper<Traits>>();
+
+		UnitTest.AddSimulationObject(Collection);
+		UnitTest.AddSimulationObject(Collection2);
+		UnitTest.AddSimulationObject(TNewSimulationObject<GeometryType::RigidFloor>::Init<Traits>()->As<RigidBodyWrapper>());
+
+		UnitTest.Initialize();
+
+		TSharedPtr<FGeometryDynamicCollection> DynamicCollection = Collection->DynamicCollection;
+		TSharedPtr<FGeometryDynamicCollection> DynamicCollection2 = Collection2->DynamicCollection;
+
+
+		TArray<float> Distances;
+		TManagedArray<FTransform>& Transform = DynamicCollection->Transform;
+		TManagedArray<FTransform>& Transform2 = DynamicCollection2->Transform;
+
+
+		auto& Clustering = UnitTest.Solver->GetEvolution()->GetRigidClustering();
+		const auto& ClusterMap = Clustering.GetChildrenMap();
+		const auto& ParticleHandles = UnitTest.Solver->GetParticles().GetParticleHandles();
+
+		EXPECT_EQ(ClusterMap.Num(), 2);
+		EXPECT_TRUE(ClusterMapContains(ClusterMap, ParticleHandles.Handle(2)->CastToRigidParticle(), { ParticleHandles.Handle(1)->CastToRigidParticle(),ParticleHandles.Handle(0)->CastToRigidParticle() }));
+		EXPECT_TRUE(ClusterMapContains(ClusterMap, ParticleHandles.Handle(5)->CastToRigidParticle(), { ParticleHandles.Handle(4)->CastToRigidParticle(),ParticleHandles.Handle(3)->CastToRigidParticle() }));
+
+		for (int Frame = 0; Frame < 100; Frame++)
+		{
+			UnitTest.Advance();
+
+			TArray<FTransform> GlobalTransform;
+			GeometryCollectionAlgo::GlobalMatrices(DynamicCollection->Transform, DynamicCollection->Parent, GlobalTransform);
+
+			TArray<FTransform> GlobalTransform2;
+			GeometryCollectionAlgo::GlobalMatrices(DynamicCollection2->Transform, DynamicCollection2->Parent, GlobalTransform2);
+
+			EXPECT_EQ(ClusterMap.Num(), 1);
+			EXPECT_TRUE(ClusterMapContains(ClusterMap, ParticleHandles.Handle(7)->CastToRigidParticle(), { ParticleHandles.Handle(1)->CastToRigidParticle(),ParticleHandles.Handle(0)->CastToRigidParticle(),ParticleHandles.Handle(3)->CastToRigidParticle(),ParticleHandles.Handle(4)->CastToRigidParticle() }));
+
+			EXPECT_TRUE(DynamicCollection->Parent[0] == INDEX_NONE);
+			EXPECT_TRUE(DynamicCollection->Parent[1] == INDEX_NONE);
+			EXPECT_TRUE(DynamicCollection->Parent[2] == INDEX_NONE);
+
+			EXPECT_TRUE(DynamicCollection2->Parent[0] == INDEX_NONE);
+			EXPECT_TRUE(DynamicCollection2->Parent[1] == INDEX_NONE);
+			EXPECT_TRUE(DynamicCollection2->Parent[2] == INDEX_NONE);
+
+			EXPECT_TRUE(GlobalTransform[0].GetTranslation().Z > 0.f);
+			EXPECT_TRUE(GlobalTransform[1].GetTranslation().Z > 0.f);
+			EXPECT_TRUE(GlobalTransform2[0].GetTranslation().Z > 0.f);
+			EXPECT_TRUE(GlobalTransform2[1].GetTranslation().Z > 0.f);
+		}
+	}
+	
+
+	
+	TYPED_TEST(AllTraits, GeometryCollection_RigidBodies_ClusterTest_ReleaseClusterParticle_ClusteredNode)
+	{
+		if(bSkipGeomFailing)
+		{
+			return;
+		}
+		// Test : Build to geometry collections, cluster them together, release the sub bodies of the first collection. 
+		//        ... should create a internal cluster with property transform mappings. 
+
+		using Traits = TypeParam;
+		TFramework<Traits> UnitTest;
+
+		TSharedPtr<FGeometryCollection> RestCollection = CreateClusteredBody(FVector(0, 0, 100));
+		CreationParameters Params;
+		Params.RestCollection = RestCollection;
+		Params.DynamicState = EObjectStateTypeEnum::Chaos_Object_Dynamic;
+		Params.ImplicitType = EImplicitTypeEnum::Chaos_Implicit_Sphere;
+		Params.CollisionType = ECollisionTypeEnum::Chaos_Volumetric;
+		Params.Simulating = true;
+		Params.EnableClustering = true;
+		Params.DamageThreshold = { FLT_MAX };
+		Params.ClusterGroupIndex = 1;
+		Params.ClusterConnectionMethod = Chaos::FClusterCreationParameters<FReal>::EConnectionMethod::DelaunayTriangulation;
+		TGeometryCollectionWrapper<Traits>* Collection1 = TNewSimulationObject<GeometryType::GeometryCollectionWithSuppliedRestCollection>::Init<Traits>(Params)->As<TGeometryCollectionWrapper<Traits>>();
+		UnitTest.AddSimulationObject(Collection1);
+
+
+		TSharedPtr<FGeometryCollection> RestCollection2 = CreateClusteredBody(FVector(0, 0, 200));
+		CreationParameters Params2;
+		Params2.RestCollection = RestCollection2;
+		Params2.DynamicState = EObjectStateTypeEnum::Chaos_Object_Dynamic;
+		Params2.ImplicitType = EImplicitTypeEnum::Chaos_Implicit_Sphere;
+		Params2.CollisionType = ECollisionTypeEnum::Chaos_Volumetric;
+		Params2.Simulating = true;
+		Params2.EnableClustering = true;
+		Params2.DamageThreshold = { FLT_MAX };
+		Params2.ClusterGroupIndex = 1;
+		Params2.ClusterConnectionMethod = Chaos::FClusterCreationParameters<FReal>::EConnectionMethod::DelaunayTriangulation;
+		TGeometryCollectionWrapper<Traits>* Collection2 = TNewSimulationObject<GeometryType::GeometryCollectionWithSuppliedRestCollection>::Init<Traits>(Params2)->As<TGeometryCollectionWrapper<Traits>>();
+
+		UnitTest.AddSimulationObject(Collection2);
+		UnitTest.Initialize();
+
+		auto& Clustering = UnitTest.Solver->GetEvolution()->GetRigidClustering();
+		const auto& ClusterMap = Clustering.GetChildrenMap();
+
+		TSharedPtr<FGeometryDynamicCollection> DynamicCollection1 = Collection1->DynamicCollection;
+		TSharedPtr<FGeometryDynamicCollection> DynamicCollection2 = Collection2->DynamicCollection;
+
+		TArray<FTransform> Collection1_InitialTM; GeometryCollectionAlgo::GlobalMatrices(DynamicCollection1->Transform, DynamicCollection1->Parent, Collection1_InitialTM);
+		TArray<FTransform> Collection2_InitialTM; GeometryCollectionAlgo::GlobalMatrices(DynamicCollection2->Transform, DynamicCollection2->Parent, Collection2_InitialTM);
+
+		TArray<Chaos::TPBDRigidClusteredParticleHandle<float, 3>*>& Collection1Handles = Collection1->PhysObject->GetSolverParticleHandles();
+		TArray<Chaos::TPBDRigidClusteredParticleHandle<float, 3>*>& Collection2Handles = Collection2->PhysObject->GetSolverParticleHandles();
+		const auto& SovlerParticleHandles = UnitTest.Solver->GetParticles().GetParticleHandles();
+		
+		EXPECT_EQ(SovlerParticleHandles.Size(), 6);
+		EXPECT_EQ(ClusterMap.Num(), 2);
+		EXPECT_TRUE(ClusterMapContains(ClusterMap, Collection1Handles[2], { Collection1Handles[1],Collection1Handles[0]}));
+		EXPECT_TRUE(ClusterMapContains(ClusterMap, Collection2Handles[2], { Collection2Handles[1],Collection2Handles[0]}));
+
+		UnitTest.Advance();
+
+		EXPECT_EQ(SovlerParticleHandles.Size(), 7);
+		EXPECT_EQ(ClusterMap.Num(), 1);
+		EXPECT_TRUE(ClusterMapContains(ClusterMap, SovlerParticleHandles.Handle(6)->CastToRigidParticle(), { Collection1Handles[1],Collection1Handles[0], Collection2Handles[1],Collection2Handles[0] }));
+
+		TArray<FTransform> Collection1_PreReleaseTM; GeometryCollectionAlgo::GlobalMatrices(DynamicCollection1->Transform, DynamicCollection1->Parent, Collection1_PreReleaseTM);
+		TArray<FTransform> Collection2_PreReleaseTM; GeometryCollectionAlgo::GlobalMatrices(DynamicCollection2->Transform, DynamicCollection2->Parent, Collection2_PreReleaseTM);
+		for (int Idx = 0; Idx < Collection1_PreReleaseTM.Num() - 1; Idx++) {
+			EXPECT_LT(Collection1_PreReleaseTM[Idx].GetTranslation().Z, Collection1_InitialTM[Idx].GetTranslation().Z);
+			EXPECT_LT(Collection2_PreReleaseTM[Idx].GetTranslation().Z, Collection2_InitialTM[Idx].GetTranslation().Z);
+		}
+
+		UnitTest.Solver->GetEvolution()->GetRigidClustering().ReleaseClusterParticles({ Collection1Handles[0],Collection1Handles[1] });
+
+		EXPECT_EQ(SovlerParticleHandles.Size(), 8);
+		EXPECT_EQ(ClusterMap.Num(), 1);
+		EXPECT_TRUE(ClusterMapContains(ClusterMap, SovlerParticleHandles.Handle(7)->CastToRigidParticle(), { Collection2Handles[1],Collection2Handles[0] }));
+
+		UnitTest.Advance();
+
+		EXPECT_EQ(SovlerParticleHandles.Size(), 8);
+		EXPECT_EQ(ClusterMap.Num(), 1);
+		EXPECT_TRUE(ClusterMapContains(ClusterMap, SovlerParticleHandles.Handle(7)->CastToRigidParticle(), {Collection2Handles[1],Collection2Handles[0] }));
+
+		TArray<FTransform> Collection1_PostReleaseTM; GeometryCollectionAlgo::GlobalMatrices(DynamicCollection1->Transform, DynamicCollection1->Parent, Collection1_PostReleaseTM);
+		TArray<FTransform> Collection2_PostReleaseTM; GeometryCollectionAlgo::GlobalMatrices(DynamicCollection2->Transform, DynamicCollection2->Parent, Collection2_PostReleaseTM);
+		for (int Idx = 0; Idx < Collection1_PostReleaseTM.Num() - 1; Idx++) {
+			EXPECT_LT(Collection1_PostReleaseTM[Idx].GetTranslation().Z, Collection1_PreReleaseTM[Idx].GetTranslation().Z);
+			EXPECT_LT(Collection2_PostReleaseTM[Idx].GetTranslation().Z, Collection2_PreReleaseTM[Idx].GetTranslation().Z);
+		}
+	}
+
+
+	
+	TYPED_TEST(AllTraits, GeometryCollection_RigidBodies_ClusterTest_ReleaseClusterParticle_ClusteredKinematicNode)
+	{		
+		if(bSkipGeomFailing)
+		{
+			return;
+		}
+		// Test : Build to geometry collections, cluster them together, release the sub bodies of the first collection. 
+//        ... should create a internal cluster with property transform mappings. 
+
+		using Traits = TypeParam;
+		TFramework<Traits> UnitTest;
+
+		TSharedPtr<FGeometryCollection> RestCollection = CreateClusteredBody(FVector(0, 0, 100));
+		CreationParameters Params;
+		Params.RestCollection = RestCollection;
+		Params.DynamicState = EObjectStateTypeEnum::Chaos_Object_Dynamic;
+		Params.ImplicitType = EImplicitTypeEnum::Chaos_Implicit_Sphere;
+		Params.CollisionType = ECollisionTypeEnum::Chaos_Volumetric;
+		Params.Simulating = true;
+		Params.EnableClustering = true;
+		Params.DamageThreshold = { FLT_MAX };
+		Params.ClusterGroupIndex = 1;
+		Params.ClusterConnectionMethod = Chaos::FClusterCreationParameters<FReal>::EConnectionMethod::DelaunayTriangulation;
+		TGeometryCollectionWrapper<Traits>* Collection1 = TNewSimulationObject<GeometryType::GeometryCollectionWithSuppliedRestCollection>::Init<Traits>(Params)->As<TGeometryCollectionWrapper<Traits>>();
+		UnitTest.AddSimulationObject(Collection1);
+
+
+		TSharedPtr<FGeometryCollection> RestCollection2 = CreateClusteredBody(FVector(0, 0, 200));
+		CreationParameters Params2;
+		Params2.RestCollection = RestCollection2;
+		Params2.DynamicState = EObjectStateTypeEnum::Chaos_Object_Dynamic;
+		Params2.ImplicitType = EImplicitTypeEnum::Chaos_Implicit_Sphere;
+		Params2.CollisionType = ECollisionTypeEnum::Chaos_Volumetric;
+		Params2.Simulating = true;
+		Params2.EnableClustering = true;
+		Params2.DamageThreshold = { FLT_MAX };
+		Params2.ClusterGroupIndex = 1;
+		Params2.ClusterConnectionMethod = Chaos::FClusterCreationParameters<FReal>::EConnectionMethod::DelaunayTriangulation;
+		TGeometryCollectionWrapper<Traits>* Collection2 = TNewSimulationObject<GeometryType::GeometryCollectionWithSuppliedRestCollection>::Init<Traits>(Params2)->As<TGeometryCollectionWrapper<Traits>>();
+		UnitTest.AddSimulationObject(Collection2);
+
+		TSharedPtr<FGeometryDynamicCollection> DynamicCollection1 = Collection1->DynamicCollection;
+		TSharedPtr<FGeometryDynamicCollection> DynamicCollection2 = Collection2->DynamicCollection;
+		DynamicCollection1->GetAttribute<int32>("DynamicState", FGeometryCollection::TransformGroup)[1] = (uint8)EObjectStateTypeEnum::Chaos_Object_Kinematic;
+
+		UnitTest.Initialize();
+
+		auto& Clustering = UnitTest.Solver->GetEvolution()->GetRigidClustering();
+		const auto& ClusterMap = Clustering.GetChildrenMap();
+
+
+
+		TArray<FTransform> Collection1_InitialTM; GeometryCollectionAlgo::GlobalMatrices(DynamicCollection1->Transform, DynamicCollection1->Parent, Collection1_InitialTM);
+		TArray<FTransform> Collection2_InitialTM; GeometryCollectionAlgo::GlobalMatrices(DynamicCollection2->Transform, DynamicCollection2->Parent, Collection2_InitialTM);
+
+		TArray<Chaos::TPBDRigidClusteredParticleHandle<float, 3>*>& Collection1Handles = Collection1->PhysObject->GetSolverParticleHandles();
+		TArray<Chaos::TPBDRigidClusteredParticleHandle<float, 3>*>& Collection2Handles = Collection2->PhysObject->GetSolverParticleHandles();
+		const auto& SovlerParticleHandles = UnitTest.Solver->GetParticles().GetParticleHandles();
+
+		EXPECT_EQ(SovlerParticleHandles.Size(), 6);
+		EXPECT_EQ(ClusterMap.Num(), 2);
+		EXPECT_TRUE(ClusterMapContains(ClusterMap, Collection1Handles[2], { Collection1Handles[1],Collection1Handles[0] }));
+		EXPECT_TRUE(ClusterMapContains(ClusterMap, Collection2Handles[2], { Collection2Handles[1],Collection2Handles[0] }));
+
+		UnitTest.Advance();
+
+		EXPECT_EQ(SovlerParticleHandles.Size(), 7);
+		EXPECT_EQ(ClusterMap.Num(), 1);
+		EXPECT_TRUE(ClusterMapContains(ClusterMap, SovlerParticleHandles.Handle(6)->CastToRigidParticle(), { Collection1Handles[1],Collection1Handles[0], Collection2Handles[1],Collection2Handles[0] }));
+
+		TArray<FTransform> Collection1_PreReleaseTM; GeometryCollectionAlgo::GlobalMatrices(DynamicCollection1->Transform, DynamicCollection1->Parent, Collection1_PreReleaseTM);
+		TArray<FTransform> Collection2_PreReleaseTM; GeometryCollectionAlgo::GlobalMatrices(DynamicCollection2->Transform, DynamicCollection2->Parent, Collection2_PreReleaseTM);
+		for (int Idx = 0; Idx < Collection1_PreReleaseTM.Num() - 1; Idx++) {
+			EXPECT_EQ(Collection1_PreReleaseTM[Idx].GetTranslation().Z, Collection1_InitialTM[Idx].GetTranslation().Z);
+			EXPECT_EQ(Collection2_PreReleaseTM[Idx].GetTranslation().Z, Collection2_InitialTM[Idx].GetTranslation().Z);
+		}
+
+		UnitTest.Solver->GetEvolution()->GetRigidClustering().ReleaseClusterParticles({ Collection1Handles[0],Collection1Handles[1] });
+
+		EXPECT_EQ(SovlerParticleHandles.Size(), 8);
+		EXPECT_EQ(ClusterMap.Num(), 1);
+		EXPECT_TRUE(ClusterMapContains(ClusterMap, SovlerParticleHandles.Handle(7)->CastToRigidParticle(), { Collection2Handles[1],Collection2Handles[0] }));
+
+		UnitTest.Advance();
+
+		EXPECT_EQ(SovlerParticleHandles.Size(), 8);
+		EXPECT_EQ(ClusterMap.Num(), 1);
+		EXPECT_TRUE(ClusterMapContains(ClusterMap, SovlerParticleHandles.Handle(7)->CastToRigidParticle(), { Collection2Handles[1],Collection2Handles[0] }));
+
+		// validate that DynamicCollection2 became dynamic and fell from the cluster. 
+
+		TArray<FTransform> Collection1_PostReleaseTM; GeometryCollectionAlgo::GlobalMatrices(DynamicCollection1->Transform, DynamicCollection1->Parent, Collection1_PostReleaseTM);
+		TArray<FTransform> Collection2_PostReleaseTM; GeometryCollectionAlgo::GlobalMatrices(DynamicCollection2->Transform, DynamicCollection2->Parent, Collection2_PostReleaseTM);
+		for (int Idx = 0; Idx < Collection1_PostReleaseTM.Num() - 1; Idx++) {
+			if( Idx==1)
+				EXPECT_EQ(Collection1_PostReleaseTM[Idx].GetTranslation().Z, Collection1_PreReleaseTM[Idx].GetTranslation().Z); // the original kinematic should be frozen
+			else
+				EXPECT_LT(Collection1_PostReleaseTM[Idx].GetTranslation().Z, Collection1_PreReleaseTM[Idx].GetTranslation().Z);
+			EXPECT_LT(Collection2_PostReleaseTM[Idx].GetTranslation().Z, Collection2_PreReleaseTM[Idx].GetTranslation().Z);
+		}
+	}
+
+	
+	TYPED_TEST(AllTraits, GeometryCollection_RigidBodies_ClusterTest_ReleaseClusterParticles_AllLeafNodes)
+	{
+		if(bSkipGeomFailing)
+		{
+			return;
+		}
+		// Release the leaf nodes of a cluster. This test exercises the clusters ability to deactivate from the bottom up. 
+
+		using Traits = TypeParam;
+		TFramework<Traits> UnitTest;
+
+		TSharedPtr<FGeometryCollection> RestCollection = CreateClusteredBody_TwoParents_TwoBodies(FVector(0, 0, 100));
+		CreationParameters Params;
+		Params.RestCollection = RestCollection;
+		Params.DynamicState = EObjectStateTypeEnum::Chaos_Object_Dynamic;
+		Params.ImplicitType = EImplicitTypeEnum::Chaos_Implicit_Box;
+		Params.CollisionType = ECollisionTypeEnum::Chaos_Surface_Volumetric;
+		Params.Simulating = true;
+		Params.EnableClustering = true;
+		Params.DamageThreshold = { FLT_MAX };
+		Params.MaxClusterLevel = 1;
+		Params.ClusterGroupIndex = 0;		
+		TGeometryCollectionWrapper<Traits>* Collection1 = TNewSimulationObject<GeometryType::GeometryCollectionWithSuppliedRestCollection>::Init<Traits>(Params)->As<TGeometryCollectionWrapper<Traits>>();
+		TSharedPtr<FGeometryDynamicCollection> DynamicCollection1 = Collection1->DynamicCollection;
+		DynamicCollection1->GetAttribute<int32>("DynamicState", FGeometryCollection::TransformGroup)[1] = (uint8)EObjectStateTypeEnum::Chaos_Object_Kinematic;
+
+		UnitTest.AddSimulationObject(Collection1);
+
+		UnitTest.Initialize();
+
+		auto& Clustering = UnitTest.Solver->GetEvolution()->GetRigidClustering();
+		const auto& ClusterMap = Clustering.GetChildrenMap();
+		TArray<FTransform> Collection1_InitialTM; GeometryCollectionAlgo::GlobalMatrices(DynamicCollection1->Transform, DynamicCollection1->Parent, Collection1_InitialTM);
+		TArray<Chaos::TPBDRigidClusteredParticleHandle<float, 3>*>& Collection1Handles = Collection1->PhysObject->GetSolverParticleHandles();
+		const auto& SovlerParticleHandles = UnitTest.Solver->GetParticles().GetParticleHandles();
+
+		EXPECT_EQ(SovlerParticleHandles.Size(), 4);
+		EXPECT_EQ(ClusterMap.Num(), 2);
+		EXPECT_TRUE(ClusterMapContains(ClusterMap, Collection1Handles[2], { Collection1Handles[1],Collection1Handles[0] }));
+		EXPECT_TRUE(ClusterMapContains(ClusterMap, Collection1Handles[3], { Collection1Handles[2] }));
+
+		UnitTest.Advance();
+
+		EXPECT_EQ(SovlerParticleHandles.Size(), 4);
+		EXPECT_EQ(ClusterMap.Num(), 2);
+		EXPECT_TRUE(ClusterMapContains(ClusterMap, Collection1Handles[2], { Collection1Handles[1],Collection1Handles[0] }));
+		EXPECT_TRUE(ClusterMapContains(ClusterMap, Collection1Handles[3], { Collection1Handles[2] }));
+
+		TArray<FTransform> Collection1_PreReleaseTM; GeometryCollectionAlgo::GlobalMatrices(DynamicCollection1->Transform, DynamicCollection1->Parent, Collection1_PreReleaseTM);
+		for (int Idx = 0; Idx < Collection1_PreReleaseTM.Num() - 1; Idx++) {
+			EXPECT_EQ(Collection1_PreReleaseTM[Idx].GetTranslation().Z, Collection1_InitialTM[Idx].GetTranslation().Z);
+		}
+
+		UnitTest.Solver->GetEvolution()->GetRigidClustering().ReleaseClusterParticles({ Collection1Handles[0],Collection1Handles[1] });
+
+		/*
+		EXPECT_EQ(SovlerParticleHandles.Size(), 8);
+		EXPECT_EQ(ClusterMap.Num(), 0);
+
+		UnitTest.Advance();
+
+		EXPECT_EQ(SovlerParticleHandles.Size(), 8);
+		EXPECT_EQ(ClusterMap.Num(), 0);
+
+		// validate that DynamicCollection1 BODY 2 became dynamic and fell from the cluster. 
+		TArray<FTransform> Collection1_PostReleaseTM; GeometryCollectionAlgo::GlobalMatrices(DynamicCollection1->Transform, DynamicCollection1->Parent, Collection1_PostReleaseTM);
+		for (int Idx = 0; Idx < Collection1_PostReleaseTM.Num() - 1; Idx++) {
+			if (Idx == 1)
+				EXPECT_EQ(Collection1_PostReleaseTM[Idx].GetTranslation().Z, Collection1_PreReleaseTM[Idx].GetTranslation().Z); // the original kinematic should be frozen
+			else
+				EXPECT_LT(Collection1_PostReleaseTM[Idx].GetTranslation().Z, Collection1_PreReleaseTM[Idx].GetTranslation().Z);
+		}
+		*/
+	}
+
+	
+	TYPED_TEST(AllTraits, GeometryCollection_RigidBodies_ClusterTest_ReleaseClusterParticles_ClusterNodeAndSubClusterNode)
+	{
+		if(bSkipGeomFailing)
+		{
+			return;
+		}
+		using Traits = TypeParam;
+		TFramework<Traits> UnitTest;
+
+		TSharedPtr<FGeometryCollection> RestCollection = CreateClusteredBody_TwoParents_TwoBodies(FVector(0, 0, 100));
+		CreationParameters Params;
+		Params.RestCollection = RestCollection;
+		Params.DynamicState = EObjectStateTypeEnum::Chaos_Object_Dynamic;
+		Params.ImplicitType = EImplicitTypeEnum::Chaos_Implicit_Box;
+		Params.CollisionType = ECollisionTypeEnum::Chaos_Surface_Volumetric;
+		Params.Simulating = true;
+		Params.EnableClustering = true;
+		Params.DamageThreshold = { FLT_MAX };
+		Params.MaxClusterLevel = 1;
+		Params.ClusterGroupIndex = 1;
+		TGeometryCollectionWrapper<Traits>* Collection = TNewSimulationObject<GeometryType::GeometryCollectionWithSuppliedRestCollection>::Init<Traits>(Params)->As<TGeometryCollectionWrapper<Traits>>();
+
+
+		TSharedPtr<FGeometryDynamicCollection> DynamicCollection = Collection->DynamicCollection;
+		DynamicCollection->GetAttribute<int32>("DynamicState", FGeometryCollection::TransformGroup)[1] = (uint8)EObjectStateTypeEnum::Chaos_Object_Kinematic;
+
+		UnitTest.AddSimulationObject(Collection);
+
+		TArray<Chaos::TPBDRigidClusteredParticleHandle<float, 3>*>& ParticleHandles = Collection->PhysObject->GetSolverParticleHandles();
+
+		UnitTest.Advance();
+
+		auto& Clustering = UnitTest.Solver->GetEvolution()->GetRigidClustering();
+		const auto& ClusterMap = Clustering.GetChildrenMap();
+		const Chaos::TArrayCollectionArray<Chaos::ClusterId>& ClusterIdsArray = Clustering.GetClusterIdsArray();
+
+		EXPECT_EQ(ClusterMap.Num(), 2);
+		EXPECT_TRUE(ClusterMapContains(ClusterMap, ParticleHandles[2], { ParticleHandles[0],ParticleHandles[1] }));
+		EXPECT_TRUE(ClusterMapContains(ClusterMap, ParticleHandles[4], { ParticleHandles[2] }));
+
+
+		for (int Frame = 1; Frame < 10; Frame++)
+		{
+			UnitTest.Advance();
+
+			/*
+			if (Frame == 5)
+			{
+				Solver->GetRigidClustering().ReleaseClusterParticles({ 2 });
+			}
+			*/
+
+			if (Frame < 5)
+			{
+
+				EXPECT_TRUE(ParticleHandles[2]->Disabled());
+				EXPECT_NE(ClusterIdsArray[2].Id, nullptr);
+				EXPECT_EQ(ClusterIdsArray[3].Id, nullptr);
+				EXPECT_EQ(ClusterIdsArray[4].Id, nullptr);
+			}
+			else
+			{
+				EXPECT_TRUE(!ParticleHandles[2]->Disabled());
+				EXPECT_EQ(ClusterIdsArray[2].Id, nullptr);
+				EXPECT_EQ(ClusterIdsArray[3].Id, nullptr);
+				EXPECT_EQ(ClusterIdsArray[4].Id, nullptr);
+
+				EXPECT_EQ(ClusterMap.Num(), 1);
+				EXPECT_TRUE(ClusterMapContains(ClusterMap, ParticleHandles[2], { ParticleHandles[0], ParticleHandles[1] }));
+			}				
+		}
+	}
+
+	
+	TYPED_TEST(AllTraits, GeometryCollection_RigidBodies_ClusterTest_RemoveOnFracture)
+	{
+		if(bSkipGeomFailing)
+		{
+			return;
+		}
+		using Traits = TypeParam;
+		TFramework<Traits> UnitTest;
+
+		TSharedPtr<FGeometryCollection> RestCollection = GeometryCollection::MakeCubeElement(FTransform(FQuat::MakeFromEuler(FVector(0, 0, 0.)), FVector(0, -10, 10)), FVector(1.0));
+		RestCollection->AppendGeometry(*GeometryCollection::MakeCubeElement(FTransform(FQuat::MakeFromEuler(FVector(0, 0, 0.)), FVector(0, 10, 10)), FVector(1.0)));
+		EXPECT_EQ(RestCollection->Transform.Num(), 2);
+
+		// this transform should have a zero scale after the simulation has run to the point of fracture
+		RestCollection->SetFlags(1, FGeometryCollection::FS_RemoveOnFracture);
+
+		FGeometryCollectionClusteringUtility::ClusterAllBonesUnderNewRoot(RestCollection.Get());
+		EXPECT_EQ(RestCollection->Transform.Num(), 3);
+		RestCollection->Transform[2] = FTransform(FQuat::MakeFromEuler(FVector(90.f, 0, 0.)), FVector(0, 0, 40));
+
+		CreationParameters Params;
+		Params.RestCollection = RestCollection;
+		Params.DynamicState = EObjectStateTypeEnum::Chaos_Object_Dynamic;
+		Params.ImplicitType = EImplicitTypeEnum::Chaos_Implicit_Box;
+		Params.CollisionType = ECollisionTypeEnum::Chaos_Surface_Volumetric;
+		Params.Simulating = true;
+		Params.EnableClustering = true;
+		Params.DamageThreshold = { 0.1f };		
+		Params.RemoveOnFractureEnabled = true;
+		TGeometryCollectionWrapper<Traits>* Collection = TNewSimulationObject<GeometryType::GeometryCollectionWithSuppliedRestCollection>::Init<Traits>(Params)->As<TGeometryCollectionWrapper<Traits>>();
+
+		FRadialFalloff * FalloffField = new FRadialFalloff();
+		FalloffField->Magnitude = 10.5;
+		FalloffField->Radius = 100.0;
+		FalloffField->Position = FVector(0.0, 0.0, 0.0);
+		FalloffField->Falloff = EFieldFalloffType::Field_FallOff_None;
+
+		FFieldSystemPhysicsProxy* FieldObject = new FFieldSystemPhysicsProxy(nullptr);
+
+		TSharedPtr<FGeometryDynamicCollection> DynamicCollection = Collection->DynamicCollection;
+		DynamicCollection->GetAttribute<int32>("DynamicState", FGeometryCollection::TransformGroup)[1] = (uint8)EObjectStateTypeEnum::Chaos_Object_Kinematic;
+
+		UnitTest.AddSimulationObject(Collection);
+
+		UnitTest.Solver->RegisterObject(FieldObject);
+
+		auto& Clustering = UnitTest.Solver->GetEvolution()->GetRigidClustering();
+		const auto& ClusterMap = Clustering.GetChildrenMap();
+
+		TManagedArray<FTransform>& Transform = Collection->DynamicCollection->Transform;
+		float StartingRigidDistance = (Transform[1].GetTranslation() - Transform[0].GetTranslation()).Size(), CurrentRigidDistance = 0.f;
+						
+		
+		
+		// #todo: is this even used?
+		/*
+		Chaos::TArrayCollectionArray<float>& InternalStrain = Clustering.GetStrainArray();
+		*/
+
+		FName TargetName = GetFieldPhysicsName(EFieldPhysicsType::Field_ExternalClusterStrain);
+		FFieldSystemCommand Command(TargetName, FalloffField->NewCopy());
+		FFieldSystemMetaDataProcessingResolution* ResolutionData = new FFieldSystemMetaDataProcessingResolution(EFieldResolutionType::Field_Resolution_Maximum);
+		Command.MetaData.Add(FFieldSystemMetaData::EMetaType::ECommandData_ProcessingResolution, TUniquePtr< FFieldSystemMetaDataProcessingResolution >(ResolutionData));
+		FieldObject->BufferCommand(UnitTest.Solver, Command);
+
+		FVector Scale = Transform[1].GetScale3D();
+
+		EXPECT_NEAR(Scale.X, 1.0f, SMALL_NUMBER);
+		EXPECT_NEAR(Scale.Y, 1.0f, SMALL_NUMBER);
+		EXPECT_NEAR(Scale.Z, 1.0f, SMALL_NUMBER);
+		
+		UnitTest.Advance();		
+
+		FieldObject->BufferCommand(UnitTest.Solver, { TargetName, FalloffField->NewCopy() });
+
+		UnitTest.Advance();		
+
+		FVector Scale2 = Transform[1].GetScale3D();
+		// geometry hidden by 0 scaling on transform
+		EXPECT_NEAR(Scale2.X, 0.0f, SMALL_NUMBER);
+		EXPECT_NEAR(Scale2.Y, 0.0f, SMALL_NUMBER);
+		EXPECT_NEAR(Scale2.Z, 0.0f, SMALL_NUMBER);
+		
+		delete FalloffField;
+		delete FieldObject;
+	
+	}
+
+	
+	TYPED_TEST(AllTraits, GeometryCollection_RigidBodiess_ClusterTest_ParticleImplicitCollisionGeometry)
+	{
+		if(bSkipGeomFailing)
+		{
+			return;
+		}
+		using Traits = TypeParam;
+		TFramework<Traits> UnitTest;
+
+		TSharedPtr<FGeometryCollection> RestCollection = CreateClusteredBody_FracturedGeometry();
+
+		CreationParameters Params;
+		Params.RestCollection = RestCollection;
+		Params.DynamicState = EObjectStateTypeEnum::Chaos_Object_Dynamic;
+		Params.ImplicitType = EImplicitTypeEnum::Chaos_Implicit_LevelSet;
+		Params.CollisionType = ECollisionTypeEnum::Chaos_Surface_Volumetric;
+		Params.Simulating = true;
+		Params.EnableClustering = true;				
+		Params.CollisionGroup = -1;
+		Params.MinLevelSetResolution = 15;
+		Params.MaxLevelSetResolution = 20;
+		TGeometryCollectionWrapper<Traits>* Collection = TNewSimulationObject<GeometryType::GeometryCollectionWithSuppliedRestCollection>::Init<Traits>(Params)->As<TGeometryCollectionWrapper<Traits>>();
+
+		UnitTest.AddSimulationObject(Collection);
+		
+		
+		// Todo: these aren't used anywhere in the test?
+		//typedef TUniquePtr<Chaos::FImplicitObject> FImplicitPointer;
+		//const TManagedArray<FImplicitPointer> & Implicits = RestCollection->template GetAttribute<FImplicitPointer>(FGeometryCollectionPhysicsProxy::ImplicitsAttribute, FTransformCollection::TransformGroup);
+
+		typedef TUniquePtr< FCollisionStructureManager::FSimplicial > FSimplicialPointer;
+		const TManagedArray<FSimplicialPointer> & Simplicials = RestCollection->template GetAttribute<FSimplicialPointer>(FGeometryDynamicCollection::SimplicialsAttribute, FTransformCollection::TransformGroup);
+
+		UnitTest.Advance();		
+
+		TArray<Chaos::TPBDRigidClusteredParticleHandle<float, 3>*>& ParticleHandles = Collection->PhysObject->GetSolverParticleHandles();
+
+		float CollisionParticlesPerObjectFractionDefault = 0.5f;
+		IConsoleVariable*  CVarCollisionParticlesPerObjectFractionDefault = IConsoleManager::Get().FindConsoleVariable(TEXT("p.CollisionParticlesPerObjectFractionDefault"));
+		EXPECT_NE(CVarCollisionParticlesPerObjectFractionDefault, nullptr);
+		if (CVarCollisionParticlesPerObjectFractionDefault != nullptr)
+		{
+			CollisionParticlesPerObjectFractionDefault = CVarCollisionParticlesPerObjectFractionDefault->GetFloat();
+		}
+		
+/*
+		todo: what is the replacement here?
+		EXPECT_EQ(Particles.CollisionParticles(Object->PhysicsProxy->RigidBodyIDArray_TestingAccess()[10])->Size(), (int)(Simplicials[10]->Size() * CollisionParticlesPerObjectFractionDefault));
+		EXPECT_EQ(Particles.CollisionParticles(Object->PhysicsProxy->RigidBodyIDArray_TestingAccess()[11])->Size(), (int)(Simplicials[11]->Size() * CollisionParticlesPerObjectFractionDefault));
+		EXPECT_EQ(Particles.CollisionParticles(Object->PhysicsProxy->RigidBodyIDArray_TestingAccess()[12])->Size(), (int)(Simplicials[12]->Size() * CollisionParticlesPerObjectFractionDefault));
+*/
+	}
+}
