Merge geometry collection into headless chaos.

diff --git a/Engine/Source/Programs/HeadlessChaos/Private/GeometryCollection/GeometryCollectionTestCollisionResolution.cpp b/Engine/Source/Programs/HeadlessChaos/Private/GeometryCollection/GeometryCollectionTestCollisionResolution.cpp
new file mode 100644
--- /dev/null
+++ b/Engine/Source/Programs/HeadlessChaos/Private/GeometryCollection/GeometryCollectionTestCollisionResolution.cpp
@@ -0,0 +1,519 @@
+// Copyright Epic Games, Inc. All Rights Reserved.
+
+#include "GeometryCollection/GeometryCollectionTestCollisionResolution.h"
+#include "GeometryCollection/GeometryCollectionTestFramework.h"
+#include "GeometryCollection/GeometryCollectionAlgo.h"
+#include "GeometryCollectionProxyData.h"
+#include "HeadlessChaosTestUtility.h"
+
+namespace GeometryCollectionTest
+{
+	using namespace ChaosTest;
+	TYPED_TEST(AllTraits, GeometryCollection_CollisionResolutionTest)
+	{
+		using Traits = TypeParam;
+		TFramework<Traits> UnitTest;
+
+		TGeometryCollectionWrapper<Traits>* Collection = nullptr;
+		{
+			FVector GlobalTranslation(0, 0, 10); FQuat GlobalRotation = FQuat::MakeFromEuler(FVector(0));
+			CreationParameters Params; Params.DynamicState = EObjectStateTypeEnum::Chaos_Object_Dynamic; Params.EnableClustering = false;
+			Params.ImplicitType = EImplicitTypeEnum::Chaos_Implicit_Sphere;  Params.SimplicialType = ESimplicialType::Chaos_Simplicial_Sphere; Params.CollisionType = ECollisionTypeEnum::Chaos_Volumetric;
+			Params.RootTransform = FTransform(GlobalRotation, GlobalTranslation); Params.NestedTransforms = { FTransform::Identity, FTransform::Identity, FTransform::Identity };
+			Collection = TNewSimulationObject<GeometryType::GeometryCollectionWithSingleRigid>::Init<Traits>(Params)->As<TGeometryCollectionWrapper<Traits>>();
+			EXPECT_EQ(Collection->DynamicCollection->Parent[0], 1); // is a child of index one
+			EXPECT_TRUE(Collection->DynamicCollection->MassToLocal[0].Equals(FTransform::Identity)); // we are not testing MassToLocal in this test
+			UnitTest.AddSimulationObject(Collection);
+		}
+
+		RigidBodyWrapper* Floor = TNewSimulationObject<GeometryType::RigidFloor>::Init<Traits>()->As<RigidBodyWrapper>();
+		UnitTest.AddSimulationObject(Floor);
+
+		UnitTest.Initialize();
+		for (int i = 0; i < 10000; i++) 
+		{
+			UnitTest.Advance();
+		}
+		{
+			// validate that Simplicials are null when CollisionType==Chaos_Volumetric
+			EXPECT_EQ(Collection->DynamicCollection->Transform.Num(), 4);
+			EXPECT_EQ(Collection->DynamicCollection->Simplicials[0], nullptr);
+			EXPECT_EQ(Collection->DynamicCollection->Simplicials[1], nullptr);
+			EXPECT_EQ(Collection->DynamicCollection->Simplicials[2], nullptr);
+			EXPECT_EQ(Collection->DynamicCollection->Simplicials[3], nullptr);
+			EXPECT_EQ(UnitTest.Solver->GetParticles().GetGeometryCollectionParticles().CollisionParticles(0), nullptr);
+
+			EXPECT_LT(FMath::Abs(Collection->RestCollection->Transform[0].GetTranslation().Z)-10.f, KINDA_SMALL_NUMBER);
+			EXPECT_LT(FMath::Abs(Collection->DynamicCollection->Transform[0].GetTranslation().Z - 1.0), 0.1);
+		}
+	}
+
+	TYPED_TEST(AllTraits,GeometryCollection_CollisionResolution_SimplicialSphereToPlane)
+	{
+		using Traits = TypeParam;
+		TFramework<Traits> UnitTest;
+
+		TGeometryCollectionWrapper<Traits>* Collection = nullptr;
+		{
+			FVector GlobalTranslation(0, 0, 10); FQuat GlobalRotation = FQuat::MakeFromEuler(FVector(0));
+			CreationParameters Params; Params.DynamicState = EObjectStateTypeEnum::Chaos_Object_Dynamic; Params.EnableClustering = false;
+			Params.ImplicitType = EImplicitTypeEnum::Chaos_Implicit_Sphere;  Params.SimplicialType = ESimplicialType::Chaos_Simplicial_Sphere; Params.CollisionType = ECollisionTypeEnum::Chaos_Surface_Volumetric;
+			Params.RootTransform = FTransform(GlobalRotation, GlobalTranslation); Params.NestedTransforms = { FTransform::Identity, FTransform::Identity, FTransform::Identity };
+			Collection = TNewSimulationObject<GeometryType::GeometryCollectionWithSingleRigid>::Init<Traits>(Params)->As<TGeometryCollectionWrapper<Traits>>();
+			EXPECT_EQ(Collection->DynamicCollection->Parent[0], 1); // is a child of index one
+			EXPECT_TRUE(Collection->DynamicCollection->MassToLocal[0].Equals(FTransform::Identity)); // we are not testing MassToLocal in this test
+
+			UnitTest.AddSimulationObject(Collection);
+		}
+
+		RigidBodyWrapper* Floor = TNewSimulationObject<GeometryType::RigidFloor>::Init<Traits>()->As<RigidBodyWrapper>();
+		UnitTest.AddSimulationObject(Floor);
+
+		UnitTest.Initialize();
+		for (int i = 0; i < 10000; i++)
+		{
+			UnitTest.Advance();
+		}
+		{
+			// validate that Simplicials are null when CollisionType==Chaos_Volumetric
+			EXPECT_EQ(Collection->DynamicCollection->Transform.Num(), 4);
+			EXPECT_TRUE(Collection->DynamicCollection->Simplicials[0]!=nullptr);
+			EXPECT_TRUE(UnitTest.Solver->GetParticles().GetGeometryCollectionParticles().CollisionParticles(0)!=nullptr);
+			EXPECT_TRUE(Collection->DynamicCollection->Simplicials[0]->Size() == UnitTest.Solver->GetParticles().GetGeometryCollectionParticles().CollisionParticles(0)->Size());
+			EXPECT_TRUE(Collection->DynamicCollection->Simplicials[0]->Size() != 0);
+			int32 Size = Collection->DynamicCollection->Simplicials[0]->Size();
+
+			EXPECT_LT(FMath::Abs(Collection->RestCollection->Transform[0].GetTranslation().Z) - 10.f, KINDA_SMALL_NUMBER);
+			EXPECT_LT(FMath::Abs(Collection->DynamicCollection->Transform[0].GetTranslation().Z - 1.0), 0.04); // ball settles close to 1.0
+		}
+	}
+
+
+
+	TYPED_TEST(AllTraits,GeometryCollection_CollisionResolution_AnalyticSphereToAnalyticSphere)
+	{
+		using Traits = TypeParam;
+		// simplicial sphere to implicit sphere
+		TFramework<Traits> UnitTest;
+
+		CreationParameters Params;
+		Params.CollisionType = ECollisionTypeEnum::Chaos_Volumetric;
+
+		Params.EnableClustering = false;
+
+		FVector Scale(1.f);
+		Params.GeomTransform.SetScale3D(Scale); // Sphere radius
+		Params.EnableClustering = false;
+
+		// Make a dynamic simplicial sphere
+		Params.SimplicialType = ESimplicialType::Chaos_Simplicial_Sphere;
+		//Params.ImplicitType = EImplicitTypeEnum::Chaos_Implicit_None; // Fails, falls right through
+		Params.ImplicitType = EImplicitTypeEnum::Chaos_Implicit_Sphere;
+
+		Params.DynamicState = EObjectStateTypeEnum::Chaos_Object_Dynamic;
+		Params.RootTransform =
+			FTransform(FQuat::MakeFromEuler(FVector(0)), FVector(0, 0, 3.0));
+		TGeometryCollectionWrapper<Traits>* SimplicialSphereCollection =
+			TNewSimulationObject<GeometryType::GeometryCollectionWithSingleRigid>::Init<Traits>(Params)->As<TGeometryCollectionWrapper<Traits>>();
+		UnitTest.AddSimulationObject(SimplicialSphereCollection);
+
+		// Make a kinematic implicit sphere
+		Params.SimplicialType = ESimplicialType::Chaos_Simplicial_Sphere;
+		Params.ImplicitType = EImplicitTypeEnum::Chaos_Implicit_Sphere;
+
+		Params.DynamicState = EObjectStateTypeEnum::Chaos_Object_Kinematic;
+		Params.RootTransform =
+			FTransform(FQuat::MakeFromEuler(FVector(0)), FVector(0));
+		TGeometryCollectionWrapper<Traits>* ImplicitSphereCollection =
+			TNewSimulationObject<GeometryType::GeometryCollectionWithSingleRigid>::Init<Traits>(Params)->As<TGeometryCollectionWrapper<Traits>>();
+		UnitTest.AddSimulationObject(ImplicitSphereCollection);
+
+		// Hard code masstolocal on rest collection to identity
+		{
+			TManagedArray<FTransform>& MassToLocal =
+				SimplicialSphereCollection->RestCollection->GetAttribute<FTransform>(
+					TEXT("MassToLocal"), FTransformCollection::TransformGroup);
+			check(MassToLocal.Num() == 1);
+			MassToLocal[0] = FTransform::Identity;
+		}
+		{
+			TManagedArray<FTransform>& MassToLocal =
+				ImplicitSphereCollection->RestCollection->GetAttribute<FTransform>(
+					TEXT("MassToLocal"), FTransformCollection::TransformGroup);
+			check(MassToLocal.Num() == 1);
+			MassToLocal[0] = FTransform::Identity;
+		}
+
+		UnitTest.Initialize();
+		EXPECT_EQ(
+			SimplicialSphereCollection->DynamicCollection->Transform[0].GetTranslation().Z,
+			ImplicitSphereCollection->DynamicCollection->Transform[0].GetTranslation().Z + 3);
+
+		const FVector FirstX = SimplicialSphereCollection->DynamicCollection->Transform[0].GetTranslation();
+		FVector PrevX = FirstX;
+		for (int i = 0; i < 10; i++)
+		{
+			UnitTest.Advance();
+
+			const FVector& CurrX = SimplicialSphereCollection->DynamicCollection->Transform[0].GetTranslation();
+			EXPECT_NE(CurrX.Z, FirstX.Z); // moved since init
+			EXPECT_LE(CurrX.Z, PrevX.Z); // falling in -Z, or stopped
+			EXPECT_LE(FMath::Abs(CurrX.X), KINDA_SMALL_NUMBER); // straight down
+			EXPECT_LE(FMath::Abs(CurrX.Y), KINDA_SMALL_NUMBER); // straight down
+			PrevX = CurrX;
+		}
+
+		{
+			// We expect the simplical sphere to drop by 0.1 in Z and come to rest
+			// on top of the implicit sphere.
+			const FVector& CurrX = SimplicialSphereCollection->DynamicCollection->Transform[0].GetTranslation();
+			EXPECT_LE(CurrX.Z - 2.0, 0.2); // Relative large fudge factor accounts for aliasing?
+		}
+	}
+
+	TYPED_TEST(AllTraits,DISABLED_GeometryCollection_CollisionResolution_AnalyticCubeToAnalyticCube)
+	{
+		using Traits = TypeParam;
+
+		// simplicial sphere to implicit sphere
+		TFramework<Traits> UnitTest;
+
+		CreationParameters Params;
+		Params.EnableClustering = false;
+
+		FVector Scale(1.f, 1.f, 1.f);
+		Params.GeomTransform.SetScale3D(Scale); // Box dimensions
+		Params.EnableClustering = false;
+
+		// Make a dynamic box
+		Params.SimplicialType = ESimplicialType::Chaos_Simplicial_Box;
+		Params.ImplicitType = EImplicitTypeEnum::Chaos_Implicit_Box;
+
+		Params.DynamicState = EObjectStateTypeEnum::Chaos_Object_Dynamic;
+		Params.RootTransform = FTransform(FQuat::MakeFromEuler(FVector(0)), FVector(0, 0, 3.0));
+		TGeometryCollectionWrapper<Traits>* BoxCollection0 =
+			TNewSimulationObject<GeometryType::GeometryCollectionWithSingleRigid>::Init<Traits>(Params)->As<TGeometryCollectionWrapper<Traits>>();
+		UnitTest.AddSimulationObject(BoxCollection0);
+
+		// Make a kinematic box
+		Params.SimplicialType = ESimplicialType::Chaos_Simplicial_Box;
+		Params.ImplicitType = EImplicitTypeEnum::Chaos_Implicit_Box;
+
+		Params.DynamicState = EObjectStateTypeEnum::Chaos_Object_Kinematic;
+		Params.RootTransform = FTransform(FQuat::MakeFromEuler(FVector(0)), FVector(0));
+		TGeometryCollectionWrapper<Traits>* BoxCollection1 =
+			TNewSimulationObject<GeometryType::GeometryCollectionWithSingleRigid>::Init<Traits>(Params)->As<TGeometryCollectionWrapper<Traits>>();
+		UnitTest.AddSimulationObject(BoxCollection1);
+/*
+		// Hard code masstolocal on rest collection to identity
+		{
+			TManagedArray<FTransform>& MassToLocal =
+				BoxCollection0->RestCollection->GetAttribute<FTransform>(
+					TEXT("MassToLocal"), FTransformCollection::TransformGroup);
+			check(MassToLocal.Num() == 1);
+			MassToLocal[0] = FTransform::Identity;
+		}
+		{
+			TManagedArray<FTransform>& MassToLocal =
+				BoxCollection1->RestCollection->GetAttribute<FTransform>(
+					TEXT("MassToLocal"), FTransformCollection::TransformGroup);
+			check(MassToLocal.Num() == 1);
+			MassToLocal[0] = FTransform::Identity;
+		}
+*/
+		UnitTest.Initialize();
+		EXPECT_EQ(
+			BoxCollection0->DynamicCollection->Transform[0].GetTranslation().Z,
+			BoxCollection1->DynamicCollection->Transform[0].GetTranslation().Z + 3);
+
+		const FVector FirstX = BoxCollection0->DynamicCollection->Transform[0].GetTranslation();
+		FVector PrevX = FirstX;
+		for (int i = 0; i < 10; i++)
+		{
+			UnitTest.Advance();
+
+			const FVector& CurrX = BoxCollection0->DynamicCollection->Transform[0].GetTranslation();
+			EXPECT_NE(CurrX.Z, FirstX.Z); // moved since init
+			EXPECT_LE(CurrX.Z, PrevX.Z); // falling in -Z, or stopped
+			EXPECT_LE(FMath::Abs(CurrX.X), KINDA_SMALL_NUMBER); // No deflection
+			EXPECT_LE(FMath::Abs(CurrX.Y), KINDA_SMALL_NUMBER); // No deflection
+			PrevX = CurrX;
+		}
+
+		{
+			// We expect the simplical sphere to drop by 0.1 in Z and come to rest
+			// on top of the implicit sphere.
+			const FVector& CurrX = BoxCollection0->DynamicCollection->Transform[0].GetTranslation();
+			EXPECT_LE(CurrX.Z - 2.0, 0.2); // Relative large fudge factor accounts for aliasing?
+		}
+
+	}
+
+
+	template<typename Traits>
+	void CollisionResolution_SimplicialCubeToAnalyticCube()
+	{
+
+	}
+
+
+	TYPED_TEST(AllTraits,GeometryCollection_CollisionResolution_SimplicialSphereToAnalyticSphere)
+	{
+		using Traits = TypeParam;
+		TFramework<Traits> UnitTest;
+		// This should exercise CollisionResolution::ConstructLevelsetLevelsetConstraints(...) with ispc:SampleSphere* (Paticle to Analytic Sphere)
+
+		TGeometryCollectionWrapper<Traits>* Collection = nullptr;
+		{
+			FVector GlobalTranslation(0, 0, 10); FQuat GlobalRotation = FQuat::MakeFromEuler(FVector(0));
+			CreationParameters Params; Params.DynamicState = EObjectStateTypeEnum::Chaos_Object_Dynamic; Params.EnableClustering = false;
+			Params.ImplicitType = EImplicitTypeEnum::Chaos_Implicit_LevelSet;  Params.SimplicialType = ESimplicialType::Chaos_Simplicial_Sphere; Params.CollisionType = ECollisionTypeEnum::Chaos_Surface_Volumetric;
+			Params.RootTransform = FTransform(GlobalRotation, GlobalTranslation); Params.NestedTransforms = { FTransform::Identity, FTransform::Identity, FTransform::Identity };
+			Collection = TNewSimulationObject<GeometryType::GeometryCollectionWithSingleRigid>::Init<Traits>(Params)->As<TGeometryCollectionWrapper<Traits>>();
+			EXPECT_EQ(Collection->DynamicCollection->Parent[0], 1); // is a child of index one
+			EXPECT_TRUE(Collection->DynamicCollection->MassToLocal[0].Equals(FTransform::Identity)); // we are not testing MassToLocal in this test
+
+			UnitTest.AddSimulationObject(Collection);
+		}
+
+		TGeometryCollectionWrapper<Traits>* CollectionStaticSphere = nullptr;
+		{
+			FVector GlobalTranslation(0, 0, 0); FQuat GlobalRotation = FQuat::MakeFromEuler(FVector(0));
+			CreationParameters Params; Params.DynamicState = EObjectStateTypeEnum::Chaos_Object_Static; Params.EnableClustering = false;
+			Params.ImplicitType = EImplicitTypeEnum::Chaos_Implicit_Sphere;  Params.SimplicialType = ESimplicialType::Chaos_Simplicial_Sphere; Params.CollisionType = ECollisionTypeEnum::Chaos_Volumetric;
+			Params.RootTransform = FTransform(GlobalRotation, GlobalTranslation); Params.NestedTransforms = { FTransform::Identity, FTransform::Identity, FTransform::Identity };
+			CollectionStaticSphere = TNewSimulationObject<GeometryType::GeometryCollectionWithSingleRigid>::Init<Traits>(Params)->As<TGeometryCollectionWrapper<Traits>>();
+			EXPECT_EQ(CollectionStaticSphere->DynamicCollection->Parent[0], 1); // is a child of index one
+			EXPECT_TRUE(CollectionStaticSphere->DynamicCollection->MassToLocal[0].Equals(FTransform::Identity)); // we are not testing MassToLocal in this test
+
+			UnitTest.AddSimulationObject(CollectionStaticSphere);
+		}
+
+		UnitTest.Initialize();
+
+		for (int i = 0; i < 20; i++)
+		{
+			UnitTest.Advance();
+		}
+		{
+			// validate simplicials and implicits are configured correctly
+			EXPECT_EQ(Collection->DynamicCollection->Transform.Num(), 4);
+			EXPECT_TRUE(Collection->DynamicCollection->Simplicials[0] != nullptr);
+			EXPECT_TRUE(UnitTest.Solver->GetParticles().GetGeometryCollectionParticles().CollisionParticles(0) != nullptr);
+			EXPECT_TRUE(Collection->DynamicCollection->Simplicials[0]->Size() == UnitTest.Solver->GetParticles().GetGeometryCollectionParticles().CollisionParticles(0)->Size());
+			EXPECT_TRUE(Collection->DynamicCollection->Simplicials[0]->Size() != 0);
+			EXPECT_TRUE(Collection->DynamicCollection->Implicits[0]->GetType() == (int32)Chaos::ImplicitObjectType::LevelSet);
+
+			EXPECT_EQ(CollectionStaticSphere->DynamicCollection->Transform.Num(), 4);
+			EXPECT_TRUE(CollectionStaticSphere->DynamicCollection->Simplicials[0] == nullptr);
+			EXPECT_TRUE(CollectionStaticSphere->DynamicCollection->Implicits[0]->GetType() == (int32)Chaos::ImplicitObjectType::Sphere);
+
+			// validate the ball collides and moved away from the static ball
+			EXPECT_LT(FMath::Abs(Collection->RestCollection->Transform[0].GetTranslation().Z) - 10.f, KINDA_SMALL_NUMBER);
+			EXPECT_TRUE(FMath::Abs(Collection->DynamicCollection->Transform[0].GetTranslation().X) < KINDA_SMALL_NUMBER); // No deflection
+			EXPECT_TRUE(FMath::Abs(Collection->DynamicCollection->Transform[0].GetTranslation().Y) < KINDA_SMALL_NUMBER); // No deflection
+			EXPECT_LT(Collection->DynamicCollection->Transform[0].GetTranslation().Z, 0.f); // ball fell
+		}
+	}
+
+
+	TYPED_TEST(AllTraits,GeometryCollection_CollisionResolution_SimplicialSphereToImplicitSphere)
+	{
+		using Traits = TypeParam;
+		// simplicial sphere to implicit sphere
+		TFramework<Traits> UnitTest;
+
+		CreationParameters Params; 
+		Params.EnableClustering = false;
+
+		FVector Scale(1.f);
+		Params.GeomTransform.SetScale3D(Scale); // Sphere radius
+		Params.EnableClustering = false;
+
+		// Make a dynamic simplicial sphere
+		Params.SimplicialType = ESimplicialType::Chaos_Simplicial_Sphere;
+		//Params.ImplicitType = EImplicitTypeEnum::Chaos_Implicit_None; // Fails, falls right through
+		Params.ImplicitType = EImplicitTypeEnum::Chaos_Implicit_Sphere;
+
+		Params.DynamicState = EObjectStateTypeEnum::Chaos_Object_Dynamic;
+		Params.RootTransform =
+			FTransform(FQuat::MakeFromEuler(FVector(0)), FVector(0, 0, 3.0));
+		TGeometryCollectionWrapper<Traits>* SimplicialSphereCollection =
+			TNewSimulationObject<GeometryType::GeometryCollectionWithSingleRigid>::Init<Traits>(Params)->As<TGeometryCollectionWrapper<Traits>>();
+		UnitTest.AddSimulationObject(SimplicialSphereCollection);
+
+		// Make a kinematic implicit sphere
+		Params.SimplicialType = ESimplicialType::Chaos_Simplicial_Sphere;
+		Params.ImplicitType = EImplicitTypeEnum::Chaos_Implicit_LevelSet;
+
+		Params.DynamicState = EObjectStateTypeEnum::Chaos_Object_Kinematic;
+		Params.RootTransform =
+			FTransform(FQuat::MakeFromEuler(FVector(0)), FVector(0));
+		TGeometryCollectionWrapper<Traits>* ImplicitSphereCollection =
+			TNewSimulationObject<GeometryType::GeometryCollectionWithSingleRigid>::Init<Traits>(Params)->As<TGeometryCollectionWrapper<Traits>>();
+		UnitTest.AddSimulationObject(ImplicitSphereCollection);
+
+		// Hard code masstolocal on rest collection to identity
+		{
+			TManagedArray<FTransform>& MassToLocal =
+				SimplicialSphereCollection->RestCollection->GetAttribute<FTransform>(
+					TEXT("MassToLocal"), FTransformCollection::TransformGroup);
+			check(MassToLocal.Num() == 1);
+			MassToLocal[0] = FTransform::Identity;
+		}
+		{
+			TManagedArray<FTransform>& MassToLocal =
+				ImplicitSphereCollection->RestCollection->GetAttribute<FTransform>(
+					TEXT("MassToLocal"), FTransformCollection::TransformGroup);
+			check(MassToLocal.Num() == 1);
+			MassToLocal[0] = FTransform::Identity;
+		}
+
+		UnitTest.Initialize();
+		EXPECT_EQ(
+			SimplicialSphereCollection->DynamicCollection->Transform[0].GetTranslation().Z,
+			ImplicitSphereCollection->DynamicCollection->Transform[0].GetTranslation().Z + 3);
+
+		const FVector FirstX = SimplicialSphereCollection->DynamicCollection->Transform[0].GetTranslation();
+		FVector PrevX = FirstX;
+		for (int i = 0; i < 10; i++)
+		{
+			UnitTest.Advance();
+
+			const FVector& CurrX = SimplicialSphereCollection->DynamicCollection->Transform[0].GetTranslation();
+			EXPECT_NE(CurrX.Z, FirstX.Z); // moved since init
+			EXPECT_LE(FMath::Abs(CurrX.X), KINDA_SMALL_NUMBER); // straight down
+			EXPECT_LE(FMath::Abs(CurrX.Y), KINDA_SMALL_NUMBER); // straight down
+			PrevX = CurrX;
+		}
+		
+		{
+			// We expect the simplical sphere to drop by 0.1 in Z and come to rest
+			// on top of the implicit sphere.
+			const FVector& CurrX = SimplicialSphereCollection->DynamicCollection->Transform[0].GetTranslation();
+			EXPECT_LE(CurrX.Z - 2.0, 0.2); // Relative large fudge factor accounts for aliasing?
+		}
+	}
+
+
+
+	TYPED_TEST(AllTraits,GeometryCollection_CollisionResolution_SimplicialCubeToImplicitCube)
+	{
+		using Traits = TypeParam;
+
+		// simplicial sphere to implicit sphere
+		TFramework<Traits> UnitTest;
+
+		CreationParameters Params;
+		Params.EnableClustering = false;
+
+		FVector Scale(1.f, 1.f, 1.f);
+		Params.GeomTransform.SetScale3D(Scale); // Box dimensions
+		Params.EnableClustering = false;
+
+		// Make a dynamic box
+		Params.SimplicialType = ESimplicialType::Chaos_Simplicial_Box;
+		Params.ImplicitType = EImplicitTypeEnum::Chaos_Implicit_Box;
+
+		Params.DynamicState = EObjectStateTypeEnum::Chaos_Object_Dynamic;
+		Params.RootTransform = FTransform(FQuat::MakeFromEuler(FVector(0)), FVector(0, 0, 3.0));
+		TGeometryCollectionWrapper<Traits>* BoxCollection0 =
+			TNewSimulationObject<GeometryType::GeometryCollectionWithSingleRigid>::Init<Traits>(Params)->As<TGeometryCollectionWrapper<Traits>>();
+		UnitTest.AddSimulationObject(BoxCollection0);
+
+		// Make a kinematic box
+		Params.SimplicialType = ESimplicialType::Chaos_Simplicial_Box;
+		Params.ImplicitType = EImplicitTypeEnum::Chaos_Implicit_LevelSet;
+
+		Params.DynamicState = EObjectStateTypeEnum::Chaos_Object_Kinematic;
+		Params.RootTransform = FTransform(FQuat::MakeFromEuler(FVector(0)), FVector(0));
+		TGeometryCollectionWrapper<Traits>* BoxCollection1 =
+			TNewSimulationObject<GeometryType::GeometryCollectionWithSingleRigid>::Init<Traits>(Params)->As<TGeometryCollectionWrapper<Traits>>();
+		UnitTest.AddSimulationObject(BoxCollection1);
+/*
+		// Hard code masstolocal on rest collection to identity
+		{
+			TManagedArray<FTransform>& MassToLocal =
+				BoxCollection0->RestCollection->GetAttribute<FTransform>(
+					TEXT("MassToLocal"), FTransformCollection::TransformGroup);
+			check(MassToLocal.Num() == 1);
+			MassToLocal[0] = FTransform::Identity;
+		}
+		{
+			TManagedArray<FTransform>& MassToLocal =
+				BoxCollection1->RestCollection->GetAttribute<FTransform>(
+					TEXT("MassToLocal"), FTransformCollection::TransformGroup);
+			check(MassToLocal.Num() == 1);
+			MassToLocal[0] = FTransform::Identity;
+		}
+*/
+		UnitTest.Initialize();
+		EXPECT_EQ(
+			BoxCollection0->DynamicCollection->Transform[0].GetTranslation().Z,
+			BoxCollection1->DynamicCollection->Transform[0].GetTranslation().Z + 3);
+
+		const FVector FirstX = BoxCollection0->DynamicCollection->Transform[0].GetTranslation();
+		FVector PrevX = FirstX;
+		for (int i = 0; i < 10; i++)
+		{
+			UnitTest.Advance();
+
+			const FVector& CurrX = BoxCollection0->DynamicCollection->Transform[0].GetTranslation();
+			EXPECT_NE(CurrX.Z, FirstX.Z); // moved since init
+			EXPECT_LE(CurrX.Z, PrevX.Z); // falling in -Z, or stopped
+			EXPECT_LE(FMath::Abs(CurrX.X), KINDA_SMALL_NUMBER); // straight down
+			EXPECT_LE(FMath::Abs(CurrX.Y), KINDA_SMALL_NUMBER); // straight down
+			PrevX = CurrX;
+		}
+
+		{
+			// We expect the simplical sphere to drop by 0.1 in Z and come to rest
+			// on top of the implicit sphere.
+			const FVector& CurrX = BoxCollection0->DynamicCollection->Transform[0].GetTranslation();
+			EXPECT_LE(CurrX.Z - 2.0, 0.2); // Relative large fudge factor accounts for aliasing?
+		}
+	}
+
+
+
+
+
+	TYPED_TEST(AllTraits,GeometryCollection_CollisionResolution_SimplicialTetrahedronWithNonUniformMassToFloor)
+	{
+		using Traits = TypeParam;
+		TFramework<Traits> UnitTest;
+
+		TGeometryCollectionWrapper<Traits>* Collection = nullptr;
+		{
+			FVector GlobalTranslation(0, 0, 10); FQuat GlobalRotation = FQuat::MakeFromEuler(FVector(0));
+			CreationParameters Params; Params.DynamicState = EObjectStateTypeEnum::Chaos_Object_Dynamic; Params.EnableClustering = false;
+			Params.ImplicitType = EImplicitTypeEnum::Chaos_Implicit_LevelSet;  Params.SimplicialType = ESimplicialType::Chaos_Simplicial_Tetrahedron; Params.CollisionType = ECollisionTypeEnum::Chaos_Surface_Volumetric;
+			Params.GeomTransform = FTransform(GlobalRotation, GlobalTranslation);
+			Collection = TNewSimulationObject<GeometryType::GeometryCollectionWithSingleRigid>::Init<Traits>(Params)->As<TGeometryCollectionWrapper<Traits>>();
+			EXPECT_EQ(Collection->DynamicCollection->Parent[0], -1); // is a child of index one
+			EXPECT_NEAR((Collection->DynamicCollection->MassToLocal[0].GetTranslation()-FVector(0,0,10)).Size(),0,KINDA_SMALL_NUMBER);
+
+			UnitTest.AddSimulationObject(Collection);
+		}
+
+		RigidBodyWrapper* Floor = TNewSimulationObject<GeometryType::RigidFloor>::Init<Traits>()->As<RigidBodyWrapper>();
+		UnitTest.AddSimulationObject(Floor);
+
+		UnitTest.Initialize();
+
+		for (int i = 0; i < 20; i++)
+		{
+			UnitTest.Advance();
+		}
+		{
+			// validate the ball collides and moved away from the static ball
+			EXPECT_EQ(Collection->RestCollection->Transform[0].GetTranslation().Z, 0.f);
+			EXPECT_NEAR(FMath::Abs(Collection->DynamicCollection->Transform[0].GetTranslation().X), 0.f, 0.01);
+			EXPECT_NEAR(FMath::Abs(Collection->DynamicCollection->Transform[0].GetTranslation().Y), 0.f, 0.01);
+			EXPECT_NEAR(Collection->DynamicCollection->Transform[0].GetTranslation().Z, -9.f, 0.1);
+		}
+	}
+
+
+
+} // namespace GeometryCollectionTest
diff --git a/Engine/Source/Programs/HeadlessChaos/Private/GeometryCollection/GeometryCollectionTestFields.cpp b/Engine/Source/Programs/HeadlessChaos/Private/GeometryCollection/GeometryCollectionTestFields.cpp
new file mode 100644
--- /dev/null
+++ b/Engine/Source/Programs/HeadlessChaos/Private/GeometryCollection/GeometryCollectionTestFields.cpp
@@ -0,0 +1,1231 @@
+// Copyright Epic Games, Inc. All Rights Reserved.
+
+#include "GeometryCollection/GeometryCollectionTestFields.h"
+
+#include "Chaos/Utilities.h"
+#include "Field/FieldSystem.h"
+#include "Field/FieldSystemNodes.h"
+#include "Field/FieldSystemCoreAlgo.h"
+#include "Field/FieldSystemNoiseAlgo.h"
+#include "GeometryCollection/GeometryCollection.h"
+#include "GeometryCollection/GeometryCollectionUtility.h"
+#include "GeometryCollection/GeometryCollectionAlgo.h"
+#include "GeometryCollection/TransformCollection.h"
+#include "Misc/FileHelper.h"
+#include "Serialization/BufferArchive.h"
+#include "Serialization/MemoryReader.h"
+#include "UObject/Package.h"
+#include "UObject/UObjectGlobals.h"
+
+DEFINE_LOG_CATEGORY_STATIC(GCTF_Log, Verbose, All);
+
+namespace GeometryCollectionTest
+{
+	void Fields_NoiseSample()
+	{
+		int Bounds = 100;
+		TArray<ContextIndex> IndicesArray;
+		ContextIndex::ContiguousIndices(IndicesArray, Bounds*Bounds);
+		TArrayView<ContextIndex> IndexView(&(IndicesArray[0]), IndicesArray.Num());
+
+		TArray<FVector> SamplesArray;
+		SamplesArray.Init(FVector(0.f), Bounds*Bounds);
+		for (int32 Index = 0, i=0; i < Bounds; i++)
+		{
+			for (int32 j = 0; j < Bounds; j++)
+			{
+				SamplesArray[Index] = 1000.*FVector(i,j,0);
+				Index++;
+			}
+		}
+		TArrayView<FVector> SamplesView(&(SamplesArray.operator[](0)), SamplesArray.Num());
+
+		float MinDoimain = -1.0, MaxDomain = 1.0;
+		FTransform Transform(FQuat::MakeFromEuler(FVector(45,45,45)), FVector(100, 0, 0), FVector(2, 1, 1));
+		FNoiseField * NoiseField = new FNoiseField(MinDoimain, MaxDomain,Transform);
+
+		FFieldContext Context{
+			IndexView,
+			SamplesView,
+			FFieldContext::UniquePointerMap()
+		};
+
+		TArray<float> ResultsArray;
+		ResultsArray.Init(false, Bounds*Bounds);
+		TArrayView<float> ResultsView(&(ResultsArray.operator[](0)), ResultsArray.Num());
+		NoiseField->Evaluate(Context, ResultsView);
+
+		float min = FLT_MAX, max = -FLT_MAX;
+		double avg = 0.0;
+		Chaos::Utilities::GetMinAvgMax(ResultsView, min, avg, max);
+		EXPECT_GE(min, MinDoimain);
+		EXPECT_LE(max, MaxDomain);
+		EXPECT_LT(min, max);
+
+		float variance = Chaos::Utilities::GetVariance(ResultsView, avg);
+		float stdDev = Chaos::Utilities::GetStandardDeviation(variance);
+		EXPECT_GT(variance, 0.0); // If variance is 0, then all values are the same.
+		EXPECT_GT(stdDev, 0.0);
+		EXPECT_LT(stdDev, 0.5);
+
+		delete NoiseField;
+	}
+
+	void Fields_RadialIntMask()
+	{
+		TArray<ContextIndex> IndicesArray;
+		ContextIndex::ContiguousIndices(IndicesArray, 10);
+		TArrayView<ContextIndex> IndexView(&(IndicesArray[0]), IndicesArray.Num());
+
+		TArray<FVector> SamplesArray;
+		SamplesArray.Init(FVector(0.f), 10);
+		for (int32 Index = 0; Index < 10; Index++)
+		{
+			SamplesArray[Index] = FVector(Index);
+		}
+		TArrayView<FVector> SamplesView(&(SamplesArray.operator[](0)), SamplesArray.Num());
+
+		FRadialIntMask * RadialMask = new FRadialIntMask(0.0, FVector(), 1.0, 0.0, ESetMaskConditionType::Field_Set_Always);
+		RadialMask->Position = FVector(0.0, 0.0, 0.0);
+		RadialMask->Radius = 5.0;
+
+		FFieldContext Context{
+			IndexView,
+			SamplesView,
+			FFieldContext::UniquePointerMap()
+		};
+
+		TArray<int32> ResultsArray;
+		ResultsArray.Init(false, 10.0);
+		TArrayView<int32> ResultsView(&(ResultsArray.operator[](0)), ResultsArray.Num());
+		RadialMask->Evaluate(Context, ResultsView);
+
+		for (int32 Index = 0; Index < 10; Index++)
+		{
+			bool bValueAsBool = ResultsView[Index] != 0;
+			if (Index <= 2)
+			{
+				EXPECT_TRUE(bValueAsBool);
+			}
+			else
+			{
+				EXPECT_FALSE(bValueAsBool);
+			}
+			//UE_LOG(GCTF_Log, Error, TEXT("[%d] %d"), Index, ResultsView[Index]);
+		}
+		delete RadialMask;
+	}
+
+	void Fields_RadialFalloff()
+	{
+		TArray<ContextIndex> IndicesArray;
+		ContextIndex::ContiguousIndices(IndicesArray, 10);
+		TArrayView<ContextIndex> IndexView(&(IndicesArray[0]), IndicesArray.Num());
+
+		TArray<FVector> SamplesArray;
+		SamplesArray.Init(FVector(0.f), 10);
+		for (int32 Index = 0; Index < 10; Index++)
+		{
+			SamplesArray[Index] = FVector(Index, 0, 0);
+		}
+		TArrayView<FVector> SamplesView(&(SamplesArray.operator[](0)), SamplesArray.Num());
+
+		FRadialFalloff * RadialFalloff = new FRadialFalloff();
+		RadialFalloff->Position = FVector(0.0, 0.0, 0.0);
+		RadialFalloff->Radius = 5.0;
+		RadialFalloff->Magnitude = 3.0;
+
+		FFieldContext Context{
+			IndexView,
+			SamplesView,
+			FFieldContext::UniquePointerMap()
+		};
+
+		TArray<float> ResultsArray;
+		ResultsArray.Init(false, 10);
+		TArrayView<float> ResultsView(&(ResultsArray.operator[](0)), ResultsArray.Num());
+		RadialFalloff->Evaluate(Context, ResultsView);
+
+		for (int32 Index = 0; Index < 10; Index++)
+		{
+			float ExpectedVal = RadialFalloff->Magnitude*float((RadialFalloff->Radius * RadialFalloff->Radius) - (Index*Index)) / (RadialFalloff->Radius*RadialFalloff->Radius);
+
+			if (Index <= 5)
+			{
+				EXPECT_LT(FMath::Abs(ResultsView[Index]) - ExpectedVal, KINDA_SMALL_NUMBER);
+			}
+			else
+			{
+				EXPECT_FALSE(ResultsView[Index]);
+			}
+			//UE_LOG(GCTF_Log, Error, TEXT("[%d] sample:%3.5f (%3.5f,%3.5f,%3.5f) %3.5f"), Index,
+			//	ExpectedVal ,
+			//	SamplesArray[Index].X, SamplesArray[Index].Y, SamplesArray[Index].Z,
+			//	ResultsView[Index]);
+		}
+		delete RadialFalloff;
+	}
+
+	void Fields_PlaneFalloff()
+	{
+		TArray<ContextIndex> IndicesArray;
+		ContextIndex::ContiguousIndices(IndicesArray, 10);
+		TArrayView<ContextIndex> IndexView(&(IndicesArray[0]), IndicesArray.Num());
+
+
+		TArray<FVector> SamplesArray;
+		SamplesArray.Init(FVector(0.f), 10);
+		for (int32 Index = 0; Index < 10; Index++)
+		{
+			SamplesArray[Index] = FVector(0, 0, Index-5);
+		}
+		TArrayView<FVector> SamplesView(&(SamplesArray.operator[](0)), SamplesArray.Num());
+
+		FPlaneFalloff * PlaneFalloff = new FPlaneFalloff();
+		PlaneFalloff->Position = FVector(0.0, 0.0, 0.0);
+		PlaneFalloff->Normal = FVector(0.0, 0.0, 1.0);
+		PlaneFalloff->Magnitude = 3.0;
+
+		FFieldContext Context{
+			IndexView,
+			SamplesView,
+			FFieldContext::UniquePointerMap()
+		};
+
+		TArray<float> ResultsArray;
+		ResultsArray.Init(false, 10);
+		TArrayView<float> ResultsView(&(ResultsArray.operator[](0)), ResultsArray.Num());
+		PlaneFalloff->Evaluate(Context, ResultsView);
+
+		FPlane Plane(PlaneFalloff->Position, PlaneFalloff->Normal);
+		for (int32 Index = 0; Index < 10; Index++)
+		{
+			float ExpectedVal = 0.f;
+			float Distance = Plane.PlaneDot(SamplesArray[Index]);
+			if (Distance < 0)
+			{
+				ExpectedVal = -PlaneFalloff->Magnitude * Distance;
+			}
+
+			EXPECT_LT(FMath::Abs(ResultsView[Index]) - ExpectedVal, KINDA_SMALL_NUMBER);
+		}
+		delete PlaneFalloff;
+	}
+
+	void Fields_UniformVector()
+	{
+		TArray<ContextIndex> IndicesArray;
+		ContextIndex::ContiguousIndices(IndicesArray, 10);
+		TArrayView<ContextIndex> IndexView(&(IndicesArray[0]), IndicesArray.Num());
+
+
+		TArray<FVector> SamplesArray;
+		SamplesArray.Init(FVector(0.f), 10);
+		for (int32 Index = 0; Index < 10; Index++)
+		{
+			SamplesArray[Index] = FVector(Index, 0, 0);
+		}
+		TArrayView<FVector> SamplesView(&(SamplesArray.operator[](0)), SamplesArray.Num());
+
+		FUniformVector * UniformVector = new FUniformVector();
+		UniformVector->Direction = FVector(3, 5, 7);
+		UniformVector->Magnitude = 10.0;
+
+		FFieldContext Context{
+			IndexView,
+			SamplesView,
+			FFieldContext::UniquePointerMap()
+		};
+
+		TArray<FVector> ResultsArray;
+		ResultsArray.Init(FVector(0), 10);
+		TArrayView<FVector> ResultsView(&(ResultsArray.operator[](0)), ResultsArray.Num());
+		UniformVector->Evaluate(Context, ResultsView);
+
+		for (int32 Index = 0; Index < 10; Index++)
+		{
+			FVector ExpectedVal = 10.0*FVector(3, 5, 7);
+			EXPECT_LT((ResultsView[Index] - ExpectedVal).Size(), KINDA_SMALL_NUMBER);
+			//UE_LOG(GCTF_Log, Error, TEXT("[%d] sample:(%3.5f,%3.5f,%3.5f) (%3.5f,%3.5f,%3.5f) (%3.5f,%3.5f,%3.5f)"), Index,
+			//	ExpectedVal.X, ExpectedVal.Y, ExpectedVal.Z,
+			//	SamplesArray[Index].X, SamplesArray[Index].Y, SamplesArray[Index].Z,
+			//	ResultsView[Index].X, ResultsView[Index].Y, ResultsView[Index].Z);
+		}
+		delete UniformVector;
+	}
+
+	void Fields_RaidalVector()
+	{
+		TArray<ContextIndex> IndicesArray;
+		ContextIndex::ContiguousIndices(IndicesArray, 10);
+		TArrayView<ContextIndex> IndexView(&(IndicesArray[0]), IndicesArray.Num());
+
+
+		TArray<FVector> SamplesArray;
+		SamplesArray.Init(FVector(0.f), 10);
+		for (int32 Index = 0; Index < 10; Index++)
+		{
+			SamplesArray[Index] = FVector(100.*(FMath::SRand() - 0.5), 100.*(FMath::SRand() - 0.5), 100.*(FMath::SRand() - 0.5));
+		}
+		TArrayView<FVector> SamplesView(&(SamplesArray.operator[](0)), SamplesArray.Num());
+
+		FRadialVector * RadialVector = new FRadialVector();
+		RadialVector->Position = FVector(3, 4, 5);
+		RadialVector->Magnitude = 10.0;
+
+		FFieldContext Context{
+			IndexView,
+			SamplesView,
+			FFieldContext::UniquePointerMap()
+		};
+
+		TArray<FVector> ResultsArray;
+		ResultsArray.Init(FVector(0), 10);
+		TArrayView<FVector> ResultsView(&(ResultsArray.operator[](0)), ResultsArray.Num());
+		RadialVector->Evaluate(Context, ResultsView);
+
+		for (int32 Index = 0; Index < 10; Index++)
+		{
+			FVector ExpectedVal = RadialVector->Magnitude * (SamplesArray[Index] - RadialVector->Position).GetSafeNormal();
+			EXPECT_LT((ResultsView[Index] - ExpectedVal).Size(), KINDA_SMALL_NUMBER);
+			//UE_LOG(GCTF_Log, Error, TEXT("[%d] sample:(%3.5f,%3.5f,%3.5f) (%3.5f,%3.5f,%3.5f) (%3.5f,%3.5f,%3.5f)"), Index,
+			//	ExpectedVal.X, ExpectedVal.Y, ExpectedVal.Z,
+			//	SamplesArray[Index].X, SamplesArray[Index].Y, SamplesArray[Index].Z,
+			//	ResultsView[Index].X, ResultsView[Index].Y, ResultsView[Index].Z);
+		}
+		delete RadialVector;
+	}
+
+	void Fields_SumVectorFullMult()
+	{
+		TArray<ContextIndex> IndicesArray;
+		ContextIndex::ContiguousIndices(IndicesArray, 10);
+		TArrayView<ContextIndex> IndexView(&(IndicesArray[0]), IndicesArray.Num());
+
+
+		float AverageSampleLength = 0.0;
+		TArray<FVector> SamplesArray;
+		SamplesArray.Init(FVector(0.f), 10);
+		for (int32 Index = 0; Index < 10; Index++)
+		{
+			SamplesArray[Index] = FVector(0);
+			if (Index > 0)
+			{
+				SamplesArray[Index] = FVector(100.*(FMath::SRand() - 0.5), 100.*(FMath::SRand() - 0.5), 100.*(FMath::SRand() - 0.5));
+			}
+			AverageSampleLength += SamplesArray[Index].Size();
+		}
+		AverageSampleLength /= 10.0;
+		TArrayView<FVector> SamplesView(&(SamplesArray.operator[](0)), SamplesArray.Num());
+
+		FRadialFalloff * RadialFalloff = new FRadialFalloff();
+		RadialFalloff->Position = FVector(0.0, 0.0, 0.0);
+		RadialFalloff->Radius = AverageSampleLength;
+		RadialFalloff->Magnitude = 3.0;
+
+		FRadialVector * RadialVector = new FRadialVector();
+		RadialVector->Position = FVector(0);
+		RadialVector->Magnitude = 10.0;
+
+		FUniformVector * UniformVector = new FUniformVector();
+		UniformVector->Direction = FVector(3, 5, 7);
+		UniformVector->Magnitude = 10.0;
+
+		FSumVector * SumVector = new FSumVector(1.0, RadialFalloff, RadialVector, UniformVector, EFieldOperationType::Field_Multiply);
+
+		FFieldContext Context{
+			IndexView,
+			SamplesView,
+			FFieldContext::UniquePointerMap()
+		};
+
+		TArray<FVector> ResultsArray;
+		ResultsArray.Init(FVector(0), 10);
+		TArrayView<FVector> ResultsView(&(ResultsArray.operator[](0)), ResultsArray.Num());
+		SumVector->Evaluate(Context, ResultsView);
+
+		float RadialFalloffSize2 = RadialFalloff->Radius * RadialFalloff->Radius;
+		for (int32 Index = 0; Index < 10; Index++)
+		{
+			FVector RightResult = UniformVector->Magnitude * UniformVector->Direction;
+			FVector LeftResult = RadialVector->Magnitude  * (SamplesArray[Index] - RadialVector->Position).GetSafeNormal();
+			float RadialFalloffDelta2 = (SamplesArray[Index] - RadialFalloff->Position).SizeSquared();
+			float ScalarResult = RadialFalloff->Magnitude* (RadialFalloffSize2 - RadialFalloffDelta2) / RadialFalloffSize2;
+			if (RadialFalloffDelta2 >= RadialFalloffSize2)
+				ScalarResult = 0.f;
+
+			FVector ExpectedVal = ScalarResult * (LeftResult * RightResult);
+
+			//UE_LOG(GCTF_Log, Error, TEXT("[%d:%3.5f] sample:(%3.5f,%3.5f,%3.5f) (%3.5f,%3.5f,%3.5f) (%3.5f,%3.5f,%3.5f)"), Index,
+			//	(ResultsView[Index] - ExpectedVal).Size(),
+			//	ExpectedVal.X, ExpectedVal.Y, ExpectedVal.Z,
+			//	SamplesArray[Index].X, SamplesArray[Index].Y, SamplesArray[Index].Z,
+			//	ResultsView[Index].X, ResultsView[Index].Y, ResultsView[Index].Z);
+
+			EXPECT_LT((ResultsView[Index] - ExpectedVal).Size(), KINDA_SMALL_NUMBER);
+		}
+		delete SumVector;
+	}
+
+	void Fields_SumVectorFullDiv()
+	{
+		TArray<ContextIndex> IndicesArray;
+		ContextIndex::ContiguousIndices(IndicesArray, 10);
+		TArrayView<ContextIndex> IndexView(&(IndicesArray[0]), IndicesArray.Num());
+
+
+		float AverageSampleLength = 0.0;
+		TArray<FVector> SamplesArray;
+		SamplesArray.Init(FVector(0.f), 10);
+		for (int32 Index = 0; Index < 10; Index++)
+		{
+			SamplesArray[Index] = FVector(0);
+			if (Index > 0)
+			{
+				SamplesArray[Index] = FVector(100.*(FMath::SRand() - 0.5), 100.*(FMath::SRand() - 0.5), 100.*(FMath::SRand() - 0.5));
+			}
+			AverageSampleLength += SamplesArray[Index].Size();
+		}
+		AverageSampleLength /= 10.0;
+		TArrayView<FVector> SamplesView(&(SamplesArray.operator[](0)), SamplesArray.Num());
+
+		FRadialFalloff * RadialFalloff = new FRadialFalloff();
+		RadialFalloff->Position = FVector(0.0, 0.0, 0.0);
+		RadialFalloff->Radius = AverageSampleLength;
+		RadialFalloff->Magnitude = 3.0;
+
+		FRadialVector * RadialVector = new FRadialVector();
+		RadialVector->Position = FVector(0);
+		RadialVector->Magnitude = 10.0;
+
+		FUniformVector * UniformVector = new FUniformVector();
+		UniformVector->Direction = FVector(3, 5, 7);
+		UniformVector->Magnitude = 10.0;
+
+		FSumVector * SumVector = new FSumVector(1.0, RadialFalloff, UniformVector, RadialVector,  EFieldOperationType::Field_Divide);
+
+		FFieldContext Context{
+			IndexView,
+			SamplesView,
+			FFieldContext::UniquePointerMap()
+		};
+
+		TArray<FVector> ResultsArray;
+		ResultsArray.Init(FVector(0), 10);
+		TArrayView<FVector> ResultsView(&(ResultsArray.operator[](0)), ResultsArray.Num());
+		SumVector->Evaluate(Context, ResultsView);
+
+		float RadialFalloffSize2 = RadialFalloff->Radius * RadialFalloff->Radius;
+		for (int32 Index = 0; Index < 10; Index++)
+		{
+			FVector RightResult = UniformVector->Magnitude * UniformVector->Direction;
+			FVector LeftResult = RadialVector->Magnitude  * (SamplesArray[Index] - RadialVector->Position).GetSafeNormal();
+			float RadialFalloffDelta2 = (SamplesArray[Index] - RadialFalloff->Position).SizeSquared();
+			float ScalarResult = RadialFalloff->Magnitude* (RadialFalloffSize2 - RadialFalloffDelta2) / RadialFalloffSize2;
+			if (RadialFalloffDelta2 >= RadialFalloffSize2)
+				ScalarResult = 0.f;
+
+			FVector ExpectedVal = ScalarResult * (LeftResult / RightResult);
+
+			//UE_LOG(GCTF_Log, Error, TEXT("[%d:%3.5f] sample:(%3.5f,%3.5f,%3.5f) (%3.5f,%3.5f,%3.5f) (%3.5f,%3.5f,%3.5f)"), Index,
+			//	(ResultsView[Index] - ExpectedVal).Size(),
+			//	ExpectedVal.X, ExpectedVal.Y, ExpectedVal.Z,
+			//	SamplesArray[Index].X, SamplesArray[Index].Y, SamplesArray[Index].Z,
+			//	ResultsView[Index].X, ResultsView[Index].Y, ResultsView[Index].Z);
+
+			EXPECT_LT((ResultsView[Index] - ExpectedVal).Size(), KINDA_SMALL_NUMBER);
+		}
+		delete SumVector;
+	}
+
+	void Fields_SumVectorFullAdd()
+	{
+		TArray<ContextIndex> IndicesArray;
+		ContextIndex::ContiguousIndices(IndicesArray, 10);
+		TArrayView<ContextIndex> IndexView(&(IndicesArray[0]), IndicesArray.Num());
+
+
+		float AverageSampleLength = 0.0;
+		TArray<FVector> SamplesArray;
+		SamplesArray.Init(FVector(0.f), 10);
+		for (int32 Index = 0; Index < 10; Index++)
+		{
+			SamplesArray[Index] = FVector(0);
+			if (Index > 0)
+			{
+				SamplesArray[Index] = FVector(100.*(FMath::SRand() - 0.5), 100.*(FMath::SRand() - 0.5), 100.*(FMath::SRand() - 0.5));
+			}
+			AverageSampleLength += SamplesArray[Index].Size();
+		}
+		AverageSampleLength /= 10.0;
+		TArrayView<FVector> SamplesView(&(SamplesArray.operator[](0)), SamplesArray.Num());
+
+
+		FRadialFalloff * RadialFalloff = new FRadialFalloff();
+		RadialFalloff->Position = FVector(0.0, 0.0, 0.0);
+		RadialFalloff->Radius = AverageSampleLength;
+		RadialFalloff->Magnitude = 3.0;
+
+		FRadialVector * RadialVector = new FRadialVector();
+		RadialVector->Position = FVector(0);
+		RadialVector->Magnitude = 10.0;
+
+		FUniformVector * UniformVector = new FUniformVector();
+		UniformVector->Direction = FVector(3, 5, 7);
+		UniformVector->Magnitude = 10.0;
+
+		FSumVector * SumVector = new FSumVector(1.0, RadialFalloff, UniformVector, RadialVector, EFieldOperationType::Field_Add);
+
+		FFieldContext Context{
+			IndexView,
+			SamplesView,
+			FFieldContext::UniquePointerMap()
+		};
+
+		TArray<FVector> ResultsArray;
+		ResultsArray.Init(FVector(0), 10);
+		TArrayView<FVector> ResultsView(&(ResultsArray.operator[](0)), ResultsArray.Num());
+		SumVector->Evaluate(Context, ResultsView);
+
+		float RadialFalloffSize2 = RadialFalloff->Radius * RadialFalloff->Radius;
+		for (int32 Index = 0; Index < 10; Index++)
+		{
+			FVector RightResult = UniformVector->Magnitude * UniformVector->Direction;
+			FVector LeftResult = RadialVector->Magnitude  * (SamplesArray[Index] - RadialVector->Position).GetSafeNormal();
+			float RadialFalloffDelta2 = (SamplesArray[Index] - RadialFalloff->Position).SizeSquared();
+			float ScalarResult = RadialFalloff->Magnitude* (RadialFalloffSize2 - RadialFalloffDelta2) / RadialFalloffSize2;
+			if (RadialFalloffDelta2 >= RadialFalloffSize2)
+				ScalarResult = 0.f;
+
+			FVector ExpectedVal = ScalarResult * (LeftResult + RightResult);
+
+			//UE_LOG(GCTF_Log, Error, TEXT("[%d:%3.5f] sample:(%3.5f,%3.5f,%3.5f) (%3.5f,%3.5f,%3.5f) (%3.5f,%3.5f,%3.5f)"), Index,
+			//	(ResultsView[Index] - ExpectedVal).Size(),
+			//	ExpectedVal.X, ExpectedVal.Y, ExpectedVal.Z,
+			//	SamplesArray[Index].X, SamplesArray[Index].Y, SamplesArray[Index].Z,
+			//	ResultsView[Index].X, ResultsView[Index].Y, ResultsView[Index].Z);
+
+			EXPECT_LT((ResultsView[Index] - ExpectedVal).Size(), KINDA_SMALL_NUMBER);
+		}
+		delete SumVector;
+	}
+
+	void Fields_SumVectorFullSub()
+	{
+		TArray<ContextIndex> IndicesArray;
+		ContextIndex::ContiguousIndices(IndicesArray, 10);
+		TArrayView<ContextIndex> IndexView(&(IndicesArray[0]), IndicesArray.Num());
+
+
+		float AverageSampleLength = 0.0;
+		TArray<FVector> SamplesArray;
+		SamplesArray.Init(FVector(0.f), 10);
+		for (int32 Index = 0; Index < 10; Index++)
+		{
+			SamplesArray[Index] = FVector(0);
+			if (Index > 0)
+			{
+				SamplesArray[Index] = FVector(100.*(FMath::SRand() - 0.5), 100.*(FMath::SRand() - 0.5), 100.*(FMath::SRand() - 0.5));
+			}
+			AverageSampleLength += SamplesArray[Index].Size();
+		}
+		AverageSampleLength /= 10.0;
+		TArrayView<FVector> SamplesView(&(SamplesArray.operator[](0)), SamplesArray.Num());
+
+
+		FRadialFalloff * RadialFalloff = new FRadialFalloff();
+		RadialFalloff->Position = FVector(0.0, 0.0, 0.0);
+		RadialFalloff->Radius = AverageSampleLength;
+		RadialFalloff->Magnitude = 3.0;
+
+		FRadialVector * RadialVector = new FRadialVector();
+		RadialVector->Position = FVector(0);
+		RadialVector->Magnitude = 10.0;
+
+		FUniformVector * UniformVector = new FUniformVector();
+		UniformVector->Direction = FVector(3, 5, 7);
+		UniformVector->Magnitude = 10.0;
+
+		FSumVector * SumVector = new FSumVector(1.0, RadialFalloff, UniformVector, RadialVector, EFieldOperationType::Field_Substract);
+
+		FFieldContext Context{
+			IndexView,
+			SamplesView,
+			FFieldContext::UniquePointerMap()
+		};
+
+		TArray<FVector> ResultsArray;
+		ResultsArray.Init(FVector(0), 10);
+		TArrayView<FVector> ResultsView(&(ResultsArray.operator[](0)), ResultsArray.Num());
+		SumVector->Evaluate(Context, ResultsView);
+
+		float RadialFalloffSize2 = RadialFalloff->Radius * RadialFalloff->Radius;
+		for (int32 Index = 0; Index < 10; Index++)
+		{
+			FVector RightResult = UniformVector->Magnitude * UniformVector->Direction;
+			FVector LeftResult = RadialVector->Magnitude  * (SamplesArray[Index] - RadialVector->Position).GetSafeNormal();
+			float RadialFalloffDelta2 = (SamplesArray[Index] - RadialFalloff->Position).SizeSquared();
+			float ScalarResult = RadialFalloff->Magnitude* (RadialFalloffSize2 - RadialFalloffDelta2) / RadialFalloffSize2;
+			if (RadialFalloffDelta2 >= RadialFalloffSize2)
+				ScalarResult = 0.f;
+
+			FVector ExpectedVal = ScalarResult * (LeftResult - RightResult);
+
+			//UE_LOG(GCTF_Log, Error, TEXT("[%d:%3.5f] sample:(%3.5f,%3.5f,%3.5f) (%3.5f,%3.5f,%3.5f) (%3.5f,%3.5f,%3.5f)"), Index,
+			//	(ResultsView[Index] - ExpectedVal).Size(),
+			//	ExpectedVal.X, ExpectedVal.Y, ExpectedVal.Z,
+			//	SamplesArray[Index].X, SamplesArray[Index].Y, SamplesArray[Index].Z,
+			//	ResultsView[Index].X, ResultsView[Index].Y, ResultsView[Index].Z);
+
+			EXPECT_LT((ResultsView[Index] - ExpectedVal).Size(), KINDA_SMALL_NUMBER);
+		}
+		delete SumVector;
+	}
+
+	void Fields_SumVectorLeftSide()
+	{
+		TArray<ContextIndex> IndicesArray;
+		ContextIndex::ContiguousIndices(IndicesArray, 10);
+		TArrayView<ContextIndex> IndexView(&(IndicesArray[0]), IndicesArray.Num());
+
+
+		float AverageSampleLength = 0.0;
+		TArray<FVector> SamplesArray;
+		SamplesArray.Init(FVector(0.f), 10);
+		for (int32 Index = 0; Index < 10; Index++)
+		{
+			SamplesArray[Index] = FVector(0);
+			if (Index > 0)
+			{
+				SamplesArray[Index] = FVector(100.*(FMath::SRand() - 0.5), 100.*(FMath::SRand() - 0.5), 100.*(FMath::SRand() - 0.5));
+			}
+			AverageSampleLength += SamplesArray[Index].Size();
+		}
+		AverageSampleLength /= 10.0;
+		TArrayView<FVector> SamplesView(&(SamplesArray.operator[](0)), SamplesArray.Num());
+
+		FRadialFalloff * RadialFalloff = new FRadialFalloff();
+		RadialFalloff->Position = FVector(0.0, 0.0, 0.0);
+		RadialFalloff->Radius = AverageSampleLength;
+		RadialFalloff->Magnitude = 3.0;
+
+		FRadialVector * RadialVector = new FRadialVector();
+		RadialVector->Position = FVector(0);
+		RadialVector->Magnitude = 10.0;
+
+		FUniformVector * UniformVector = new FUniformVector();
+		UniformVector->Direction = FVector(3, 5, 7);
+		UniformVector->Magnitude = 10.0;
+
+		FSumVector * SumVector = new FSumVector(1.0, RadialFalloff, nullptr, RadialVector, EFieldOperationType::Field_Multiply);
+
+		FFieldContext Context{
+			IndexView,
+			SamplesView,
+			FFieldContext::UniquePointerMap()
+		};
+
+		TArray<FVector> ResultsArray;
+		ResultsArray.Init(FVector(0), 10);
+		TArrayView<FVector> ResultsView(&(ResultsArray.operator[](0)), ResultsArray.Num());
+		SumVector->Evaluate(Context, ResultsView);
+
+		float RadialFalloffSize2 = RadialFalloff->Radius * RadialFalloff->Radius;
+		for (int32 Index = 0; Index < 10; Index++)
+		{
+			FVector RightResult = UniformVector->Magnitude * UniformVector->Direction;
+			FVector LeftResult = RadialVector->Magnitude  * (SamplesArray[Index] - RadialVector->Position).GetSafeNormal();
+			float RadialFalloffDelta2 = (SamplesArray[Index] - RadialFalloff->Position).SizeSquared();
+			float ScalarResult = RadialFalloff->Magnitude* (RadialFalloffSize2 - RadialFalloffDelta2) / RadialFalloffSize2;
+			if (RadialFalloffDelta2 >= RadialFalloffSize2)
+				ScalarResult = 0.f;
+
+			FVector ExpectedVal = ScalarResult * (LeftResult);
+
+			//UE_LOG(GCTF_Log, Error, TEXT("[%d:%3.5f] sample:(%3.5f,%3.5f,%3.5f) (%3.5f,%3.5f,%3.5f) (%3.5f,%3.5f,%3.5f)"), Index,
+			//	(ResultsView[Index] - ExpectedVal).Size(),
+			//	ExpectedVal.X, ExpectedVal.Y, ExpectedVal.Z,
+			//	SamplesArray[Index].X, SamplesArray[Index].Y, SamplesArray[Index].Z,
+			//	ResultsView[Index].X, ResultsView[Index].Y, ResultsView[Index].Z);
+
+			EXPECT_LT((ResultsView[Index] - ExpectedVal).Size(), KINDA_SMALL_NUMBER);
+		}
+		delete SumVector;
+		delete UniformVector;
+	}
+
+	void Fields_SumVectorRightSide()
+	{
+		TArray<ContextIndex> IndicesArray;
+		ContextIndex::ContiguousIndices(IndicesArray, 10);
+		TArrayView<ContextIndex> IndexView(&(IndicesArray[0]), IndicesArray.Num());
+
+		float AverageSampleLength = 0.0;
+		TArray<FVector> SamplesArray;
+		SamplesArray.Init(FVector(0.f), 10);
+		for (int32 Index = 0; Index < 10; Index++)
+		{
+			SamplesArray[Index] = FVector(0);
+			if (Index > 0)
+			{
+				SamplesArray[Index] = FVector(100.*(FMath::SRand() - 0.5), 100.*(FMath::SRand() - 0.5), 100.*(FMath::SRand() - 0.5));
+			}
+			AverageSampleLength += SamplesArray[Index].Size();
+		}
+		AverageSampleLength /= 10.0;
+		TArrayView<FVector> SamplesView(&(SamplesArray.operator[](0)), SamplesArray.Num());
+
+
+		FRadialFalloff * RadialFalloff = new FRadialFalloff();
+		RadialFalloff->Position = FVector(0.0, 0.0, 0.0);
+		RadialFalloff->Radius = AverageSampleLength;
+		RadialFalloff->Magnitude = 3.0;
+
+		FRadialVector * RadialVector = new FRadialVector();
+		RadialVector->Position = FVector(0);
+		RadialVector->Magnitude = 10.0;
+
+		FUniformVector * UniformVector = new FUniformVector();
+		UniformVector->Direction = FVector(3, 5, 7);
+		UniformVector->Magnitude = 10.0;
+
+		FSumVector * SumVector = new FSumVector(1.0, RadialFalloff, UniformVector, nullptr, EFieldOperationType::Field_Multiply);
+
+		FFieldContext Context{
+			IndexView,
+			SamplesView,
+			FFieldContext::UniquePointerMap()
+		};
+
+		TArray<FVector> ResultsArray;
+		ResultsArray.Init(FVector(0), 10);
+		TArrayView<FVector> ResultsView(&(ResultsArray.operator[](0)), ResultsArray.Num());
+		SumVector->Evaluate(Context, ResultsView);
+
+		float RadialFalloffSize2 = RadialFalloff->Radius * RadialFalloff->Radius;
+		for (int32 Index = 0; Index < 10; Index++)
+		{
+			FVector RightResult = UniformVector->Magnitude * UniformVector->Direction;
+			FVector LeftResult = RadialVector->Magnitude  * (SamplesArray[Index] - RadialVector->Position).GetSafeNormal();
+			float RadialFalloffDelta2 = (SamplesArray[Index] - RadialFalloff->Position).SizeSquared();
+			float ScalarResult = RadialFalloff->Magnitude* (RadialFalloffSize2 - RadialFalloffDelta2) / RadialFalloffSize2;
+			if (RadialFalloffDelta2 >= RadialFalloffSize2)
+				ScalarResult = 0.f;
+
+			FVector ExpectedVal = ScalarResult * (RightResult);
+
+			//UE_LOG(GCTF_Log, Error, TEXT("[%d:%3.5f] sample:(%3.5f,%3.5f,%3.5f) (%3.5f,%3.5f,%3.5f) (%3.5f,%3.5f,%3.5f)"), Index,
+			//	(ResultsView[Index] - ExpectedVal).Size(),
+			//	ExpectedVal.X, ExpectedVal.Y, ExpectedVal.Z,
+			//	SamplesArray[Index].X, SamplesArray[Index].Y, SamplesArray[Index].Z,
+			//	ResultsView[Index].X, ResultsView[Index].Y, ResultsView[Index].Z);
+
+			EXPECT_LT((ResultsView[Index] - ExpectedVal).Size(), KINDA_SMALL_NUMBER);
+		}
+		delete SumVector;
+		delete RadialVector;
+	}
+
+	void Fields_SumScalar()
+	{
+		int32 NumPoints = 20;
+		TArray<ContextIndex> IndicesArray;
+		ContextIndex::ContiguousIndices(IndicesArray, NumPoints);
+		TArrayView<ContextIndex> IndexView(&(IndicesArray[0]), IndicesArray.Num());
+
+
+		TArray<FVector> SamplesArray;
+		SamplesArray.Init(FVector(0.f), NumPoints);
+		for (int32 Index = -10; Index < 10; Index++)
+		{
+			SamplesArray[Index + 10] = FVector(Index, 0, 0);
+		}
+		TArrayView<FVector> SamplesView(&(SamplesArray.operator[](0)), SamplesArray.Num());
+		
+		FRadialFalloff * RadialFalloff = new FRadialFalloff();
+		RadialFalloff->Position = FVector(5.0, 0.0, 0.0);
+		RadialFalloff->Radius = 10;
+		RadialFalloff->Magnitude = 3.0;
+		float RadialFalloffRadius2 = RadialFalloff->Radius * RadialFalloff->Radius;
+
+		FRadialFalloff * RadialFalloff2 = new FRadialFalloff();
+		RadialFalloff2->Position = FVector(5.0, 0.0, 0.0);
+		RadialFalloff2->Radius = 10;
+		RadialFalloff2->Magnitude = 3.0;
+		float RadialFalloff2Radius2 = RadialFalloff2->Radius * RadialFalloff2->Radius;
+
+		FSumScalar * SumScalar = new FSumScalar(1.f, RadialFalloff, RadialFalloff2,EFieldOperationType::Field_Multiply);
+
+		FFieldContext Context{
+		IndexView,
+		SamplesView,
+		FFieldContext::UniquePointerMap()
+		};
+
+		TArray<float> ResultsArray;
+		ResultsArray.Init(0.f, SamplesArray.Num());
+		TArrayView<float> ResultsView(&(ResultsArray.operator[](0)), ResultsArray.Num());
+		SumScalar->Evaluate(Context, ResultsView);
+
+		for (int32 Index = 0; Index < NumPoints; Index++)
+		{
+			float ScalarRight = 0.f;
+			{// FRadialFalloff::Evaluate
+				float  RadialFalloffDelta2 = (RadialFalloff->Position - SamplesArray[Index]).SizeSquared();
+				if (RadialFalloffDelta2 < RadialFalloffRadius2)
+				{
+					ScalarRight = RadialFalloff->Magnitude * (RadialFalloffRadius2 - RadialFalloffDelta2) / RadialFalloffRadius2;
+				}
+			}
+
+			float ScalarLeft = 0.f;
+			{ //  FRadialFalloff2::Evaluate
+				float  RadialFalloff2Delta2 = (RadialFalloff2->Position - SamplesArray[Index]).SizeSquared();
+				if (RadialFalloff2Delta2 < RadialFalloff2Radius2)
+				{
+					ScalarLeft = RadialFalloff2->Magnitude * (RadialFalloff2Radius2 - RadialFalloff2Delta2) / RadialFalloff2Radius2;
+				}
+			}
+
+			float ExpectedVal = ScalarLeft * ScalarRight;
+
+			//UE_LOG(GCTF_Log, Error, TEXT("[%d:%3.5f] sample:(%3.5f,%3.5f,%3.5f) %3.5f -> %3.5f"), Index,
+			//	ResultsView[Index] - ExpectedVal,
+			//	SamplesArray[Index].X, SamplesArray[Index].Y, SamplesArray[Index].Z,
+			//	ExpectedVal, ResultsView[Index]);
+
+			EXPECT_LT((ResultsView[Index] - ExpectedVal), KINDA_SMALL_NUMBER);
+		}
+		delete SumScalar;
+	}
+
+	void Fields_SumScalarRightSide()
+	{
+		int32 NumPoints = 20;
+		TArray<ContextIndex> IndicesArray;
+		ContextIndex::ContiguousIndices(IndicesArray, NumPoints);
+		TArrayView<ContextIndex> IndexView(&(IndicesArray[0]), IndicesArray.Num());
+
+		TArray<FVector> SamplesArray;
+		SamplesArray.Init(FVector(0.f), NumPoints);
+		for (int32 Index = -10; Index < 10; Index++)
+		{
+			SamplesArray[Index + 10] = FVector(Index, 0, 0);
+		}
+		TArrayView<FVector> SamplesView(&(SamplesArray.operator[](0)), SamplesArray.Num());
+
+		FRadialFalloff * RadialFalloff = new FRadialFalloff();
+		RadialFalloff->Position = FVector(5.0, 0.0, 0.0);
+		RadialFalloff->Radius = 10;
+		RadialFalloff->Magnitude = 3.0;
+		float RadialFalloffRadius2 = RadialFalloff->Radius * RadialFalloff->Radius;
+
+		FRadialFalloff * RadialFalloff2 = new FRadialFalloff();
+		RadialFalloff2->Position = FVector(5.0, 0.0, 0.0);
+		RadialFalloff2->Radius = 10;
+		RadialFalloff2->Magnitude = 3.0;
+		float RadialFalloff2Radius2 = RadialFalloff2->Radius * RadialFalloff2->Radius;
+
+		FSumScalar * SumScalar = new FSumScalar(1.f, RadialFalloff, nullptr, EFieldOperationType::Field_Multiply);
+
+		FFieldContext Context{
+			IndexView,
+			SamplesView,
+			FFieldContext::UniquePointerMap()
+		};
+
+		TArray<float> ResultsArray;
+		ResultsArray.Init(0.f, SamplesArray.Num());
+		TArrayView<float> ResultsView(&(ResultsArray.operator[](0)), ResultsArray.Num());
+		SumScalar->Evaluate(Context, ResultsView);
+
+		for (int32 Index = 0; Index < NumPoints; Index++)
+		{
+			float ScalarRight = 0.f;
+			{// FRadialFalloff::Evaluate
+				float  RadialFalloffDelta2 = (RadialFalloff->Position - SamplesArray[Index]).SizeSquared();
+				if (RadialFalloffDelta2 < RadialFalloffRadius2)
+				{
+					ScalarRight = RadialFalloff->Magnitude * (RadialFalloffRadius2 - RadialFalloffDelta2) / RadialFalloffRadius2;
+				}
+			}
+
+			float ScalarLeft = 1.f;
+			float ExpectedVal = ScalarLeft * ScalarRight;
+
+			//UE_LOG(GCTF_Log, Error, TEXT("[%d:%3.5f] sample:(%3.5f,%3.5f,%3.5f) %3.5f -> %3.5f"), Index,
+			//	ResultsView[Index] - ExpectedVal,
+			//	SamplesArray[Index].X, SamplesArray[Index].Y, SamplesArray[Index].Z,
+			//	ExpectedVal, ResultsView[Index]);
+
+			EXPECT_LT((ResultsView[Index] - ExpectedVal), KINDA_SMALL_NUMBER);
+		}
+		delete SumScalar;
+		delete RadialFalloff2;
+	}
+
+	void Fields_SumScalarLeftSide()
+	{
+		int32 NumPoints = 20;
+		TArray<ContextIndex> IndicesArray;
+		ContextIndex::ContiguousIndices(IndicesArray, NumPoints);
+		TArrayView<ContextIndex> IndexView(&(IndicesArray[0]), IndicesArray.Num());
+
+		TArray<FVector> SamplesArray;
+		SamplesArray.Init(FVector(0.f), NumPoints);
+		for (int32 Index = -10; Index < 10; Index++)
+		{
+			SamplesArray[Index + 10] = FVector(Index, 0, 0);
+		}
+		TArrayView<FVector> SamplesView(&(SamplesArray.operator[](0)), SamplesArray.Num());
+
+		FRadialFalloff * RadialFalloff = new FRadialFalloff();
+		RadialFalloff->Position = FVector(5.0, 0.0, 0.0);
+		RadialFalloff->Radius = 10;
+		RadialFalloff->Magnitude = 3.0;
+		float RadialFalloffRadius2 = RadialFalloff->Radius * RadialFalloff->Radius;
+
+		FRadialFalloff * RadialFalloff2 = new FRadialFalloff();
+		RadialFalloff2->Position = FVector(5.0, 0.0, 0.0);
+		RadialFalloff2->Radius = 10;
+		RadialFalloff2->Magnitude = 3.0;
+		float RadialFalloff2Radius2 = RadialFalloff2->Radius * RadialFalloff2->Radius;
+
+		FSumScalar * SumScalar = new FSumScalar(1.f, nullptr, RadialFalloff2, EFieldOperationType::Field_Multiply);
+
+		FFieldContext Context{
+			IndexView,
+			SamplesView,
+			FFieldContext::UniquePointerMap()
+		};
+
+		TArray<float> ResultsArray;
+		ResultsArray.Init(0.f, SamplesArray.Num());
+		TArrayView<float> ResultsView(&(ResultsArray.operator[](0)), ResultsArray.Num());
+		SumScalar->Evaluate(Context, ResultsView);
+
+		for (int32 Index = 0; Index < NumPoints; Index++)
+		{
+			float ScalarRight = 1.f;
+
+			float ScalarLeft = 0.f;
+			{ //  FRadialFalloff2::Evaluate
+				float  RadialFalloff2Delta2 = (RadialFalloff2->Position - SamplesArray[Index]).SizeSquared();
+				if (RadialFalloff2Delta2 < RadialFalloff2Radius2)
+				{
+					ScalarLeft = RadialFalloff2->Magnitude * (RadialFalloff2Radius2 - RadialFalloff2Delta2) / RadialFalloff2Radius2;
+				}
+			}
+
+			float ExpectedVal = ScalarLeft * ScalarRight;
+
+			//UE_LOG(GCTF_Log, Error, TEXT("[%d:%3.5f] sample:(%3.5f,%3.5f,%3.5f) %3.5f -> %3.5f"), Index,
+			//	ResultsView[Index] - ExpectedVal,
+			//	SamplesArray[Index].X, SamplesArray[Index].Y, SamplesArray[Index].Z,
+			//	ExpectedVal, ResultsView[Index]);
+
+			EXPECT_LT((ResultsView[Index] - ExpectedVal), KINDA_SMALL_NUMBER);
+		}
+		delete SumScalar;
+		delete RadialFalloff;
+	}
+
+	void Fields_Culling()
+	{
+		int32 NumPoints = 20;
+		TArray<ContextIndex> IndicesArray;
+		ContextIndex::ContiguousIndices(IndicesArray, NumPoints);
+		TArrayView<ContextIndex> IndexView(&(IndicesArray[0]), IndicesArray.Num());
+
+		TArray<FVector> SamplesArray;
+		SamplesArray.Init(FVector(0.f), NumPoints);
+		for (int32 Index = -10; Index < 10; Index++)
+		{
+			SamplesArray[Index + 10] = FVector(Index, 0, 0);
+		}
+		TArrayView<FVector> SamplesView(&(SamplesArray.operator[](0)), SamplesArray.Num());
+
+		FRadialFalloff * RadialFalloff = new FRadialFalloff();
+		RadialFalloff->Position = FVector(0.0, 0.0, 0.0);
+		RadialFalloff->Radius = 4.;
+		RadialFalloff->Magnitude = 3.0;
+		float RadialFalloffRadius2 = RadialFalloff->Radius * RadialFalloff->Radius;
+
+		FRadialFalloff * RadialFalloff2 = new FRadialFalloff();
+		RadialFalloff2->Position = FVector(5.0, 0.0, 0.0);
+		RadialFalloff2->Radius = 10;
+		RadialFalloff2->Magnitude = 3.0;
+		float RadialFalloff2Radius2 = RadialFalloff2->Radius * RadialFalloff2->Radius;
+
+		FCullingField<float> * CullingField = new FCullingField<float>(RadialFalloff, RadialFalloff2, EFieldCullingOperationType::Field_Culling_Outside);
+
+		FFieldContext Context{
+			IndexView,
+			SamplesView,
+			FFieldContext::UniquePointerMap()
+		};
+
+		TArray<float> ResultsArray;
+		ResultsArray.Init(0.f, SamplesArray.Num());
+		TArrayView<float> ResultsView(&(ResultsArray.operator[](0)), ResultsArray.Num());
+		CullingField->Evaluate(Context, ResultsView);
+
+		for (int32 Index = 0; Index < NumPoints; Index++)
+		{
+			float ScalarRight = 1.f;
+
+			float ScalarLeft = 0.f;
+			{ //  FCullingField::Evaluate
+				float  RadialFalloff2Delta2 = (RadialFalloff2->Position - SamplesArray[Index]).SizeSquared();
+				if (RadialFalloff2Delta2 < RadialFalloff2Radius2)
+				{
+					ScalarLeft = RadialFalloff2->Magnitude * (RadialFalloff2Radius2 - RadialFalloff2Delta2) / RadialFalloff2Radius2;
+				}
+			}
+
+			float ExpectedVal = ScalarLeft * ScalarRight;
+
+			//UE_LOG(GCTF_Log, Error, TEXT("[%d:%3.5f] sample:(%3.5f,%3.5f,%3.5f) %3.5f -> %3.5f"), Index,
+			//	ResultsView[Index] - ExpectedVal,
+			//	SamplesArray[Index].X, SamplesArray[Index].Y, SamplesArray[Index].Z,
+			//	ExpectedVal, ResultsView[Index]);
+
+			EXPECT_LT((ResultsView[Index] - ExpectedVal), KINDA_SMALL_NUMBER);
+		}
+		delete CullingField;
+	}
+
+	FFieldSystemCommand SaveAndLoad(FFieldSystemCommand & CommandOut)
+	{
+		FBufferArchive Ar;
+		Ar.SetIsSaving(true);
+		Ar.SetIsLoading(false);
+		FString Filename = "Fields_SerializeAPI.tmp";
+
+		FFieldSystemCommand CommandIn;
+
+		CommandOut.Serialize(Ar);
+		FFileHelper::SaveArrayToFile(Ar, *Filename);
+		Ar.FlushCache();
+		Ar.Empty();
+
+		TArray<uint8> InputArray;
+		FFileHelper::LoadFileToArray(InputArray, *Filename);
+		FMemoryReader InputArchive = FMemoryReader(InputArray, true);
+		InputArchive.Seek(0);
+		CommandIn.Serialize(InputArchive);
+
+		return CommandIn;
+	}
+
+	// Helper for testing FFieldSystemCommand equality. Defined very similarly to the == operator but with more granular
+	// checks at each stage of testing. 
+	bool TestFFieldSystemEquality(const FFieldSystemCommand& CommandIn, const FFieldSystemCommand& CommandOut)
+	{
+		EXPECT_TRUE(CommandOut.TargetAttribute.IsEqual(CommandIn.TargetAttribute));
+		if (CommandOut.TargetAttribute.IsEqual(CommandIn.TargetAttribute))
+		{
+			EXPECT_EQ(CommandOut.RootNode.IsValid(), CommandIn.RootNode.IsValid());
+			if (CommandOut.RootNode.IsValid() == CommandIn.RootNode.IsValid())
+			{
+				if (CommandOut.RootNode.IsValid())
+				{
+					EXPECT_EQ(CommandOut.RootNode->SerializationType(), CommandIn.RootNode->SerializationType());
+					if (CommandOut.RootNode->SerializationType() == CommandIn.RootNode->SerializationType())
+					{
+						EXPECT_TRUE(CommandOut.RootNode->operator==(*CommandIn.RootNode));
+						if (CommandOut.RootNode->operator==(*CommandIn.RootNode))
+						{
+							return true;
+						}
+					}
+				}
+				else
+				{
+					return true;
+				}
+			}
+		}
+		return false;
+	}
+
+	void Fields_SerializeAPI()
+	{
+		{
+			FFieldSystemCommand CommandOut("FUniformInteger", new FUniformInteger(3));
+			FFieldSystemCommand CommandIn = SaveAndLoad(CommandOut);
+			EXPECT_TRUE(TestFFieldSystemEquality(CommandIn, CommandOut));
+		}
+
+		{
+			FFieldSystemCommand CommandOut("FRadialIntMask", new FRadialIntMask(1.f,FVector(3,5,7),11,13) );
+			FFieldSystemCommand CommandIn = SaveAndLoad(CommandOut);
+			EXPECT_TRUE(TestFFieldSystemEquality(CommandIn, CommandOut));
+		}
+
+		{
+			FFieldSystemCommand CommandOut("FUniformScalar", new FUniformScalar(13));
+			FFieldSystemCommand CommandIn = SaveAndLoad(CommandOut);
+			EXPECT_TRUE(TestFFieldSystemEquality(CommandIn, CommandOut));
+		}
+
+		{
+			FFieldSystemCommand CommandOut("FRadialFalloff", new FRadialFalloff(1.f,3.f,5.f,7.f,11.f,FVector(13, 17, 19)));
+			FFieldSystemCommand CommandIn = SaveAndLoad(CommandOut);
+			EXPECT_TRUE(TestFFieldSystemEquality(CommandIn, CommandOut));
+		}
+
+		{
+			FFieldSystemCommand CommandOut("FPlaneFalloff", new FPlaneFalloff(1.f, 3.f, 5.f, 7.f, 100.f, FVector(9, 11, 13), FVector(17, 19, 23)));
+			FFieldSystemCommand CommandIn = SaveAndLoad(CommandOut);
+			EXPECT_TRUE(TestFFieldSystemEquality(CommandIn, CommandOut));
+		}
+		
+		{
+			FFieldSystemCommand CommandOut("FBoxFalloff", new FBoxFalloff(1.f, 7.f, 9.f, 13.f, FTransform::Identity));
+			FFieldSystemCommand CommandIn = SaveAndLoad(CommandOut);
+			EXPECT_TRUE(TestFFieldSystemEquality(CommandIn, CommandOut));
+		}
+		
+		{
+			FFieldSystemCommand CommandOut("FNoiseField", new	FNoiseField(1.f, 3.f));
+			FFieldSystemCommand CommandIn = SaveAndLoad(CommandOut);
+			EXPECT_TRUE(TestFFieldSystemEquality(CommandIn, CommandOut));
+		}
+
+		{
+			FFieldSystemCommand CommandOut("FUniformVector", new FUniformVector(1.f, FVector(3, 5, 7)));
+			FFieldSystemCommand CommandIn = SaveAndLoad(CommandOut);
+			EXPECT_TRUE(TestFFieldSystemEquality(CommandIn, CommandOut));
+		}
+
+		{
+			FFieldSystemCommand CommandOut("FRadialVector", new FRadialVector(1.f,FVector(3, 5, 7)));
+			FFieldSystemCommand CommandIn = SaveAndLoad(CommandOut);
+			EXPECT_TRUE(TestFFieldSystemEquality(CommandIn, CommandOut));
+		}
+
+		{
+			FFieldSystemCommand CommandOut("FRandomVector", new FRandomVector(1.f));
+			FFieldSystemCommand CommandIn = SaveAndLoad(CommandOut);
+			EXPECT_TRUE(TestFFieldSystemEquality(CommandIn, CommandOut));
+		}
+
+		{
+			FUniformScalar* UniformScalar = new FUniformScalar(1.f);
+			FRadialFalloff* RadialScalar = new FRadialFalloff(1.f, 3.f, 5.f, 7.f, 11.f, FVector(13, 17, 19));
+			FFieldSystemCommand CommandOut("FSumScalar", new FSumScalar(1.f, UniformScalar, RadialScalar, EFieldOperationType::Field_Substract));
+			FFieldSystemCommand CommandIn = SaveAndLoad(CommandOut);
+			EXPECT_TRUE(TestFFieldSystemEquality(CommandIn, CommandOut));
+		}
+
+		{
+			FUniformScalar* UniformScalar = new FUniformScalar(41.f);
+			FUniformVector* UniformVector = new FUniformVector(31.f, FVector(3, 5, 7));
+			FRadialVector* RadialVector = new FRadialVector(21.f, FVector(3, 5, 7));
+			FFieldSystemCommand CommandOut("FSumVector", new FSumVector(1.f, UniformScalar, UniformVector, RadialVector, EFieldOperationType::Field_Divide));
+			FFieldSystemCommand CommandIn = SaveAndLoad(CommandOut);
+			EXPECT_TRUE(TestFFieldSystemEquality(CommandIn, CommandOut));
+		}
+
+		// conversion fields
+		{
+			FUniformScalar* UniformScalar = new FUniformScalar(41.f);
+			FFieldSystemCommand CommandOut("FConversionField", new FConversionField<float,int32>(UniformScalar));
+			FFieldSystemCommand CommandIn = SaveAndLoad(CommandOut);
+			EXPECT_TRUE(TestFFieldSystemEquality(CommandIn, CommandOut));
+		}
+		{
+			FUniformInteger* UniformInteger = new FUniformInteger(3);
+			FFieldSystemCommand CommandOut("FConversionField", new FConversionField<int32, float>(UniformInteger));
+			FFieldSystemCommand CommandIn = SaveAndLoad(CommandOut);
+			EXPECT_TRUE(TestFFieldSystemEquality(CommandIn, CommandOut));
+		}
+
+		// culling fields
+		{
+			FUniformInteger* UniformInteger = new FUniformInteger(3);
+			FRadialFalloff* RadialScalar = new FRadialFalloff(1.f, 3.f, 5.f, 7.f, 11.f, FVector(13, 17, 19));
+			FFieldSystemCommand CommandOut("FCullingField", new FCullingField<int32>(RadialScalar,UniformInteger));
+			FFieldSystemCommand CommandIn = SaveAndLoad(CommandOut);
+			EXPECT_TRUE(TestFFieldSystemEquality(CommandIn, CommandOut));
+		}
+		{
+			FUniformScalar* UniformScalar = new FUniformScalar(3.f);
+			FRadialFalloff* RadialScalar = new FRadialFalloff(1.f, 3.f, 5.f, 7.f, 11.f, FVector(13, 17, 19));
+			FFieldSystemCommand CommandOut("FCullingField", new FCullingField<float>(RadialScalar, UniformScalar));
+			FFieldSystemCommand CommandIn = SaveAndLoad(CommandOut);
+			EXPECT_TRUE(TestFFieldSystemEquality(CommandIn, CommandOut));
+		}
+		{
+			FUniformVector* UniformVector = new FUniformVector(3.f);
+			FRadialFalloff* RadialScalar = new FRadialFalloff(1.f, 3.f, 5.f, 7.f, 11.f, FVector(13, 17, 19));
+			FFieldSystemCommand CommandOut("FCullingField", new FCullingField<FVector>(RadialScalar, UniformVector));
+			FFieldSystemCommand CommandIn = SaveAndLoad(CommandOut);
+			EXPECT_TRUE(TestFFieldSystemEquality(CommandIn, CommandOut));
+		}
+
+		// terminals
+		{
+			FFieldSystemCommand CommandOut("FReturnResultsTerminal", new FReturnResultsTerminal<int32>());
+			FFieldSystemCommand CommandIn = SaveAndLoad(CommandOut);
+			EXPECT_TRUE(TestFFieldSystemEquality(CommandIn, CommandOut));
+		}
+		{
+			FFieldSystemCommand CommandOut("FReturnResultsTerminal", new FReturnResultsTerminal<float>());
+			FFieldSystemCommand CommandIn = SaveAndLoad(CommandOut);
+			EXPECT_TRUE(TestFFieldSystemEquality(CommandIn, CommandOut));
+		}
+		{
+			FFieldSystemCommand CommandOut("FReturnResultsTerminal", new FReturnResultsTerminal<FVector>());
+			FFieldSystemCommand CommandIn = SaveAndLoad(CommandOut);
+			EXPECT_TRUE(TestFFieldSystemEquality(CommandIn, CommandOut));
+		}
+
+
+		// depth test with lots of nodes.
+		{
+			FUniformScalar* UniformScalar = new FUniformScalar(3);
+			FConversionField<float, int32>* ConversionFieldFI = new FConversionField<float, int32>(UniformScalar);
+
+			FBoxFalloff* BoxFalloff = new FBoxFalloff(1.f, 7.f, 9.f, 13.f, FTransform::Identity);
+			FCullingField<int32>* CullingFieldI = new FCullingField<int32>(BoxFalloff, ConversionFieldFI);
+
+			FUniformInteger* UniformInteger = new FUniformInteger(3);
+			FConversionField<int32, float>*  ConversionFieldIF = new FConversionField<int32, float>(UniformInteger);
+
+			FPlaneFalloff* PlaneFalloff = new FPlaneFalloff(1.f, 3.f, 5.f, 7.f, 100.f, FVector(9, 11, 13), FVector(17, 19, 23));
+			FCullingField<float>* CullingFieldF = new FCullingField<float>(PlaneFalloff, ConversionFieldIF);
+
+
+			FNoiseField* NoiseField2 = new FNoiseField(1.f, 3.f);
+			FRandomVector* RandomVector = new FRandomVector(1.f);
+			FCullingField<FVector>* CullingFieldV = new FCullingField<FVector>(NoiseField2, RandomVector);
+
+			FNoiseField* NoiseField = new FNoiseField(1.f, 3.f);
+			FSumScalar* SumScalar = new FSumScalar(1.f, CullingFieldF, NoiseField, EFieldOperationType::Field_Substract);
+
+			FUniformVector* UniformVector = new FUniformVector(1.f, FVector(3, 5, 7));
+			FSumVector* SumVector = new FSumVector(1.f, SumScalar, CullingFieldV, UniformVector, EFieldOperationType::Field_Divide);
+
+			FReturnResultsTerminal<int32>*  ReturnResultsTerminalI = new FReturnResultsTerminal<int32>();
+			FConversionField<int32, float>*  ConversionFieldIF2 = new FConversionField<int32, float>(ReturnResultsTerminalI);
+			FReturnResultsTerminal<float>* ReturnResultsTerminalF = new FReturnResultsTerminal<float>();
+			FSumScalar* SumScalar2 = new FSumScalar(1.f, ReturnResultsTerminalF, ConversionFieldIF2, EFieldOperationType::Field_Substract);
+
+
+			FReturnResultsTerminal<FVector>* ReturnResultsTerminalV =new FReturnResultsTerminal<FVector>();
+			FSumVector* SumVector2 = new FSumVector(1.f, SumScalar2, ReturnResultsTerminalV, SumVector, EFieldOperationType::Field_Divide);
+
+			FReturnResultsTerminal<FVector>* ReturnResultsTerminalV2 = new FReturnResultsTerminal<FVector>();
+			FConversionField<int32, float>*  ConversionFieldIF3 = new FConversionField<int32, float>(CullingFieldI);
+			FSumVector* SumVector3 = new FSumVector(1.f, ConversionFieldIF3, SumVector2, ReturnResultsTerminalV2, EFieldOperationType::Field_Divide);
+
+			FFieldSystemCommand CommandOut("DeepTreeOfEverything", SumVector3);
+			FFieldSystemCommand CommandIn = SaveAndLoad(CommandOut);
+			EXPECT_TRUE(TestFFieldSystemEquality(CommandIn, CommandOut));
+		}
+
+	}
+	//template void Fields_SerializeAPI<float>();
+}
diff --git a/Engine/Source/Programs/HeadlessChaos/Private/GeometryCollection/GeometryCollectionTestFramework.cpp b/Engine/Source/Programs/HeadlessChaos/Private/GeometryCollection/GeometryCollectionTestFramework.cpp
new file mode 100644
--- /dev/null
+++ b/Engine/Source/Programs/HeadlessChaos/Private/GeometryCollection/GeometryCollectionTestFramework.cpp
@@ -0,0 +1,362 @@
+// Copyright Epic Games, Inc. All Rights Reserved.
+
+#include "GeometryCollection/GeometryCollectionTestFramework.h"
+#include "GeometryCollection/GeometryCollectionTestUtility.h"
+
+#include "GeometryCollection/GeometryCollectionUtility.h"
+#include "GeometryCollection/GeometryCollectionAlgo.h"
+
+#include "Generators/SphereGenerator.h"
+#include "Generators/GridBoxMeshGenerator.h"
+
+#include "ChaosSolversModule.h"
+#include "Chaos/ErrorReporter.h"
+
+#include "../Resource/SphereGeometry.h"
+
+
+namespace GeometryCollectionTest
+{
+	TSharedPtr<FGeometryCollection> MakeSphereElement(
+		FTransform RootTranform, FTransform GeomTransform,
+		const int NumberOfMaterials = 2)
+	{
+		FSphereGenerator SphereGen;
+		SphereGen.Radius = 1.f;
+		SphereGen.NumPhi = 16;		// Vertical divisions
+		SphereGen.NumTheta = 16;	// Horizontal divisions
+		SphereGen.Generate();
+
+		// FSphereGenerator's points drift off the surface just a bit, so we correct for that.
+		Chaos::TSphere<float, 3> Sphere(Chaos::TVector<float, 3>(0), SphereGen.Radius);
+		Chaos::TVector<float, 3> Normal;
+		for (int32 Idx = 0; Idx < SphereGen.Vertices.Num(); Idx++)
+		{
+			auto& SrcPt = SphereGen.Vertices[Idx]; // double precision
+			Chaos::TVector<float, 3> Pt(SrcPt[0], SrcPt[1], SrcPt[2]); // single precision
+			const float Phi = Sphere.PhiWithNormal(Pt, Normal);
+			SrcPt[0] -= Phi * Normal[0];
+			SrcPt[1] -= Phi * Normal[1];
+			SrcPt[2] -= Phi * Normal[2];
+
+			// Ensure all the normals are pointing the right direction
+			check(FVector::DotProduct(FVector(SphereGen.Normals[Idx][0], SphereGen.Normals[Idx][1], SphereGen.Normals[Idx][2]), Normal) > 0.0f);
+		}
+
+		// Ensure each edge has 2 associated faces (no holes)
+		TMap<TPair<int32, int32>, uint8> EdgeCount;
+		for (auto& Tri : SphereGen.Triangles)
+		{
+			// Reverse triangle orientations
+			{
+				int32 Tmp = Tri[0];
+				Tri[0] = Tri[1];
+				Tri[1] = Tmp;
+			}
+			for (int32 i = 0; i < 3; i++)
+			{
+				int32 A = Tri[i];
+				int32 B = Tri[(i+1)%3];
+				TPair<int32, int32> Edge(
+					A < B ? A : B,
+					A > B ? A : B);
+				if (uint8* Count = EdgeCount.Find(Edge))
+				{
+					(*Count)++;
+				}
+				else
+				{
+					EdgeCount.Emplace(Edge) = 1;
+				}
+			}
+		}
+		for (auto& KV : EdgeCount)
+		{
+			check(KV.Value == 2);
+		}
+
+		return GeometryCollection::MakeMeshElement(
+			SphereGen.Vertices, 
+			SphereGen.Normals, 
+			SphereGen.Triangles, 
+			SphereGen.UVs, 
+			RootTranform, GeomTransform, NumberOfMaterials);
+	}
+
+	TSharedPtr<FGeometryCollection> MakeCubeElement(FTransform RootTranform, FTransform GeomTransform)
+	{
+		const TArray<FVector> PointsIn = { FVector(-1,1,-1), FVector(1,1,-1), FVector(1,-1,-1), FVector(-1,-1,-1), FVector(-1,1,1), FVector(1,1,1), FVector(1,-1,1), FVector(-1,-1,1)};
+		const TArray<FVector> NormalsIn = { FVector(-1,1,-1).GetSafeNormal(), FVector(1,1,-1).GetSafeNormal(), FVector(1,-1,-1).GetSafeNormal(), FVector(-1,-1,-1).GetSafeNormal(), FVector(-1,1,1).GetSafeNormal(), FVector(1,1,1).GetSafeNormal(), FVector(1,-1,1).GetSafeNormal(), FVector(-1,-1,1).GetSafeNormal() };
+		const TArray<FVector3i> TrianglesIn = { FVector3i(0,1,2),FVector3i(0,2,3), FVector3i(2,1,6),FVector3i(1,5,6),FVector3i(2,6,7),FVector3i(3,2,7),FVector3i(4,7,3),FVector3i(4,0,3),FVector3i(4,1,0),FVector3i(4,5,1),FVector3i(5,4,7),FVector3i(5,7,6)};
+		const TArray<FVector2f> UVsIn = { FVector2f(0.f,0.f), FVector2f(0.f,0.f), FVector2f(0.f,0.f), FVector2f(0.f,0.f),FVector2f(0.f,0.f), FVector2f(0.f,0.f), FVector2f(0.f,0.f), FVector2f(0.f,0.f) };
+		return GeometryCollection::MakeMeshElement(PointsIn, NormalsIn, TrianglesIn, UVsIn, RootTranform, GeomTransform);
+	}
+
+	TSharedPtr<FGeometryCollection> MakeTetrahedronElement(FTransform RootTranform, FTransform GeomTransform)
+	{
+		const TArray<FVector> PointsIn = { FVector(-1,1,-1), FVector(-1,-1,1), FVector(1,-1,-1), FVector(1,1,1) };
+		const TArray<FVector> NormalsIn = { FVector(-1,1,-1).GetSafeNormal(), FVector(-1,-1,1).GetSafeNormal(), FVector(1,-1,-1).GetSafeNormal(), FVector(1,1,1).GetSafeNormal() };
+		const TArray<FVector3i> TrianglesIn = { FVector3i(1,0,2), FVector3i(2,1,3), FVector3i(3,2,0), FVector3i(3,0,1) };
+		const TArray<FVector2f> UVsIn = { FVector2f(0.f,0.f), FVector2f(0.f,0.f), FVector2f(0.f,0.f), FVector2f(0.f,0.f) };
+		return GeometryCollection::MakeMeshElement(PointsIn, NormalsIn, TrianglesIn, UVsIn, RootTranform, GeomTransform);
+	}
+
+	TSharedPtr<FGeometryCollection> MakeImportedSphereElement(FTransform RootTranform, FTransform GeomTransform)
+	{
+		FGeometryCollection* Collection = FGeometryCollection::NewGeometryCollection(SphereGeometry::RawVertexArray, SphereGeometry::RawIndicesArray);
+		TManagedArray<FVector>& Vertices = Collection->Vertex;
+		for (int i = 0; i < Vertices.Num(); i++) Vertices[i] = GeomTransform.TransformPosition(Vertices[i]);
+		return TSharedPtr<FGeometryCollection>(Collection);
+	}
+
+
+	TSharedPtr<FGeometryCollection> MakeGriddedBoxElement(
+		FTransform RootTranform, FTransform GeomTransform,
+		const FVector& Extents = FVector(1, 1, 1),
+		const FIndex3i& EdgeVertices = FIndex3i(4, 4, 4),
+		const int NumberOfMaterials = 2)
+	{
+		FGridBoxMeshGenerator BoxGen;
+		BoxGen.Box = FOrientedBox3d(FVector(0), Extents); // box center, box dimensions
+		BoxGen.EdgeVertices = EdgeVertices;
+		BoxGen.Generate();
+
+		return GeometryCollection::MakeMeshElement(
+			BoxGen.Vertices,
+			BoxGen.Normals,
+			BoxGen.Triangles,
+			BoxGen.UVs,
+			FTransform::Identity, GeomTransform, NumberOfMaterials);
+	}
+
+	template <typename Traits>
+	WrapperBase* CommonInit(const TSharedPtr<FGeometryCollection>& RestCollection, const CreationParameters& Params)
+	{
+		FTransformCollection SingleTransform = FTransformCollection::SingleTransform();
+		for (int i = 0; i < Params.NestedTransforms.Num(); i++)
+		{
+			int ChildIndex = RestCollection->NumElements(FTransformCollection::TransformGroup) - 1;
+			int32 ParentIndex = RestCollection->AppendTransform(SingleTransform, Params.NestedTransforms[0]);
+			RestCollection->ParentTransforms(ParentIndex, ChildIndex);
+		}
+
+		TSharedPtr<Chaos::FChaosPhysicsMaterial> PhysicalMaterial = MakeShared<Chaos::FChaosPhysicsMaterial>();
+		InitMaterialToZero(PhysicalMaterial.Get());
+
+		TSharedPtr<FGeometryDynamicCollection> DynamicCollection = GeometryCollectionToGeometryDynamicCollection(RestCollection.Get(), Params.DynamicState);
+
+		FSimulationParameters SimulationParams;
+		{
+			SimulationParams.RestCollection = RestCollection.Get();
+			SimulationParams.PhysicalMaterial = MakeSerializable(PhysicalMaterial);
+			SimulationParams.Shared.Mass = Params.Mass;
+			SimulationParams.Shared.bMassAsDensity = Params.bMassAsDensity;
+			SimulationParams.Shared.SizeSpecificData[0].CollisionType = Params.CollisionType;
+			SimulationParams.Shared.SizeSpecificData[0].ImplicitType = Params.ImplicitType;
+			SimulationParams.Simulating = Params.Simulating;
+			SimulationParams.EnableClustering = Params.EnableClustering;
+			SimulationParams.InitialLinearVelocity = Params.InitialLinearVelocity;
+			SimulationParams.InitialVelocityType = Params.InitialVelocityType;
+			SimulationParams.DamageThreshold = Params.DamageThreshold;
+			SimulationParams.MaxClusterLevel = Params.MaxClusterLevel;	
+			SimulationParams.ClusterConnectionMethod = Params.ClusterConnectionMethod;
+			SimulationParams.RemoveOnFractureEnabled = Params.RemoveOnFractureEnabled;
+			SimulationParams.CollisionGroup = Params.CollisionGroup;
+			SimulationParams.ClusterGroupIndex = Params.ClusterGroupIndex;
+
+			FSharedSimulationSizeSpecificData Tmp;
+			Tmp.MinLevelSetResolution = Params.MinLevelSetResolution;
+			Tmp.MaxLevelSetResolution = Params.MaxLevelSetResolution;
+			SimulationParams.Shared.SizeSpecificData.Add(Tmp);
+
+			Chaos::FErrorReporter ErrorReporter;
+			BuildSimulationData(ErrorReporter, *RestCollection.Get(), SimulationParams.Shared);
+
+			FGeometryCollectionPhysicsProxy::InitializeDynamicCollection(*DynamicCollection, *RestCollection, SimulationParams);
+		}
+
+		FCollisionFilterData SimFilterData;
+		FCollisionFilterData QueryFilterData;
+
+		// Enable all collisions
+		SimFilterData.Word1 = 0xFFFF; // this body channel
+		SimFilterData.Word3 = 0xFFFF; // collision candidate channels
+
+		TGeometryCollectionPhysicsProxy<Traits>* PhysObject =
+			new TGeometryCollectionPhysicsProxy<Traits>(
+				nullptr,			// UObject owner
+				*DynamicCollection, // Game thread collection
+				SimulationParams,
+				SimFilterData,
+				QueryFilterData,
+				nullptr,			// Init func
+				nullptr,			// Cache sync func
+				nullptr);			// Final sync func
+		return new TGeometryCollectionWrapper<Traits>(PhysicalMaterial, RestCollection, DynamicCollection, PhysObject);
+	}
+
+	template <>
+	template<typename Traits>
+	WrapperBase* TNewSimulationObject<GeometryType::GeometryCollectionWithSingleRigid>::Init(const CreationParameters Params)
+	{
+		TSharedPtr<FGeometryCollection> RestCollection;
+		switch (Params.SimplicialType)
+		{
+		case ESimplicialType::Chaos_Simplicial_Box:
+			RestCollection = MakeCubeElement(Params.RootTransform, Params.GeomTransform );
+			break;
+		case ESimplicialType::Chaos_Simplicial_Sphere:
+			RestCollection = MakeSphereElement(Params.RootTransform, Params.GeomTransform);
+			break;
+		case ESimplicialType::Chaos_Simplicial_GriddleBox:
+			RestCollection = MakeGriddedBoxElement(Params.RootTransform, Params.GeomTransform);
+			break;
+		case ESimplicialType::Chaos_Simplicial_Tetrahedron:
+			RestCollection = MakeTetrahedronElement(Params.RootTransform, Params.GeomTransform);
+			break;
+		case ESimplicialType::Chaos_Simplicial_Imported_Sphere:
+			RestCollection = MakeImportedSphereElement(Params.RootTransform, Params.GeomTransform);
+			break;
+		case ESimplicialType::Chaos_Simplicial_None:
+			RestCollection = TSharedPtr<FGeometryCollection>(new FGeometryCollection());
+			RestCollection->AddElements(1, FGeometryCollection::GeometryGroup);
+			RestCollection->TransformIndex[0] = 0;
+			RestCollection->InnerRadius[0] = 1.f;	// Assume sphere w/radius 1
+			RestCollection->OuterRadius[0] = 1.f;	// Assume sphere w/radius 1
+			RestCollection->AddElements(1, FGeometryCollection::TransformGroup);
+			RestCollection->Transform[0] = Params.RootTransform;
+			RestCollection->Transform[0].NormalizeRotation();
+			break;
+		default:
+			check(false); // unimplemented!
+			break;
+		}
+
+		return CommonInit<Traits>(RestCollection, Params);
+	}
+
+	template <>
+	template<typename Traits>
+	WrapperBase* TNewSimulationObject<GeometryType::RigidFloor>::Init(const CreationParameters Params)
+	{
+		TSharedPtr<Chaos::FChaosPhysicsMaterial> PhysicalMaterial = MakeShared<Chaos::FChaosPhysicsMaterial>(); InitMaterialToZero(PhysicalMaterial.Get());
+		TGeometryParticle<float, 3>* Particle = TGeometryParticle<float, 3>::CreateParticle().Release();
+		Particle->SetGeometry(TUniquePtr<TPlane<float, 3>>(new TPlane<float, 3>(FVector(0), FVector(0, 0, 1))));
+
+		FCollisionFilterData FilterData;
+		FilterData.Word1 = 0xFFFF;
+		FilterData.Word3 = 0xFFFF;
+		Particle->SetShapeSimData(0, FilterData);
+
+
+		return new RigidBodyWrapper(PhysicalMaterial, Particle);
+	}
+
+	template <>
+	template<typename Traits>
+	WrapperBase* TNewSimulationObject<GeometryType::GeometryCollectionWithSuppliedRestCollection>::Init(const CreationParameters Params)
+	{
+		check(Params.RestCollection.IsValid());
+
+		return CommonInit<Traits>(Params.RestCollection, Params);
+	}
+
+	template<typename Traits>
+	TFramework<Traits>::TFramework(FrameworkParameters Parameters)
+		: Dt(Parameters.Dt)
+		, Module(FChaosSolversModule::GetModule())
+		, Solver(Module->CreateSolver<Traits>(nullptr, ESolverFlags::Standalone))
+	{
+		Module->ChangeThreadingMode(Parameters.ThreadingMode);
+	}
+
+	template<typename Traits>
+	TFramework<Traits>::~TFramework()
+	{
+		for (WrapperBase* Object : PhysicsObjects)
+		{
+			if (TGeometryCollectionWrapper<Traits>* GCW = Object->As<TGeometryCollectionWrapper<Traits>>())
+			{
+				delete GCW->PhysObject;
+			}
+			else if (RigidBodyWrapper* BCW = Object->As<RigidBodyWrapper>())
+			{
+				delete BCW->Particle;
+			}
+			delete Object;
+		}
+
+		for (FFieldSystemPhysicsProxy* FieldObject : FieldObjects)
+		{
+			delete FieldObject;
+		}
+
+		FChaosSolversModule::GetModule()->DestroySolver(Solver);
+	}
+
+	template<typename Traits>
+	void TFramework<Traits>::AddSimulationObject(WrapperBase * Object)
+	{
+		PhysicsObjects.Add(Object);
+	}
+
+	template<typename Traits>
+	void TFramework<Traits>::AddFieldObject(FFieldSystemPhysicsProxy* Object)
+	{
+		FieldObjects.Add(Object);
+	}
+
+	template<typename Traits>
+	void TFramework<Traits>::Initialize()
+	{
+		Solver->SetEnabled(true);
+
+		for (WrapperBase* Object : PhysicsObjects)
+		{
+			if (TGeometryCollectionWrapper<Traits>* GCW = Object->As<TGeometryCollectionWrapper<Traits>>())
+			{
+				Solver->RegisterObject(GCW->PhysObject);
+				Solver->AddDirtyProxy(GCW->PhysObject);
+			}
+			else if (RigidBodyWrapper* RBW = Object->As<RigidBodyWrapper>())
+			{
+				Solver->RegisterObject(RBW->Particle);
+				Solver->AddDirtyProxy(RBW->Particle->GetProxy());
+			}
+		}
+
+		for (FFieldSystemPhysicsProxy* FieldObject : FieldObjects)
+		{
+			Solver->RegisterObject(FieldObject);
+			Solver->AddDirtyProxy(FieldObject);
+		}
+
+		Solver->PushPhysicsState(Module->GetDispatcher());
+	}
+
+	template<typename Traits>
+	void TFramework<Traits>::Advance()
+	{
+		Solver->SyncEvents_GameThread();
+
+		Solver->AdvanceSolverBy(Dt);
+
+		Solver->BufferPhysicsResults();
+		Solver->FlipBuffers();
+		Solver->UpdateGameThreadStructures();
+	}
+	
+#define EVOLUTION_TRAIT(Trait) template class TFramework<Trait>;
+#include "Chaos/EvolutionTraits.inl"
+#undef EVOLUTION_TRAIT
+
+	
+#define EVOLUTION_TRAIT(Trait)\
+template WrapperBase* TNewSimulationObject<GeometryType::GeometryCollectionWithSingleRigid>::Init<Trait>(const CreationParameters Params);\
+template WrapperBase* TNewSimulationObject<GeometryType::RigidFloor>::Init<Trait>(const CreationParameters Params);\
+template WrapperBase* TNewSimulationObject<GeometryType::GeometryCollectionWithSuppliedRestCollection>::Init<Trait>(const CreationParameters Params);
+#include "Chaos/EvolutionTraits.inl"
+#undef EVOLUTION_TRAIT
+
+} // end namespace GeometryCollectionTest
+
diff --git a/Engine/Source/Programs/HeadlessChaos/Private/GeometryCollection/GeometryCollectionTestImplicitCapsule.cpp b/Engine/Source/Programs/HeadlessChaos/Private/GeometryCollection/GeometryCollectionTestImplicitCapsule.cpp
new file mode 100644
--- /dev/null
+++ b/Engine/Source/Programs/HeadlessChaos/Private/GeometryCollection/GeometryCollectionTestImplicitCapsule.cpp
@@ -0,0 +1,202 @@
+// Copyright Epic Games, Inc. All Rights Reserved.
+
+#include "GeometryCollection/GeometryCollectionTestImplicitCapsule.h"
+
+#include "Chaos/Capsule.h"
+
+namespace GeometryCollectionTest
+{
+	template <class T>
+	void RunTestComputeSamplePoints(const Chaos::TCapsule<T> &Capsule)
+	{
+		const Chaos::TCapsule<T> OACapsule = Chaos::TCapsule<T>::NewFromOriginAndAxis(Capsule.GetOrigin(), Capsule.GetAxis(), Capsule.GetHeight(), Capsule.GetRadius());
+		EXPECT_NEAR((Capsule.GetOrigin() - OACapsule.GetOrigin()).Size(), 0, KINDA_SMALL_NUMBER) << *FString("Capsule != OACapsule, origin.");
+		EXPECT_NEAR((Capsule.GetInsertion() - OACapsule.GetInsertion()).Size(), 0, KINDA_SMALL_NUMBER) << *FString("Capsule != OACapsule, insertion.");
+		EXPECT_NEAR((Capsule.GetAxis() - OACapsule.GetAxis()).Size(), 0, KINDA_SMALL_NUMBER) << *FString("Capsule != OACapsule, axis.");
+		EXPECT_NEAR(Capsule.GetHeight() - OACapsule.GetHeight(), 0, KINDA_SMALL_NUMBER) << *FString("Capsule != OACapsule, height.");
+		EXPECT_NEAR(Capsule.GetRadius() - OACapsule.GetRadius(), 0, KINDA_SMALL_NUMBER) << *FString("Capsule != OACapsule, radius.");
+
+		Chaos::TVector<T, 3> Point;
+		T Phi;
+		EXPECT_EQ(Capsule.GetType(), Chaos::ImplicitObjectType::Capsule) << *FString("Implicit object type is not 'capsule'.");
+
+		Point = Capsule.GetAxis();
+		EXPECT_LT(FMath::Abs(Point.Size() - 1.0), KINDA_SMALL_NUMBER) << *FString("Capsule axis is not unit length.");
+		
+		Point = Capsule.GetOrigin() + Capsule.GetAxis() * (Capsule.GetHeight() + Capsule.GetRadius() * 2);
+		EXPECT_LT((Point - Capsule.GetInsertion()).Size(), KINDA_SMALL_NUMBER) << *FString("Capsule is broken.");
+
+		Point = Capsule.GetInsertion();// Capsule.GetOrigin() + Capsule.GetAxis() * Capsule.GetHeight();
+		Phi = Capsule.SignedDistance(Point);
+		EXPECT_NEAR(Phi, 0, KINDA_SMALL_NUMBER) << *FString("Capsule failed phi surface (insertion) sanity test.");
+
+		Point = Capsule.GetOrigin() + Capsule.GetAxis() * (Capsule.GetHeight() + Capsule.GetRadius() * 2) * 0.25;
+		Phi = Capsule.SignedDistance(Point);
+		EXPECT_LE(Phi, (T)0.0) << *FString("Capsule failed phi depth (1/4 origin) sanity test.");
+
+		Point = Capsule.GetOrigin() + Capsule.GetAxis() * (Capsule.GetHeight() + Capsule.GetRadius() * 2) * 0.75;
+		Phi = Capsule.SignedDistance(Point);
+		EXPECT_LE(Phi, (T)0.0) << *FString("Capsule failed phi depth (3/4 origin) sanity test.");
+
+		EXPECT_NEAR((Capsule.GetCenter() - Chaos::TVector<T,3>(Capsule.GetOrigin() + Capsule.GetAxis() * (Capsule.GetHeight() + Capsule.GetRadius() * 2) * 0.5)).Size(), 0, KINDA_SMALL_NUMBER) << *FString("Capsule center is off mid axis.");
+
+		Point = Capsule.GetCenter();
+		Phi = Capsule.SignedDistance(Point);
+		EXPECT_LT(Phi, (T)0.0) << *FString("Capsule failed phi depth sanity test.");
+		
+		Point = Capsule.GetOrigin();
+		Phi = Capsule.SignedDistance(Point);
+		EXPECT_NEAR(FMath::Abs(Phi), 0, KINDA_SMALL_NUMBER) << *FString("Capsule failed phi surface (origin) sanity test.");
+
+		Point = Capsule.GetInsertion();
+		Phi = Capsule.SignedDistance(Point);
+		EXPECT_NEAR(FMath::Abs(Phi), 0, KINDA_SMALL_NUMBER) << *FString("Capsule failed phi surface (origin+axis*height) sanity test.");
+
+		Point = Capsule.GetOrigin() + Capsule.GetAxis() * Capsule.GetRadius() + Capsule.GetAxis().GetOrthogonalVector().GetSafeNormal() * Capsule.GetRadius();
+		Phi = Capsule.SignedDistance(Point);
+		EXPECT_NEAR(FMath::Abs(Phi), 0, KINDA_SMALL_NUMBER) << *FString("Capsule failed phi surface (origin+orthogonalAxis*radius) sanity test.");
+
+		Point = Capsule.GetCenter() + Capsule.GetAxis().GetOrthogonalVector().GetSafeNormal() * Capsule.GetRadius();
+		Phi = Capsule.SignedDistance(Point);
+		EXPECT_NEAR(FMath::Abs(Phi), 0, KINDA_SMALL_NUMBER) << *FString("Capsule failed phi surface (center+orthogonalAxis*radius) sanity test.");
+
+		TArray<Chaos::TVector<T, 3>> Points = Capsule.ComputeSamplePoints(100);
+		check(Points.Num() == 100);
+		Point[0] = TNumericLimits<T>::Max();
+		T MinPhi = TNumericLimits<T>::Max();
+		T MaxPhi = -TNumericLimits<T>::Max();
+		for (int32 i=0; i < Points.Num(); i++)
+		{
+			const Chaos::TVector<T, 3> &Pt = Points[i];
+
+			Phi = Capsule.SignedDistance(Pt);
+			MinPhi = FMath::Min(Phi, MinPhi);
+			MaxPhi = FMath::Max(Phi, MaxPhi);
+			
+			const bool Differs = Pt != Point;
+			check(Differs);
+			EXPECT_TRUE(Differs) << *FString("Produced a redundant value.");
+			Point = Pt;
+		}
+
+		const bool OnSurface = FMath::Abs(MinPhi) <= KINDA_SMALL_NUMBER && FMath::Abs(MaxPhi) <= KINDA_SMALL_NUMBER;
+		check(OnSurface);
+		EXPECT_TRUE(OnSurface) << *FString("Produced a point not on the surface of the capsule.");
+	}
+
+	template <class T>
+	void TestComputeSamplePoints_Capsule()
+	{
+		//
+		// Height == 1
+		//
+
+		// At the origin with radius 1
+		{
+			Chaos::TCapsule<T> Capsule(Chaos::TVector<T, 3>(0,0,0), Chaos::TVector<T, 3>(0,0,1), (T)1.0);
+			RunTestComputeSamplePoints(Capsule);
+		}
+		// At the origin with radius > 1
+		{
+			Chaos::TCapsule<T> Capsule(Chaos::TVector<T, 3>(0,0,0), Chaos::TVector<T, 3>(0,0,1), (T)10.0);
+			RunTestComputeSamplePoints(Capsule);
+		}
+		// At the origin with radius < 1
+		{
+			Chaos::TCapsule<T> Capsule(Chaos::TVector<T, 3>(0,0,0), Chaos::TVector<T, 3>(0,0,1), (T)0.1);
+			RunTestComputeSamplePoints(Capsule);
+		}
+		// Off the origin with radius 1
+		{
+			Chaos::TCapsule<T> Capsule(Chaos::TVector<T, 3>(10,10,10), Chaos::TVector<T, 3>(10,10,11), (T)1.0);
+			RunTestComputeSamplePoints(Capsule);
+		}
+		// Off the origin with radius > 1
+		{
+			Chaos::TCapsule<T> Capsule(Chaos::TVector<T, 3>(10,10,10), Chaos::TVector<T, 3>(10,10,11), (T)10.0);
+			RunTestComputeSamplePoints(Capsule);
+		}
+		// Off the origin with radius < 1
+		{
+			Chaos::TCapsule<T> Capsule(Chaos::TVector<T, 3>(10,10,10), Chaos::TVector<T, 3>(10,10,11), (T)0.1);
+			RunTestComputeSamplePoints(Capsule);
+		}
+
+		//
+		// Height > 1
+		//
+
+		// At the origin with radius 1
+		{
+			Chaos::TCapsule<T> Capsule(Chaos::TVector<T, 3>(0,0,0), Chaos::TVector<T, 3>(0,0,10), (T)1.0);
+			RunTestComputeSamplePoints(Capsule);
+		}
+		// At the origin with radius > 1
+		{
+			Chaos::TCapsule<T> Capsule(Chaos::TVector<T, 3>(0,0,0), Chaos::TVector<T, 3>(0,0,10), (T)10.0);
+			RunTestComputeSamplePoints(Capsule);
+		}
+		// At the origin with radius < 1
+		{
+			Chaos::TCapsule<T> Capsule(Chaos::TVector<T, 3>(0,0,0), Chaos::TVector<T, 3>(0,0,10), (T)0.1);
+			RunTestComputeSamplePoints(Capsule);
+		}
+		// Off the origin with radius 1
+		{
+			Chaos::TCapsule<T> Capsule(Chaos::TVector<T, 3>(10,10,10), Chaos::TVector<T, 3>(10,10,21), (T)1.0);
+			RunTestComputeSamplePoints(Capsule);
+		}
+		// Off the origin with radius > 1
+		{
+			Chaos::TCapsule<T> Capsule(Chaos::TVector<T, 3>(10,10,10), Chaos::TVector<T, 3>(10,10,21), (T)10.0);
+			RunTestComputeSamplePoints(Capsule);
+		}
+		// Off the origin with radius < 1
+		{
+			Chaos::TCapsule<T> Capsule(Chaos::TVector<T, 3>(10,10,10), Chaos::TVector<T, 3>(10,10,21), (T)0.1);
+			RunTestComputeSamplePoints(Capsule);
+		}
+
+		// 
+		// Off axis
+		//
+
+		// At the origin with radius 1
+		{
+			Chaos::TCapsule<T> Capsule(Chaos::TVector<T, 3>(0,0,0), Chaos::TVector<T, 3>(1,1,1), (T)1.0);
+			RunTestComputeSamplePoints(Capsule);
+		}
+		// At the origin with radius > 1
+		{
+			Chaos::TCapsule<T> Capsule(Chaos::TVector<T, 3>(0,0,0), Chaos::TVector<T, 3>(1,1,1), (T)10.0);
+			RunTestComputeSamplePoints(Capsule);
+		}
+		// At the origin with radius < 1
+		{
+			Chaos::TCapsule<T> Capsule(Chaos::TVector<T, 3>(0,0,0), Chaos::TVector<T, 3>(1,1,1), (T)0.1);
+			RunTestComputeSamplePoints(Capsule);
+		}
+		// Off the origin with radius 1
+		{
+			Chaos::TCapsule<T> Capsule(Chaos::TVector<T, 3>(10,10,10), Chaos::TVector<T, 3>(11,11,11), (T)1.0);
+			RunTestComputeSamplePoints(Capsule);
+		}
+		// Off the origin with radius > 1
+		{
+			Chaos::TCapsule<T> Capsule(Chaos::TVector<T, 3>(10,10,10), Chaos::TVector<T, 3>(11,11,11), (T)10.0);
+			RunTestComputeSamplePoints(Capsule);
+		}
+		// Off the origin with radius < 1
+		{
+			Chaos::TCapsule<T> Capsule(Chaos::TVector<T, 3>(10,10,10), Chaos::TVector<T, 3>(11,11,11), (T)0.1);
+			RunTestComputeSamplePoints(Capsule);
+		}
+	}
+
+	template <class T>
+	void TestImplicitCapsule()
+	{
+		TestComputeSamplePoints_Capsule<T>();
+	}
+	template void TestImplicitCapsule<float>();
+} // namespace GeometryCollectionTest
