Merge geometry collection into headless chaos.

diff --git a/Engine/Source/Programs/HeadlessChaos/Private/GeometryCollection/GeometryCollectionTestImplicitCylinder.cpp b/Engine/Source/Programs/HeadlessChaos/Private/GeometryCollection/GeometryCollectionTestImplicitCylinder.cpp
new file mode 100644
--- /dev/null
+++ b/Engine/Source/Programs/HeadlessChaos/Private/GeometryCollection/GeometryCollectionTestImplicitCylinder.cpp
@@ -0,0 +1,195 @@
+// Copyright Epic Games, Inc. All Rights Reserved.
+
+#include "GeometryCollection/GeometryCollectionTestImplicitCylinder.h"
+
+#include "Chaos/Cylinder.h"
+
+namespace GeometryCollectionTest
+{
+	template <class T>
+	void RunTestComputeSamplePoints(const Chaos::TCylinder<T> &Cylinder)
+	{
+		Chaos::TVector<T, 3> Point;
+		T Phi;
+		EXPECT_EQ(Cylinder.GetType(), Chaos::ImplicitObjectType::Cylinder) << *FString("Implicit object type is not 'cylinder'.");
+
+		Point = Cylinder.GetAxis();
+		EXPECT_LT(FMath::Abs(Point.Size() - 1.0), KINDA_SMALL_NUMBER) << *FString("Cylinder axis is not unit length.");
+		
+		Point = Cylinder.GetOrigin() + Cylinder.GetAxis() * Cylinder.GetHeight();
+		EXPECT_LT((Point - Cylinder.GetInsertion()).Size(), KINDA_SMALL_NUMBER) << *FString("Cylinder is broken.");
+
+		Point = Cylinder.GetInsertion();// Cylinder.GetOrigin() + Cylinder.GetAxis() * Cylinder.GetHeight();
+		Phi = Cylinder.SignedDistance(Point);
+		EXPECT_LE(Phi, KINDA_SMALL_NUMBER) << *FString("Cylinder failed phi surface (insertion) sanity test.");
+
+		Point = Cylinder.GetOrigin() + Cylinder.GetAxis() * Cylinder.GetHeight() * 0.25;
+		Phi = Cylinder.SignedDistance(Point);
+		EXPECT_LE(Phi, (T)0.0) << *FString("Cylinder failed phi depth (1/4 origin) sanity test.");
+
+		Point = Cylinder.GetOrigin() + Cylinder.GetAxis() * Cylinder.GetHeight() * 0.75;
+		Phi = Cylinder.SignedDistance(Point);
+		EXPECT_LE(Phi, (T)0.0) << *FString("Cylinder failed phi depth (3/4 origin) sanity test.");
+
+		EXPECT_LE((Cylinder.GetCenter() - Chaos::TVector<T,3>(Cylinder.GetOrigin() + Cylinder.GetAxis() * Cylinder.GetHeight() * 0.5)).Size(), KINDA_SMALL_NUMBER) << *FString("Cylinder center is off mid axis.");
+
+		Point = Cylinder.GetCenter();
+		Phi = Cylinder.SignedDistance(Point);
+		EXPECT_LT(Phi, (T)0.0) << *FString("Cylinder failed phi depth sanity test.");
+		
+		Point = Cylinder.GetOrigin();
+		Phi = Cylinder.SignedDistance(Point);
+		EXPECT_LE(FMath::Abs(Phi), KINDA_SMALL_NUMBER) << *FString("Cylinder failed phi surface (origin) sanity test.");
+
+		Point = Cylinder.GetOrigin() + Cylinder.GetAxis() * Cylinder.GetHeight();
+		Phi = Cylinder.SignedDistance(Point);
+		EXPECT_LE(FMath::Abs(Phi), KINDA_SMALL_NUMBER) << *FString("Cylinder failed phi surface (origin+axis*height) sanity test.");
+
+		Point = Cylinder.GetOrigin() + Cylinder.GetAxis().GetOrthogonalVector().GetSafeNormal() * Cylinder.GetRadius();
+		Phi = Cylinder.SignedDistance(Point);
+		EXPECT_LE(FMath::Abs(Phi), KINDA_SMALL_NUMBER) << *FString("Cylinder failed phi surface (origin+orthogonalAxis*radius) sanity test.");
+
+		Point = Cylinder.GetCenter() + Cylinder.GetAxis().GetOrthogonalVector().GetSafeNormal() * Cylinder.GetRadius();
+		Phi = Cylinder.SignedDistance(Point);
+		EXPECT_LE(FMath::Abs(Phi), KINDA_SMALL_NUMBER) << *FString("Cylinder failed phi surface (center+orthogonalAxis*radius) sanity test.");
+
+		TArray<Chaos::TVector<T, 3>> Points = Cylinder.ComputeSamplePoints(100);
+		check(Points.Num() == 100);
+		Point[0] = TNumericLimits<T>::Max();
+		T MinPhi = TNumericLimits<T>::Max();
+		T MaxPhi = -TNumericLimits<T>::Max();
+		for (int32 i=0; i < Points.Num(); i++)
+		{
+			const Chaos::TVector<T, 3> &Pt = Points[i];
+
+			Phi = Cylinder.SignedDistance(Pt);
+			MinPhi = FMath::Min(Phi, MinPhi);
+			MaxPhi = FMath::Max(Phi, MaxPhi);
+			
+			const bool Differs = Pt != Point;
+			check(Differs);
+			EXPECT_TRUE(Differs) << *FString("Produced a redundant value.");
+			Point = Pt;
+		}
+
+		const bool OnSurface = FMath::Abs(MinPhi) <= KINDA_SMALL_NUMBER && FMath::Abs(MaxPhi) <= KINDA_SMALL_NUMBER;
+		check(OnSurface);
+		EXPECT_TRUE(OnSurface) << *FString("Produced a point not on the surface of the cylinder.");
+	}
+
+	template <class T>
+	void TestComputeSamplePoints_Cylinder()
+	{
+		//
+		// Height == 1
+		//
+
+		// At the origin with radius 1
+		{
+			Chaos::TCylinder<T> Cylinder(Chaos::TVector<T, 3>(0,0,0), Chaos::TVector<T, 3>(0,0,1), (T)1.0);
+			RunTestComputeSamplePoints(Cylinder);
+		}
+		// At the origin with radius > 1
+		{
+			Chaos::TCylinder<T> Cylinder(Chaos::TVector<T, 3>(0,0,0), Chaos::TVector<T, 3>(0,0,1), (T)10.0);
+			RunTestComputeSamplePoints(Cylinder);
+		}
+		// At the origin with radius < 1
+		{
+			Chaos::TCylinder<T> Cylinder(Chaos::TVector<T, 3>(0,0,0), Chaos::TVector<T, 3>(0,0,1), (T)0.1);
+			RunTestComputeSamplePoints(Cylinder);
+		}
+		// Off the origin with radius 1
+		{
+			Chaos::TCylinder<T> Cylinder(Chaos::TVector<T, 3>(10,10,10), Chaos::TVector<T, 3>(10,10,11), (T)1.0);
+			RunTestComputeSamplePoints(Cylinder);
+		}
+		// Off the origin with radius > 1
+		{
+			Chaos::TCylinder<T> Cylinder(Chaos::TVector<T, 3>(10,10,10), Chaos::TVector<T, 3>(10,10,11), (T)10.0);
+			RunTestComputeSamplePoints(Cylinder);
+		}
+		// Off the origin with radius < 1
+		{
+			Chaos::TCylinder<T> Cylinder(Chaos::TVector<T, 3>(10,10,10), Chaos::TVector<T, 3>(10,10,11), (T)0.1);
+			RunTestComputeSamplePoints(Cylinder);
+		}
+
+		//
+		// Height > 1
+		//
+
+		// At the origin with radius 1
+		{
+			Chaos::TCylinder<T> Cylinder(Chaos::TVector<T, 3>(0,0,0), Chaos::TVector<T, 3>(0,0,10), (T)1.0);
+			RunTestComputeSamplePoints(Cylinder);
+		}
+		// At the origin with radius > 1
+		{
+			Chaos::TCylinder<T> Cylinder(Chaos::TVector<T, 3>(0,0,0), Chaos::TVector<T, 3>(0,0,10), (T)10.0);
+			RunTestComputeSamplePoints(Cylinder);
+		}
+		// At the origin with radius < 1
+		{
+			Chaos::TCylinder<T> Cylinder(Chaos::TVector<T, 3>(0,0,0), Chaos::TVector<T, 3>(0,0,10), (T)0.1);
+			RunTestComputeSamplePoints(Cylinder);
+		}
+		// Off the origin with radius 1
+		{
+			Chaos::TCylinder<T> Cylinder(Chaos::TVector<T, 3>(10,10,10), Chaos::TVector<T, 3>(10,10,21), (T)1.0);
+			RunTestComputeSamplePoints(Cylinder);
+		}
+		// Off the origin with radius > 1
+		{
+			Chaos::TCylinder<T> Cylinder(Chaos::TVector<T, 3>(10,10,10), Chaos::TVector<T, 3>(10,10,21), (T)10.0);
+			RunTestComputeSamplePoints(Cylinder);
+		}
+		// Off the origin with radius < 1
+		{
+			Chaos::TCylinder<T> Cylinder(Chaos::TVector<T, 3>(10,10,10), Chaos::TVector<T, 3>(10,10,21), (T)0.1);
+			RunTestComputeSamplePoints(Cylinder);
+		}
+
+		// 
+		// Off axis
+		//
+
+		// At the origin with radius 1
+		{
+			Chaos::TCylinder<T> Cylinder(Chaos::TVector<T, 3>(0,0,0), Chaos::TVector<T, 3>(1,1,1), (T)1.0);
+			RunTestComputeSamplePoints(Cylinder);
+		}
+		// At the origin with radius > 1
+		{
+			Chaos::TCylinder<T> Cylinder(Chaos::TVector<T, 3>(0,0,0), Chaos::TVector<T, 3>(1,1,1), (T)10.0);
+			RunTestComputeSamplePoints(Cylinder);
+		}
+		// At the origin with radius < 1
+		{
+			Chaos::TCylinder<T> Cylinder(Chaos::TVector<T, 3>(0,0,0), Chaos::TVector<T, 3>(1,1,1), (T)0.1);
+			RunTestComputeSamplePoints(Cylinder);
+		}
+		// Off the origin with radius 1
+		{
+			Chaos::TCylinder<T> Cylinder(Chaos::TVector<T, 3>(10,10,10), Chaos::TVector<T, 3>(11,11,11), (T)1.0);
+			RunTestComputeSamplePoints(Cylinder);
+		}
+		// Off the origin with radius > 1
+		{
+			Chaos::TCylinder<T> Cylinder(Chaos::TVector<T, 3>(10,10,10), Chaos::TVector<T, 3>(11,11,11), (T)10.0);
+			RunTestComputeSamplePoints(Cylinder);
+		}
+		// Off the origin with radius < 1
+		{
+			Chaos::TCylinder<T> Cylinder(Chaos::TVector<T, 3>(10,10,10), Chaos::TVector<T, 3>(11,11,11), (T)0.1);
+			RunTestComputeSamplePoints(Cylinder);
+		}
+	}
+
+	template <class T>
+	void TestImplicitCylinder()
+	{
+		TestComputeSamplePoints_Cylinder<T>();
+	}
+	template void TestImplicitCylinder<float>();
+} // namespace GeometryCollectionTest
diff --git a/Engine/Source/Programs/HeadlessChaos/Private/GeometryCollection/GeometryCollectionTestImplicitSphere.cpp b/Engine/Source/Programs/HeadlessChaos/Private/GeometryCollection/GeometryCollectionTestImplicitSphere.cpp
new file mode 100644
--- /dev/null
+++ b/Engine/Source/Programs/HeadlessChaos/Private/GeometryCollection/GeometryCollectionTestImplicitSphere.cpp
@@ -0,0 +1,433 @@
+// Copyright Epic Games, Inc. All Rights Reserved.
+
+#include "GeometryCollection/GeometryCollectionTestImplicitSphere.h"
+
+#include "Chaos/Pair.h"
+#include "Chaos/Sphere.h"
+
+#include <string>
+
+namespace GeometryCollectionTest
+{
+	//==========================================================================
+	// FindClosestIntersection() tests
+	//==========================================================================
+
+	typedef Chaos::Pair<Chaos::TVector<float, 3>, bool> IntersectionResult;
+
+	template <class T>
+	void RunTestFindClosestIntersectionHelper(
+		const FString TestName,
+		const Chaos::TVector<T, 3>& Center, 
+		const T Radius,
+		const T Thickness,
+		const Chaos::TVector<T, 3>& StartPt,
+		const Chaos::TVector<T, 3>& EndPt,
+		const IntersectionResult ExpectedRes)
+	{
+		Chaos::TSphere<T, 3> Sphere(Center, Radius);
+		IntersectionResult Res = Sphere.FindClosestIntersection(StartPt, EndPt, Thickness);
+		EXPECT_TRUE(
+			(!Res.Second && !ExpectedRes.Second) ||
+			(Res.Second && ExpectedRes.Second && (Res.First - ExpectedRes.First).Size() < SMALL_NUMBER)) <<
+			*FString::Printf(
+				TEXT(
+					"%s - Sphere((%g, %g, %g), %g).FindClosestIntersection("
+						"StartPt(%g, %g, %g), "
+						"EndPt(%g, %g, %g), "
+						"Thickness:%g) = "
+					"RESULT: %d, Point: (%g, %g, %g), "
+					"EXPECTED: %d (%g, %g, %g)."),
+				*TestName,
+				Center[0], Center[1], Center[2],
+				Radius,
+
+				StartPt[0], StartPt[1], StartPt[2],
+				EndPt[0], EndPt[1], EndPt[2],
+				Thickness,
+				
+				Res.Second ? 1 : 0,
+				Res.First[0], Res.First[1], Res.First[2],
+				
+				ExpectedRes.Second ? 1 : 0,
+				ExpectedRes.First[0], ExpectedRes.First[1], ExpectedRes.First[2]
+			);
+	}
+
+	template <class T>
+	void RunTestFindClosestIntersection_Tangent(
+		const Chaos::TVector<T, 3>& Center,
+		const T Radius,
+		const T Thickness)
+	{
+		Chaos::TVector<T, 3> StartPt;
+		Chaos::TVector<T, 3> EndPt;
+
+		// Start & End points coincident, lying on the surface of the sphere w/zero thickness
+		// Fails with no collision:
+		// "RunTestFindClosestIntersection_Tangent1 - Sphere((0, 0, 0), 1).FindClosestIntersection(StartPt(0, 0, 1), EndPt(0, 0, 1), Thickness:0.1) = RESULT: 0, Point: (0, 0, 0), EXPECTED: 1 (0, 0, 1.1)."	FString
+		StartPt = EndPt = Center + Chaos::TVector<T,3>(0, 0, Radius);
+		//RunTestFindClosestIntersectionHelper("RunTestFindClosestIntersection_Tangent1", Center, Radius, Thickness,
+		//	StartPt, EndPt, 
+		//	IntersectionResult(Center+Chaos::TVector<T,3>(0,0,Radius+Thickness), true));
+
+		StartPt = Center + Chaos::TVector<T, 3>(Radius, 0, Radius+Thickness);
+		EndPt = Center + Chaos::TVector<T, 3>(-Radius, 0, Radius+Thickness);
+		RunTestFindClosestIntersectionHelper("RunTestFindClosestIntersection_Tangent2", Center, Radius, Thickness,
+			StartPt, EndPt, 
+			IntersectionResult(Center+Chaos::TVector<T,3>(0,0,Radius+Thickness), true));
+		
+		// Start & End points coincident, lying on the surface of the sphere w/thickness
+		// Fails with incorrect collision point:
+		// "RunTestFindClosestIntersection_Tangent3 - Sphere((0, 0, 0), 1).FindClosestIntersection(StartPt(0, 0, 1.1), EndPt(0, 0, 1.1), Thickness:0.1) = RESULT: 1, Point: (0, 0, 1), EXPECTED: 1 (0, 0, 1.1)."	FString
+		// "RunTestFindClosestIntersection_Tangent3 - Sphere((0, 0, 0), 1).FindClosestIntersection(StartPt(0, 0, 1.5), EndPt(0, 0, 1.5), Thickness:0.5) = RESULT: 1, Point: (0, 0, 1), EXPECTED: 1 (0, 0, 1.5)."	FString
+		// "RunTestFindClosestIntersection_Tangent3 - Sphere((0, 0, 0), 1).FindClosestIntersection(StartPt(0, 0, 2), EndPt(0, 0, 2), Thickness:1) = RESULT: 1, Point: (0, 0, 1), EXPECTED: 1 (0, 0, 2)."	FString
+		// "RunTestFindClosestIntersection_Tangent3 - Sphere((0, 0, 0), 1).FindClosestIntersection(StartPt(0, 0, 3), EndPt(0, 0, 3), Thickness:2) = RESULT: 1, Point: (0, 0, 1), EXPECTED: 1 (0, 0, 3)."	FString
+		StartPt = EndPt = Center + Chaos::TVector<T,3>(0, 0, Radius+Thickness);
+		//RunTestFindClosestIntersectionHelper("RunTestFindClosestIntersection_Tangent3", Center, Radius, Thickness,
+		//	StartPt, EndPt, 
+		//	IntersectionResult(Center+Chaos::TVector<T,3>(0,0,Radius+Thickness), true));
+
+		StartPt = Center + Chaos::TVector<T, 3>(Radius, 0, Radius+Thickness);
+		EndPt = Center + Chaos::TVector<T, 3>(-Radius, 0, Radius+Thickness);
+		RunTestFindClosestIntersectionHelper("RunTestFindClosestIntersection_Tangent4", Center, Radius, Thickness,
+			StartPt, EndPt, 
+			IntersectionResult(Center+Chaos::TVector<T,3>(0,0,Radius+Thickness), true));
+
+		// Radius + Thickness + 1 (miss)
+		StartPt = EndPt = Center + Chaos::TVector<T,3>(0, 0, Radius+Thickness+1.0);
+		RunTestFindClosestIntersectionHelper("RunTestFindClosestIntersection_Tangent5", Center, Radius, Thickness,
+			StartPt, EndPt, 
+			IntersectionResult(Center+Chaos::TVector<T,3>(Radius+Thickness,0,0), false));
+
+		StartPt = Center + Chaos::TVector<T, 3>(Radius, 0, Radius+Thickness+1.0);
+		EndPt = Center + Chaos::TVector<T, 3>(-Radius, 0, Radius+Thickness+1.0);
+		RunTestFindClosestIntersectionHelper("RunTestFindClosestIntersection_Tangent6", Center, Radius, Thickness,
+			StartPt, EndPt, 
+			IntersectionResult(Center+Chaos::TVector<T,3>(0,0,Radius+Thickness), false));
+	}
+
+	template <class T>
+	void RunTestFindClosestIntersection_Shallow(
+		const Chaos::TVector<T, 3>& Center,
+		const T Radius,
+		const T Thickness)
+	{
+		Chaos::TVector<T, 3> StartPt(0, 0, Radius * .999999);
+		Chaos::TVector<T, 3> EndPt = StartPt;
+
+		// Fails w/no collision:
+		// "RunTestFindClosestIntersection_Shallow1 - Sphere((0, 0, 0), 1).FindClosestIntersection(StartPt(0, 0, 0.999999), EndPt(0, 0, 0.999999), Thickness:0.1) = RESULT: 0, Point: (0, 0, 0), EXPECTED: 1 (1.1, 0...	FString
+		//RunTestFindClosestIntersectionHelper("RunTestFindClosestIntersection_Shallow1", Center, Radius, Thickness,
+		//	StartPt, EndPt, 
+		//	IntersectionResult(Center+Chaos::TVector<T,3>(Radius+Thickness,0,0), true));
+
+		// Fails w/no collision:
+		// "RunTestFindClosestIntersection_Shallow2 - Sphere((0, 0, 0), 1).FindClosestIntersection(StartPt(0, 0, 0.999999), EndPt(0, 0, 1), Thickness:0.1) = RESULT: 0, Point: (0, 0, 0), EXPECTED: 1 (0, 0, 1.1)."	FString
+		EndPt += Chaos::TVector<T, 3>(0, 0, Radius * .000001);
+		//RunTestFindClosestIntersectionHelper("RunTestFindClosestIntersection_Shallow2", Center, Radius, Thickness,
+		//	StartPt, EndPt, 
+		//	IntersectionResult(Center+Chaos::TVector<T,3>(0,0,Radius+Thickness), true));
+
+		// Fails w/no collision:
+		// "RunTestFindClosestIntersection_Shallow3 - Sphere((0, 0, 0), 1).FindClosestIntersection(StartPt(0, 0, 0.999999), EndPt(0, 0, 1), Thickness:0.1) = RESULT: 0, Point: (0, 0, 0), EXPECTED: 1 (0, 0, 1.1)."	FString
+		EndPt += Chaos::TVector<T, 3>(0, 0, Radius * .000001);
+		//RunTestFindClosestIntersectionHelper("RunTestFindClosestIntersection_Shallow3", Center, Radius, Thickness,
+		//	StartPt, EndPt, 
+		//	IntersectionResult(Center+Chaos::TVector<T,3>(0,0,Radius+Thickness), true));
+
+		EndPt += Chaos::TVector<T, 3>(0, 0, Thickness);
+		RunTestFindClosestIntersectionHelper("RunTestFindClosestIntersection_Shallow3.1", Center, Radius, Thickness,
+			StartPt, EndPt, 
+			IntersectionResult(Center+Chaos::TVector<T,3>(0,0,Radius+Thickness), true));
+
+		// Fails w/no collision:
+		// "RunTestFindClosestIntersection_Shallow4 - Sphere((0, 0, 0), 1).FindClosestIntersection(StartPt(0, 0, 1), EndPt(0, 0, 0.999999), Thickness:0.1) = RESULT: 0, Point: (0, 0, 0), EXPECTED: 1 (0, 0, 1.1)."	FString
+		EndPt = StartPt;
+		StartPt = EndPt + Chaos::TVector<T, 3>(0, 0, Radius * .000001);
+		//RunTestFindClosestIntersectionHelper("RunTestFindClosestIntersection_Shallow4", Center, Radius, Thickness,
+		//	StartPt, EndPt, 
+		//	IntersectionResult(Center+Chaos::TVector<T,3>(0,0,Radius+Thickness), true));
+
+		// Fails w/no collision:
+		// "RunTestFindClosestIntersection_Shallow5 - Sphere((0, 0, 0), 1).FindClosestIntersection(StartPt(0, 0, 1), EndPt(0, 0, 0.999999), Thickness:0.1) = RESULT: 0, Point: (0, 0, 0), EXPECTED: 1 (0, 0, 1.1)."	FString
+		StartPt += Chaos::TVector<T, 3>(0, 0, Radius * .000001);
+		//RunTestFindClosestIntersectionHelper("RunTestFindClosestIntersection_Shallow5", Center, Radius, Thickness,
+		//	StartPt, EndPt, 
+		//	IntersectionResult(Center+Chaos::TVector<T,3>(0,0,Radius+Thickness), true));
+
+		StartPt += Chaos::TVector<T, 3>(0, 0, Thickness);
+		RunTestFindClosestIntersectionHelper("RunTestFindClosestIntersection_Shallow5.1", Center, Radius, Thickness,
+			StartPt, EndPt, 
+			IntersectionResult(Center+Chaos::TVector<T,3>(0,0,Radius+Thickness), true));
+	}
+
+	template <class T>
+	void RunTestFindClosestIntersection_Mid(
+		const Chaos::TVector<T, 3>& Center, 
+		const T Radius,
+		const T Thickness)
+	{
+		Chaos::TVector<T, 3> StartPt(0,0,Radius * .5);
+		Chaos::TVector<T, 3> EndPt = StartPt;
+
+		// Fails w/no collision:
+		// "RunTestFindClosestIntersection_Mid1 - Sphere((0, 0, 0), 1).FindClosestIntersection(StartPt(0, 0, 0.5), EndPt(0, 0, 0.5), Thickness:0.1) = RESULT: 0, Point: (0, 0, 0), EXPECTED: 1 (1.1, 0, 0)."	FString
+		//RunTestFindClosestIntersectionHelper("RunTestFindClosestIntersection_Mid1", Center, Radius, Thickness,
+		//	StartPt, EndPt, 
+		//	IntersectionResult(Center+Chaos::TVector<T,3>(Radius+Thickness,0,0), true));
+
+		// Fails w/no collision:
+		// "RunTestFindClosestIntersection_Mid2 - Sphere((0, 0, 0), 1).FindClosestIntersection(StartPt(0, 0, 0.5), EndPt(0, 0, 1), Thickness:0.1) = RESULT: 0, Point: (0, 0, 0), EXPECTED: 1 (0, 0, 1.1)."	FString
+		//EndPt += Chaos::TVector<T, 3>(0, 0, Radius * .5);
+		//RunTestFindClosestIntersectionHelper("RunTestFindClosestIntersection_Mid2", Center, Radius, Thickness,
+		//	StartPt, EndPt, 
+		//	IntersectionResult(Center+Chaos::TVector<T,3>(0,0,Radius+Thickness), true));
+		
+		// Fails w/no collision:
+		// "RunTestFindClosestIntersection_Mid3 - Sphere((0, 0, 0), 1).FindClosestIntersection(StartPt(0, 0, 0.5), EndPt(0, 0, 1), Thickness:0.1) = RESULT: 0, Point: (0, 0, 0), EXPECTED: 1 (0, 0, 1.1)."	FString
+		EndPt += Chaos::TVector<T, 3>(0, 0, Radius * .5);
+		//RunTestFindClosestIntersectionHelper("RunTestFindClosestIntersection_Mid3", Center, Radius, Thickness,
+		//	StartPt, EndPt, 
+		//	IntersectionResult(Center+Chaos::TVector<T,3>(0,0,Radius+Thickness), true));
+
+		// Fails w/no collision:
+		// "RunTestFindClosestIntersection_Mid4 - Sphere((0, 0, 0), 1).FindClosestIntersection(StartPt(0, 0, 1), EndPt(0, 0, 0.5), Thickness:0.1) = RESULT: 0, Point: (0, 0, 0), EXPECTED: 1 (0, 0, 1.1)."	FString
+		EndPt = StartPt;
+		StartPt = EndPt + Chaos::TVector<T, 3>(0, 0, Radius * .5);
+		//RunTestFindClosestIntersectionHelper("RunTestFindClosestIntersection_Mid4", Center, Radius, Thickness,
+		//	StartPt, EndPt, 
+		//	IntersectionResult(Center+Chaos::TVector<T,3>(0,0,Radius+Thickness), true));
+
+		StartPt += Chaos::TVector<T, 3>(0, 0, Radius * .5);
+		RunTestFindClosestIntersectionHelper("RunTestFindClosestIntersection_Mid5", Center, Radius, Thickness,
+			StartPt, EndPt, 
+			IntersectionResult(Center+Chaos::TVector<T,3>(0,0,Radius+Thickness), true));
+	}
+
+	template <class T>
+	void RunTestFindClosestIntersection_Deep(
+		const Chaos::TVector<T, 3>& Center, 
+		const T Radius,
+		const T Thickness)
+	{
+		Chaos::TVector<T, 3> StartPt = Center;
+		Chaos::TVector<T, 3> EndPt = Center;
+
+		// Fails w/no collision
+		// "RunTestFindClosestIntersection_Deep1 - Sphere((0, 0, 0), 1).FindClosestIntersection(StartPt(0, 0, 0), EndPt(0, 0, 0), Thickness:0.1) = RESULT: 0, Point: (0, 0, 0), EXPECTED: 1 (1.1, 0, 0)."	FString
+		//RunTestFindClosestIntersectionHelper("RunTestFindClosestIntersection_Deep1", Center, Radius, Thickness,
+		//	StartPt, EndPt, 
+		//	IntersectionResult(Center+Chaos::TVector<T,3>(Radius+Thickness,0,0), true));
+
+		// Fails w/no collision
+		// "RunTestFindClosestIntersection_Deep2 - Sphere((0, 0, 0), 1).FindClosestIntersection(StartPt(0, 0, 0), EndPt(0, 0, 0.5), Thickness:0.1) = RESULT: 0, Point: (0, 0, 0), EXPECTED: 1 (0, 0, 1.1)."	FString
+		EndPt += Chaos::TVector<T, 3>(0, 0, Radius*.5);
+		//RunTestFindClosestIntersectionHelper("RunTestFindClosestIntersection_Deep2", Center, Radius, Thickness,
+		//	StartPt, EndPt, 
+		//	IntersectionResult(Center+Chaos::TVector<T, 3>(0, 0, Radius+Thickness), true));
+
+		EndPt += Chaos::TVector<T, 3>(0, 0, Radius);
+		RunTestFindClosestIntersectionHelper("RunTestFindClosestIntersection_Deep3", Center, Radius, Thickness,
+			StartPt, EndPt, 
+			IntersectionResult(Center+Chaos::TVector<T, 3>(0, 0, Radius+Thickness), true));
+
+		// Fails w/no collision:
+		// "RunTestFindClosestIntersection_Deep4 - Sphere((0, 0, 0), 1).FindClosestIntersection(StartPt(0, 0, 0.5), EndPt(0, 0, 0), Thickness:0.1) = RESULT: 0, Point: (0, 0, 0), EXPECTED: 1 (0, 0, 1.1)."	FString
+		EndPt = Center;
+		StartPt += Chaos::TVector<T, 3>(0, 0, Radius*.5);
+		//RunTestFindClosestIntersectionHelper("RunTestFindClosestIntersection_Deep4", Center, Radius, Thickness,
+		//	StartPt, EndPt, 
+		//	IntersectionResult(Center+Chaos::TVector<T, 3>(0, 0, Radius+Thickness), true));
+
+		StartPt += Chaos::TVector<T, 3>(0, 0, Radius);
+		RunTestFindClosestIntersectionHelper("RunTestFindClosestIntersection_Deep5", Center, Radius, Thickness,
+			StartPt, EndPt, 
+			IntersectionResult(Center+Chaos::TVector<T, 3>(0, 0, Radius+Thickness), true));
+
+		// Miss
+		StartPt = Center + Chaos::TVector<T, 3>(Radius + Thickness + 1);
+		EndPt = StartPt + Chaos::TVector<T, 3>(Radius + Thickness + 1);
+		RunTestFindClosestIntersectionHelper("RunTestFindClosestIntersection_Deep6", Center, Radius, Thickness,
+			StartPt, EndPt, 
+			IntersectionResult(StartPt, false));
+	}
+
+	template <class T>
+	void RunTestFindClosestIntersection(
+		const Chaos::TVector<T,3>& Center, const T Radius)
+	{
+		//
+		// Run different test configurations with varying thicknesses.
+		//
+
+		// Deep intersection
+		RunTestFindClosestIntersection_Deep<T>(Center, Radius, 0.0);
+		RunTestFindClosestIntersection_Deep<T>(Center, Radius, Radius*.1);
+//		RunTestFindClosestIntersection_Deep<T>(Center, Radius, Radius*.5); // fails with incorrect collision points
+//		RunTestFindClosestIntersection_Deep<T>(Center, Radius, Radius);
+//		RunTestFindClosestIntersection_Deep<T>(Center, Radius, Radius*2);
+
+		// Mid intersection
+		RunTestFindClosestIntersection_Mid<T>(Center, Radius, 0.0);
+		RunTestFindClosestIntersection_Mid<T>(Center, Radius, Radius*.1);
+//		RunTestFindClosestIntersection_Mid<T>(Center, Radius, Radius*.5); // fails with incorrect collision points
+//		RunTestFindClosestIntersection_Mid<T>(Center, Radius, Radius);
+//		RunTestFindClosestIntersection_Mid<T>(Center, Radius, Radius*2);
+
+		// Grazing intersection - all fail
+		RunTestFindClosestIntersection_Shallow<T>(Center, Radius, 0.0);
+//		RunTestFindClosestIntersection_Shallow<T>(Center, Radius, Radius*.1); // fails with incorrect collision points
+//		RunTestFindClosestIntersection_Shallow<T>(Center, Radius, Radius*.5);
+//		RunTestFindClosestIntersection_Shallow<T>(Center, Radius, Radius);
+//		RunTestFindClosestIntersection_Shallow<T>(Center, Radius, Radius*2);
+
+		// Tangent intersection
+		RunTestFindClosestIntersection_Tangent<T>(Center, Radius, 0.0);
+		RunTestFindClosestIntersection_Tangent<T>(Center, Radius, Radius*.1);
+		RunTestFindClosestIntersection_Tangent<T>(Center, Radius, Radius*.5);
+		RunTestFindClosestIntersection_Tangent<T>(Center, Radius, Radius);
+		RunTestFindClosestIntersection_Tangent<T>(Center, Radius, Radius*2);
+	}
+	template <class T>
+	void TestIntersections()
+	{
+		// At the origin
+		RunTestFindClosestIntersection(Chaos::TVector<T, 3>((T)0.0), (T)1.0);
+//		RunTestFindClosestIntersection(Chaos::TVector<T, 3>((T)0.0), (T)10.0);
+//		RunTestFindClosestIntersection(Chaos::TVector<T, 3>((T)0.0), (T)100.0);
+
+		// Off origin
+//		RunTestFindClosestIntersection(Chaos::TVector<T, 3>((T)1.0), (T)1.0);
+//		RunTestFindClosestIntersection(Chaos::TVector<T, 3>((T)10.0), (T)1.0);
+//		RunTestFindClosestIntersection(Chaos::TVector<T, 3>((T)100.0), (T)1.0);
+
+//		RunTestFindClosestIntersection(Chaos::TVector<T, 3>((T)1.0), (T)10.0);
+//		RunTestFindClosestIntersection(Chaos::TVector<T, 3>((T)10.0), (T)10.0);
+//		RunTestFindClosestIntersection(Chaos::TVector<T, 3>((T)100.0), (T)10.0);
+
+//		RunTestFindClosestIntersection(Chaos::TVector<T, 3>((T)1.0), (T)100.0);
+//		RunTestFindClosestIntersection(Chaos::TVector<T, 3>((T)10.0), (T)100.0);
+//		RunTestFindClosestIntersection(Chaos::TVector<T, 3>((T)100.0), (T)100.0);
+	}
+
+	//==========================================================================
+	// Sample points tests
+	//==========================================================================
+
+	template <class T, int d>
+	void RunTestComputeSamplePoints(const Chaos::TSphere<T, d> &Sphere)
+	{
+		EXPECT_EQ(Sphere.GetType(), Chaos::ImplicitObjectType::Sphere) << *FString("Implicit object type is not 'sphere'.");
+
+		Chaos::TVector<T, d> Point = Sphere.GetCenter();
+		T Phi = Sphere.SignedDistance(Point);
+		EXPECT_LE(FMath::Abs(Phi + Sphere.GetRadius()), SMALL_NUMBER) << *FString("Sphere failed phi depth sanity test.");
+		
+		Point[0] += Sphere.GetRadius();
+		Phi = Sphere.SignedDistance(Point);
+		EXPECT_LE(FMath::Abs(Phi), KINDA_SMALL_NUMBER) << *FString("Sphere failed phi surface sanity test.");
+
+		TArray<Chaos::TVector<T, d>> Points = Sphere.ComputeSamplePoints(100);
+		check(Points.Num() == 100);
+		Point[0] = TNumericLimits<T>::Max();
+		for (const Chaos::TVector<T, d> &Pt : Points)
+		{
+			Phi = Sphere.SignedDistance(Pt);
+			const bool OnSurface = FMath::Abs(Phi) <= KINDA_SMALL_NUMBER;
+			check(OnSurface);
+			EXPECT_TRUE(OnSurface) << *FString("Produced a point not on the surface of the sphere.");
+			const bool Differs = Pt != Point;
+			check(Differs);
+			EXPECT_TRUE(Differs) << *FString("Produced a redundant value.");
+			Point = Pt;
+		}
+	}
+
+	template <class T>
+	void RunTestComputeSemispherePoints(const Chaos::TSphere<T, 3> &Sphere)
+	{
+		TArray<Chaos::TVector<T, 3>> Points;
+		Chaos::TSphereSpecializeSamplingHelper<T, 3>::ComputeBottomHalfSemiSphere(Points, Sphere, 100);
+		for (Chaos::TVector<T, 3>& Pt : Points)
+		{
+			T Phi = Sphere.SignedDistance(Pt);
+			const bool OnSurface = FMath::Abs(Phi) <= KINDA_SMALL_NUMBER;
+			check(OnSurface);
+			EXPECT_TRUE(OnSurface) << *FString("Produced a point not on the surface of the sphere.");
+
+			const bool BelowCenter = Pt[2] < Sphere.GetCenter()[2] + KINDA_SMALL_NUMBER;
+			check(BelowCenter);
+			EXPECT_TRUE(BelowCenter) << *FString("Bottom semisphere produced a point above midline.");
+		}
+
+		Points.Reset();
+		Chaos::TSphereSpecializeSamplingHelper<T, 3>::ComputeTopHalfSemiSphere(Points, Sphere, 100);
+		for (Chaos::TVector<T, 3>& Pt : Points)
+		{
+			T Phi = Sphere.SignedDistance(Pt);
+			const bool OnSurface = FMath::Abs(Phi) <= KINDA_SMALL_NUMBER;
+			check(OnSurface);
+			EXPECT_TRUE(OnSurface) << *FString("Produced a point not on the surface of the sphere.");
+
+			const bool BelowCenter = Pt[2] > Sphere.GetCenter()[2] - KINDA_SMALL_NUMBER;
+			check(BelowCenter);
+			EXPECT_TRUE(BelowCenter) << *FString("Top semisphere produced a point above midline.");
+		}
+	}
+
+	template <class T>
+	void TestComputeSamplePoints_SemiSphere()
+	{
+		// At the origin with radius 1
+		{
+			Chaos::TSphere<T, 3> Sphere(Chaos::TVector<T, 3>((T)0.0), (T)1.0);
+			RunTestComputeSemispherePoints(Sphere);
+		}	
+	}
+
+	template <class T, int d>
+	void TestComputeSamplePoints_Sphere()
+	{
+		// At the origin with radius 1
+		{
+			Chaos::TSphere<T, d> Sphere(Chaos::TVector<T, d>((T)0.0), (T)1.0);
+			RunTestComputeSamplePoints(Sphere);
+		}
+		// At the origin with radius > 1
+		{
+			Chaos::TSphere<T, d> Sphere(Chaos::TVector<T, d>(0.0), (T)10.0);
+			RunTestComputeSamplePoints(Sphere);
+		}
+		// At the origin with radius < 1
+		{
+			Chaos::TSphere<T, d> Sphere(Chaos::TVector<T, d>(0.0), (T).1);
+			RunTestComputeSamplePoints(Sphere);
+		}
+		// Off the origin with radius 1
+		{
+			Chaos::TSphere<T, d> Sphere(Chaos::TVector<T, d>(10.0), (T)1.0);
+			RunTestComputeSamplePoints(Sphere);
+		}
+		// Off the origin with radius > 1
+		{
+			Chaos::TSphere<T, d> Sphere(Chaos::TVector<T, d>(10.0), (T)10.0);
+			RunTestComputeSamplePoints(Sphere);
+		}
+		// Off the origin with radius < 1
+		{
+			Chaos::TSphere<T, d> Sphere(Chaos::TVector<T, d>(10.0), (T).1);
+			RunTestComputeSamplePoints(Sphere);
+		}
+	}
+
+	template<class T>
+	void TestImplicitSphere()
+	{
+		//TestComputeSamplePoints<T, 2>(R);
+		TestComputeSamplePoints_Sphere<T, 3>();
+		TestComputeSamplePoints_SemiSphere<T>();
+		TestIntersections<T>();
+	}
+	template void TestImplicitSphere<float>();
+} // namespace GeometryCollectionTest
diff --git a/Engine/Source/Programs/HeadlessChaos/Private/GeometryCollection/GeometryCollectionTestInitilization.cpp b/Engine/Source/Programs/HeadlessChaos/Private/GeometryCollection/GeometryCollectionTestInitilization.cpp
new file mode 100644
--- /dev/null
+++ b/Engine/Source/Programs/HeadlessChaos/Private/GeometryCollection/GeometryCollectionTestInitilization.cpp
@@ -0,0 +1,161 @@
+// Copyright Epic Games, Inc. All Rights Reserved.
+
+#include "GeometryCollection/GeometryCollectionTestInitilization.h"
+#include "GeometryCollection/GeometryCollectionTestFramework.h"
+#include "GeometryCollection/GeometryCollectionAlgo.h"
+#include "HeadlessChaosTestUtility.h"
+
+namespace GeometryCollectionTest
+{
+using namespace ChaosTest;
+
+
+TYPED_TEST(AllTraits,GeometryCollection_Initilization_TransformedGeometryCollectionRoot)
+{
+	using Traits = TypeParam;
+	FVector GlobalTranslation(10); FQuat GlobalRotation = FQuat::MakeFromEuler(FVector(0));
+	CreationParameters Params; Params.DynamicState = EObjectStateTypeEnum::Chaos_Object_Dynamic; Params.EnableClustering = false;
+	Params.RootTransform = FTransform(GlobalRotation,GlobalTranslation); Params.NestedTransforms ={FTransform(FVector(10)),FTransform::Identity,FTransform::Identity};
+	TGeometryCollectionWrapper<Traits>* Collection = TNewSimulationObject<GeometryType::GeometryCollectionWithSingleRigid>::Init<Traits>(Params)->As<TGeometryCollectionWrapper<Traits>>();
+	EXPECT_EQ(Collection->DynamicCollection->Parent[0],1); // is a child of index one
+
+	TFramework<Traits> UnitTest;
+	UnitTest.AddSimulationObject(Collection);
+	UnitTest.Initialize();
+	UnitTest.Advance();
+
+	{ // test results
+		EXPECT_EQ(UnitTest.Solver->GetParticles().GetGeometryCollectionParticles().Size(),1);
+		FVector X = UnitTest.Solver->GetParticles().GetGeometryCollectionParticles().X(0);
+		FQuat R = UnitTest.Solver->GetParticles().GetGeometryCollectionParticles().R(0);
+		EXPECT_TRUE((R * GlobalRotation.Inverse()).IsIdentity(KINDA_SMALL_NUMBER));
+		EXPECT_NEAR(X.X - GlobalTranslation[0],0.0f,KINDA_SMALL_NUMBER);
+		EXPECT_NEAR(X.Y - GlobalTranslation[1],0.0f,KINDA_SMALL_NUMBER);
+		EXPECT_LT(X.Z,GlobalTranslation[2]);
+
+		TArray<FTransform> Transform;
+		GeometryCollectionAlgo::GlobalMatrices(Collection->DynamicCollection->Transform,Collection->DynamicCollection->Parent,Transform);
+		EXPECT_EQ(Collection->DynamicCollection->Parent[0],FGeometryCollection::Invalid); // is not a child
+		EXPECT_NEAR(Transform[0].GetTranslation().X - GlobalTranslation[0],0.0f,KINDA_SMALL_NUMBER);
+		EXPECT_NEAR(Transform[0].GetTranslation().Y - GlobalTranslation[1],0.0f,KINDA_SMALL_NUMBER);
+		EXPECT_LT(Transform[0].GetTranslation().Z,GlobalTranslation[2]);
+	}
+}
+
+
+TYPED_TEST(AllTraits,GeometryCollection_Initilization_TransformedGeometryCollectionParentNode)
+{
+	using Traits = TypeParam;
+	FVector GlobalTranslation(10); FQuat GlobalRotation = FQuat::MakeFromEuler(FVector(0));
+	CreationParameters Params; Params.DynamicState = EObjectStateTypeEnum::Chaos_Object_Dynamic; Params.EnableClustering = false;
+	Params.RootTransform = FTransform(GlobalRotation,GlobalTranslation); Params.NestedTransforms ={FTransform::Identity,FTransform(FVector(10)),FTransform::Identity};
+	TGeometryCollectionWrapper<Traits>* Collection = TNewSimulationObject<GeometryType::GeometryCollectionWithSingleRigid>::Init<Traits>(Params)->As<TGeometryCollectionWrapper<Traits>>();
+
+	TFramework<Traits> UnitTest;
+	UnitTest.AddSimulationObject(Collection);
+	UnitTest.Initialize();
+	UnitTest.Advance();
+
+	{ // test results
+		EXPECT_EQ(UnitTest.Solver->GetParticles().GetGeometryCollectionParticles().Size(),1);
+		FVector X = UnitTest.Solver->GetParticles().GetGeometryCollectionParticles().X(0);
+		FQuat R = UnitTest.Solver->GetParticles().GetGeometryCollectionParticles().R(0);
+		EXPECT_TRUE((R * GlobalRotation.Inverse()).IsIdentity(KINDA_SMALL_NUMBER));
+		EXPECT_NEAR(X.X - GlobalTranslation[0],0.0f,KINDA_SMALL_NUMBER);
+		EXPECT_NEAR(X.Y - GlobalTranslation[1],0.0f,KINDA_SMALL_NUMBER);
+		EXPECT_LT(X.Z,GlobalTranslation[2]);
+
+		TArray<FTransform> Transform;
+		GeometryCollectionAlgo::GlobalMatrices(Collection->DynamicCollection->Transform,Collection->DynamicCollection->Parent,Transform);
+		EXPECT_EQ(Collection->DynamicCollection->Parent[0],FGeometryCollection::Invalid); // is not a child
+		EXPECT_NEAR(Transform[0].GetTranslation().X - GlobalTranslation[0],0.0f,KINDA_SMALL_NUMBER);
+		EXPECT_NEAR(Transform[0].GetTranslation().Y - GlobalTranslation[1],0.0f,KINDA_SMALL_NUMBER);
+		EXPECT_LT(Transform[0].GetTranslation().Z,GlobalTranslation[2]);
+	}
+}
+
+
+TYPED_TEST(AllTraits,GeometryCollection_Initilization_TransformedGeometryCollectionGeometryNode)
+{
+	using Traits = TypeParam;
+	FVector GlobalTranslation(10); FQuat GlobalRotation = FQuat::MakeFromEuler(FVector(0));
+	CreationParameters Params; Params.DynamicState = EObjectStateTypeEnum::Chaos_Object_Dynamic; Params.EnableClustering = false;
+	Params.RootTransform = FTransform(GlobalRotation,GlobalTranslation); Params.NestedTransforms ={FTransform::Identity,FTransform::Identity,FTransform(FVector(10))};
+	TGeometryCollectionWrapper<Traits>* Collection = TNewSimulationObject<GeometryType::GeometryCollectionWithSingleRigid>::Init<Traits>(Params)->As<TGeometryCollectionWrapper<Traits>>();
+
+	TFramework<Traits> UnitTest;
+	UnitTest.AddSimulationObject(Collection);
+	UnitTest.Initialize();
+	UnitTest.Advance();
+
+	{ // test results
+		EXPECT_EQ(UnitTest.Solver->GetParticles().GetGeometryCollectionParticles().Size(),1);
+		FVector X = UnitTest.Solver->GetParticles().GetGeometryCollectionParticles().X(0);
+		FQuat R = UnitTest.Solver->GetParticles().GetGeometryCollectionParticles().R(0);
+		EXPECT_TRUE((R * GlobalRotation.Inverse()).IsIdentity(KINDA_SMALL_NUMBER));
+		EXPECT_NEAR(X.X - GlobalTranslation[0],0.0f,KINDA_SMALL_NUMBER);
+		EXPECT_NEAR(X.Y - GlobalTranslation[1],0.0f,KINDA_SMALL_NUMBER);
+		EXPECT_LT(X.Z,GlobalTranslation[2]);
+
+		TArray<FTransform> Transform;
+		GeometryCollectionAlgo::GlobalMatrices(Collection->DynamicCollection->Transform,Collection->DynamicCollection->Parent,Transform);
+		EXPECT_EQ(Collection->DynamicCollection->Parent[0],FGeometryCollection::Invalid); // is not a child
+		EXPECT_NEAR(Transform[0].GetTranslation().X - GlobalTranslation[0],0.0f,KINDA_SMALL_NUMBER);
+		EXPECT_NEAR(Transform[0].GetTranslation().Y - GlobalTranslation[1],0.0f,KINDA_SMALL_NUMBER);
+		EXPECT_LT(Transform[0].GetTranslation().Z,GlobalTranslation[2]);
+	}
+}
+
+
+TYPED_TEST(AllTraits,GeometryCollection_Initilization_TransformedGeometryCollectionGeometryVertices)
+{
+	using Traits = TypeParam;
+	FVector GlobalTranslation(10); FQuat GlobalRotation = FQuat::MakeFromEuler(FVector(0));
+	CreationParameters Params; Params.DynamicState = EObjectStateTypeEnum::Chaos_Object_Dynamic; Params.EnableClustering = false;
+	Params.SimplicialType = ESimplicialType::Chaos_Simplicial_Sphere; Params.ImplicitType = EImplicitTypeEnum::Chaos_Implicit_Sphere;
+	Params.GeomTransform = FTransform(GlobalRotation,GlobalTranslation); Params.NestedTransforms ={FTransform::Identity,FTransform::Identity,FTransform::Identity};
+	TGeometryCollectionWrapper<Traits>* Collection = TNewSimulationObject<GeometryType::GeometryCollectionWithSingleRigid>::Init<Traits>(Params)->As<TGeometryCollectionWrapper<Traits>>();
+
+	//
+	// validate the vertices have been moved.  
+	//
+	TArray<FTransform> RestTransforms;
+	int32 NumVertices = Collection->DynamicCollection->NumElements(FGeometryCollection::VerticesGroup);
+	TManagedArray<FVector>& Vertices = Collection->RestCollection->GetAttribute<FVector>("Vertex",FGeometryCollection::VerticesGroup);
+	TManagedArray<int32>& BoneMap = Collection->RestCollection->GetAttribute<int32>("BoneMap",FGeometryCollection::VerticesGroup);
+	GeometryCollectionAlgo::GlobalMatrices(Collection->RestCollection->Transform,Collection->RestCollection->Parent,RestTransforms);
+	FVector CenterOfMass(0);  for(int VertexIndex =0; VertexIndex <NumVertices; VertexIndex++) CenterOfMass += RestTransforms[BoneMap[VertexIndex]].TransformPosition(Vertices[VertexIndex]); CenterOfMass /= NumVertices;
+	EXPECT_NEAR(CenterOfMass.X - GlobalTranslation[0],0.0f,KINDA_SMALL_NUMBER);
+	EXPECT_NEAR(CenterOfMass.Y - GlobalTranslation[0],0.0f,KINDA_SMALL_NUMBER);
+	EXPECT_NEAR(CenterOfMass.Z - GlobalTranslation[0],0.0f,KINDA_SMALL_NUMBER);
+
+
+
+	TFramework<Traits> UnitTest;
+	UnitTest.AddSimulationObject(Collection);
+	UnitTest.Initialize();
+	UnitTest.Advance();
+
+	{ // test results
+		TManagedArray<FTransform>& MassToLocal = Collection->RestCollection->GetAttribute<FTransform>("MassToLocal",FGeometryCollection::TransformGroup);
+
+
+		//EXPECT_EQ(UnitTest.Solver->GetParticles().GetGeometryCollectionParticles().Size(), 1);
+		//FVector X = UnitTest.Solver->GetParticles().GetGeometryCollectionParticles().X(0);
+		//FQuat R = UnitTest.Solver->GetParticles().GetGeometryCollectionParticles().R(0);
+		//EXPECT_TRUE((R * GlobalRotation.Inverse()).IsIdentity(KINDA_SMALL_NUMBER));
+		//EXPECT_NEAR(X.X - GlobalTranslation[0], 0.0f, KINDA_SMALL_NUMBER);
+		//EXPECT_NEAR(X.Y - GlobalTranslation[1], 0.0f, KINDA_SMALL_NUMBER);
+		//EXPECT_LT(X.Z, GlobalTranslation[2]);
+
+		//TArray<FTransform> Transform;
+		//GeometryCollectionAlgo::GlobalMatrices(Collection->DynamicCollection->Transform, Collection->DynamicCollection->Parent, Transform);
+		//EXPECT_EQ(Collection->DynamicCollection->Parent[0], FGeometryCollection::Invalid); // is not a child
+		//EXPECT_NEAR(Transform[0].GetTranslation().X - GlobalTranslation[0], 0.0f, KINDA_SMALL_NUMBER);
+		//EXPECT_NEAR(Transform[0].GetTranslation().Y - GlobalTranslation[1], 0.0f, KINDA_SMALL_NUMBER);
+		//EXPECT_LT(Transform[0].GetTranslation().Z, GlobalTranslation[2]);
+	}
+}
+
+
+} // namespace GeometryCollectionTest
diff --git a/Engine/Source/Programs/HeadlessChaos/Private/GeometryCollection/GeometryCollectionTestMassProperties.cpp b/Engine/Source/Programs/HeadlessChaos/Private/GeometryCollection/GeometryCollectionTestMassProperties.cpp
new file mode 100644
--- /dev/null
+++ b/Engine/Source/Programs/HeadlessChaos/Private/GeometryCollection/GeometryCollectionTestMassProperties.cpp
@@ -0,0 +1,429 @@
+// Copyright Epic Games, Inc. All Rights Reserved.
+
+#include "GeometryCollection/GeometryCollectionTestMassProperties.h"
+
+#include "Chaos/TriangleMesh.h"
+#include "Chaos/MassProperties.h"
+#include "Math/Box.h"
+#include "GeometryCollection/GeometryCollection.h"
+#include "GeometryCollection/GeometryCollectionAlgo.h"
+#include "GeometryCollection/GeometryCollectionUtility.h"
+#include "GeometryCollection/GeometryCollectionTestFramework.h"
+#include "GeometryCollection/GeometryCollectionTestUtility.h"
+#include "GeometryCollection/TransformCollection.h"
+#include "HeadlessChaosTestUtility.h"
+
+namespace GeometryCollectionTest
+{
+	using namespace ChaosTest;
+
+	TYPED_TEST(AllTraits, GeometryCollection_MassProperties_Compute)
+	{
+		using Traits = TypeParam;
+		using namespace Chaos;
+
+		Chaos::TParticles<float, 3> Vertices;
+		Vertices.AddParticles(8);
+		Vertices.X(0) = TVector<float, 3>(-1, 1, -1);
+		Vertices.X(1) = TVector<float, 3>(1, 1, -1);
+		Vertices.X(2) = TVector<float, 3>(1, -1, -1);
+		Vertices.X(3) = TVector<float, 3>(-1, -1, -1);
+		Vertices.X(4) = TVector<float, 3>(-1, 1, 1);
+		Vertices.X(5) = TVector<float, 3>(1, 1, 1);
+		Vertices.X(6) = TVector<float, 3>(1, -1, 1);
+		Vertices.X(7) = TVector<float, 3>(-1, -1, 1);
+
+		// @todo(chaos):  breaking : this trips an ensure in the test, why?
+		for (int i = 0; i < 8; i++) {
+			Vertices.X(i) *= FVector(1, 2, 3);
+			Vertices.X(i) += FVector(1, 2, 3);
+		}
+
+		TArray<Chaos::TVector<int32, 3>> Faces;
+		Faces.SetNum(12);
+		Faces[0] = TVector<int32, 3>(0,1,2);
+		Faces[1] = TVector<int32, 3>(0,2,3);
+		Faces[2] = TVector<int32, 3>(2,1,6);
+		Faces[3] = TVector<int32, 3>(1,5,6);
+		Faces[4] = TVector<int32, 3>(2,6,7);
+		Faces[5] = TVector<int32, 3>(3,2,7);
+		Faces[6] = TVector<int32, 3>(4,7,3);
+		Faces[7] = TVector<int32, 3>(4,0,3);
+		Faces[8] = TVector<int32, 3>(4,1,0);
+		Faces[9] = TVector<int32, 3>(4,5,1);
+		Faces[10] = TVector<int32, 3>(5,4,7);
+		Faces[11] = TVector<int32, 3>(5,7,6);
+		Chaos::TTriangleMesh<float> Surface(MoveTemp(Faces));
+
+		TMassProperties<FReal, 3> MassProperties;
+		MassProperties.Mass = 1.f;
+		//Chaos::TMassProperties<float, 3> MassProperties = Chaos::CalculateMassProperties(Vertices, Surface.GetElements(), 1.f);
+		{
+			const auto& SurfaceElements = Surface.GetElements();
+			CalculateVolumeAndCenterOfMass(Vertices, SurfaceElements, MassProperties.Volume, MassProperties.CenterOfMass);
+
+			for (int32 Idx = 0; Idx < 8; ++Idx)
+			{
+				Vertices.X(Idx) -= MassProperties.CenterOfMass;
+			}
+
+			check(MassProperties.Mass > 0);
+			check(MassProperties.Volume > SMALL_NUMBER);
+			
+			// @todo(chaos) : breaking
+			CalculateInertiaAndRotationOfMass(Vertices, SurfaceElements, MassProperties.Mass / MassProperties.Volume, MassProperties.CenterOfMass, MassProperties.InertiaTensor, MassProperties.RotationOfMass);
+		}
+		EXPECT_EQ(MassProperties.Mass, 1.f);
+		EXPECT_TRUE((MassProperties.CenterOfMass-FVector(1,2,3)).Size() < SMALL_NUMBER);
+		
+		// This is just measured data to let us know when it changes. Ideally this would be derived. 
+		FVector EulerAngle = MassProperties.RotationOfMass.Euler();
+		EXPECT_TRUE((MassProperties.RotationOfMass.Euler()-FVector(115.8153,-12.4347,1.9705)).Size() > SMALL_NUMBER);
+		EXPECT_TRUE(MassProperties.InertiaTensor.M[0][0] - 14.9866095 < SMALL_NUMBER);
+		EXPECT_TRUE(MassProperties.InertiaTensor.M[1][1] - 1.40656376 < SMALL_NUMBER);
+		EXPECT_TRUE(MassProperties.InertiaTensor.M[2][2] - 13.7401619 < SMALL_NUMBER);
+
+
+	}
+
+	TYPED_TEST(AllTraits, DISABLED_GeometryCollection_MassProperties_Cube)
+	{
+		using Traits = TypeParam;
+		using namespace Chaos;
+		FVector GlobalTranslation(0); FQuat GlobalRotation = FQuat::MakeFromEuler(FVector(0));
+		CreationParameters Params; Params.SimplicialType = ESimplicialType::Chaos_Simplicial_Box; Params.ImplicitType = EImplicitTypeEnum::Chaos_Implicit_Box;
+		Params.GeomTransform = FTransform(GlobalRotation, GlobalTranslation); Params.NestedTransforms = { FTransform::Identity, FTransform::Identity,  FTransform::Identity };
+		TGeometryCollectionWrapper<Traits>* Collection = TNewSimulationObject<GeometryType::GeometryCollectionWithSingleRigid>::Init<Traits>(Params)->As<TGeometryCollectionWrapper<Traits>>();
+
+		TArray<FTransform> Transform;
+		GeometryCollectionAlgo::GlobalMatrices(Collection->RestCollection->Transform, Collection->RestCollection->Parent, Transform);
+
+		// group ?
+		const TManagedArray<bool>& Visible = Collection->RestCollection->Visible;
+
+		// VerticesGroup
+		const TManagedArray<FVector>& Vertex = Collection->RestCollection->Vertex;
+
+		// GeometryGroup
+		const int32 NumGeometries = Collection->RestCollection->NumElements(FGeometryCollection::GeometryGroup);
+		const TManagedArray<int32>& VertexCount = Collection->RestCollection->VertexCount;
+		const TManagedArray<int32>& VertexStart = Collection->RestCollection->VertexStart;
+		const TManagedArray<int32>& FaceCount = Collection->RestCollection->FaceCount;
+		const TManagedArray<int32>& FaceStart = Collection->RestCollection->FaceStart;
+		const TManagedArray<int32>& TransformIndex = Collection->RestCollection->TransformIndex;
+		const TManagedArray<FIntVector>& Indices = Collection->RestCollection->Indices;
+		const TManagedArray<int32>& BoneMap = Collection->RestCollection->BoneMap;
+		int GeometryIndex = 0;
+
+		TUniquePtr<TTriangleMesh<float>> TriMesh(
+			CreateTriangleMesh(
+				FaceStart[GeometryIndex],
+				FaceCount[GeometryIndex],
+				Visible,
+				Indices));
+
+		//TArray<Chaos::TVector<int32, 3>> Faces;
+		//Faces.SetNum(Indices.Num());
+		//for (int i = 0; i < Indices.Num(); i++) { Faces[i] = TVector<int32, 3>(Indices[i][0], Indices[i][1], Indices[i][2]); }
+		//Chaos::TTriangleMesh<float> TriMesh(MoveTemp(Faces));
+
+		TArray<TMassProperties<float, 3>> MassPropertiesArray;
+		MassPropertiesArray.AddUninitialized(NumGeometries);
+		TMassProperties<float, 3>& MassProperties = MassPropertiesArray[GeometryIndex];
+		MassProperties.CenterOfMass = FVector(0);
+		MassProperties.Mass = 1.f;
+
+		TParticles<float, 3> MassSpaceParticles;
+		MassSpaceParticles.AddParticles(Vertex.Num());
+		for (int32 Idx = 0; Idx < Vertex.Num(); ++Idx)
+		{
+			MassSpaceParticles.X(Idx) = Transform[BoneMap[Idx]].TransformPosition(Vertex[Idx]);
+		}
+
+		CalculateVolumeAndCenterOfMass(MassSpaceParticles, TriMesh->GetElements(), MassProperties.Volume, MassProperties.CenterOfMass);
+
+		EXPECT_NEAR(MassProperties.Volume - 8.0, 0.0f, KINDA_SMALL_NUMBER); 
+		EXPECT_NEAR(MassProperties.CenterOfMass.X - GlobalTranslation[0], 0.0f, KINDA_SMALL_NUMBER);
+		EXPECT_NEAR(MassProperties.CenterOfMass.Y - GlobalTranslation[1], 0.0f, KINDA_SMALL_NUMBER);
+		EXPECT_NEAR(MassProperties.CenterOfMass.Z - GlobalTranslation[2], 0.0f, KINDA_SMALL_NUMBER);
+
+		for (int32 Idx = 0; Idx < Vertex.Num(); ++Idx)
+		{
+			MassSpaceParticles.X(Idx) -= MassProperties.CenterOfMass;
+		}
+
+		float Density = 1.f;
+		TVector<float, 3> ZeroVec(0);
+		CalculateInertiaAndRotationOfMass(MassSpaceParticles, TriMesh->GetSurfaceElements(), Density, ZeroVec, MassProperties.InertiaTensor, MassProperties.RotationOfMass);
+
+		EXPECT_EQ(MassProperties.Mass, 1.f);
+		EXPECT_TRUE((MassProperties.CenterOfMass - GlobalTranslation).Size() < SMALL_NUMBER);
+
+		// This is just measured data to let us know when it changes. Ideally this would be derived. 
+		FVector EulerAngle = MassProperties.RotationOfMass.Euler();
+		EXPECT_TRUE((MassProperties.RotationOfMass.Euler() - FVector(115.8153, -12.4347, 1.9705)).Size() > SMALL_NUMBER);
+		EXPECT_TRUE(MassProperties.InertiaTensor.M[0][0] - 4.99521351 < SMALL_NUMBER);
+		EXPECT_TRUE(MassProperties.InertiaTensor.M[1][1] - 4.07145357 < SMALL_NUMBER);
+		EXPECT_TRUE(MassProperties.InertiaTensor.M[2][2] - 4.26666689 < SMALL_NUMBER);
+	}
+
+
+
+
+	TYPED_TEST(AllTraits, DISABLED_GeometryCollection_MassProperties_Sphere)
+	{
+		using Traits = TypeParam;
+		using namespace Chaos;
+		FVector GlobalTranslation(10); FQuat GlobalRotation = FQuat::MakeFromEuler(FVector(0)); FVector Scale(1);
+		CreationParameters Params; Params.SimplicialType = ESimplicialType::Chaos_Simplicial_Sphere; Params.ImplicitType = EImplicitTypeEnum::Chaos_Implicit_Sphere;
+		Params.GeomTransform = FTransform(GlobalRotation, GlobalTranslation, Scale); Params.NestedTransforms = { FTransform::Identity, FTransform::Identity,  FTransform::Identity };
+		TGeometryCollectionWrapper<Traits>* Collection = TNewSimulationObject<GeometryType::GeometryCollectionWithSingleRigid>::Init<Traits>(Params)->As<TGeometryCollectionWrapper<Traits>>();
+
+		TArray<FTransform> Transform;
+		GeometryCollectionAlgo::GlobalMatrices(Collection->RestCollection->Transform, Collection->RestCollection->Parent, Transform);
+
+		// group ?
+		const TManagedArray<bool>& Visible = Collection->RestCollection->Visible;
+
+		// VerticesGroup
+		const TManagedArray<FVector>& Vertex = Collection->RestCollection->Vertex;
+
+		// GeometryGroup
+		const int32 NumGeometries = Collection->RestCollection->NumElements(FGeometryCollection::GeometryGroup);
+		const TManagedArray<int32>& VertexCount = Collection->RestCollection->VertexCount;
+		const TManagedArray<int32>& VertexStart = Collection->RestCollection->VertexStart;
+		const TManagedArray<int32>& FaceCount = Collection->RestCollection->FaceCount;
+		const TManagedArray<int32>& FaceStart = Collection->RestCollection->FaceStart;
+		const TManagedArray<int32>& TransformIndex = Collection->RestCollection->TransformIndex;
+		const TManagedArray<FIntVector>& Indices = Collection->RestCollection->Indices;
+		const TManagedArray<int32>& BoneMap = Collection->RestCollection->BoneMap;
+		int GeometryIndex = 0;
+
+		TUniquePtr<TTriangleMesh<float>> TriMesh(
+			CreateTriangleMesh(
+				FaceStart[GeometryIndex],
+				FaceCount[GeometryIndex],
+				Visible,
+				Indices));
+
+
+		TArray<TMassProperties<float, 3>> MassPropertiesArray;
+		MassPropertiesArray.AddUninitialized(NumGeometries);
+		TMassProperties<float, 3>& MassProperties = MassPropertiesArray[GeometryIndex];
+
+		TParticles<float, 3> MassSpaceParticles;
+		MassSpaceParticles.AddParticles(Vertex.Num());
+		for (int32 Idx = 0; Idx < Vertex.Num(); ++Idx)
+		{
+			MassSpaceParticles.X(Idx) = Transform[BoneMap[Idx]].TransformPosition(Vertex[Idx]);
+		}
+
+		CalculateVolumeAndCenterOfMass(MassSpaceParticles, TriMesh->GetElements(), MassProperties.Volume, MassProperties.CenterOfMass);
+
+		// Since we're intersecting triangles with a sphere, where the vertices of the 
+		// triangle vertices are on the sphere surface, we're missing some volume.  Thus,
+		// we'd expect the volume of the triangulation to approach the analytic volume as
+		// the number of polygons goes to infinity (MakeSphereElement() currently does 
+		// 16x16 divisions in U and V).
+		const float AnalyticVolume = (4.0/3) * (22.0/7) * Scale[0] * Scale[0] * Scale[0];
+		EXPECT_NEAR(MassProperties.Volume - AnalyticVolume, 0.0f, 0.2); // this should be 4.19047642
+		EXPECT_NEAR(MassProperties.CenterOfMass.X - GlobalTranslation[0], 0.0f, KINDA_SMALL_NUMBER);
+		EXPECT_NEAR(MassProperties.CenterOfMass.Y - GlobalTranslation[1], 0.0f, KINDA_SMALL_NUMBER);
+		EXPECT_NEAR(MassProperties.CenterOfMass.Z - GlobalTranslation[2], 0.0f, KINDA_SMALL_NUMBER);
+
+
+		for (int32 Idx = 0; Idx < Vertex.Num(); ++Idx)
+		{
+			MassSpaceParticles.X(Idx) -= MassProperties.CenterOfMass;
+		}
+
+		float Density = 0.01f;
+		TVector<float, 3> ZeroVec(0);
+		CalculateInertiaAndRotationOfMass(MassSpaceParticles, TriMesh->GetSurfaceElements(), Density, ZeroVec, MassProperties.InertiaTensor, MassProperties.RotationOfMass);
+
+		// todo(chaos) : Check this. 
+		// This is just measured data to let us know when it changes. Ideally this would be derived. 
+		FVector EulerAngle = MassProperties.RotationOfMass.Euler();
+		//EXPECT_TRUE((MassProperties.RotationOfMass.Euler() - FVector(115.8153, -12.4347, 1.9705)).Size() > SMALL_NUMBER);
+		//EXPECT_TRUE(MassProperties.InertiaTensor.M[0][0] - 4.99521351 < SMALL_NUMBER);
+		//EXPECT_TRUE(MassProperties.InertiaTensor.M[1][1] - 4.07145357 < SMALL_NUMBER);
+		//EXPECT_TRUE(MassProperties.InertiaTensor.M[2][2] - 4.26666689 < SMALL_NUMBER);
+	}
+
+
+	TYPED_TEST(AllTraits, DISABLED_GeometryCollection_MassProperties_Tetrahedron)
+	{
+		using Traits = TypeParam;
+		using namespace Chaos;
+		FVector GlobalTranslation(0); FQuat GlobalRotation = FQuat::MakeFromEuler(FVector(0));
+		CreationParameters Params; Params.SimplicialType = ESimplicialType::Chaos_Simplicial_Tetrahedron; Params.ImplicitType = EImplicitTypeEnum::Chaos_Implicit_Sphere;
+		Params.GeomTransform = FTransform(GlobalRotation, GlobalTranslation); Params.NestedTransforms = { FTransform::Identity, FTransform::Identity,  FTransform::Identity };
+		TGeometryCollectionWrapper<Traits>* Collection = TNewSimulationObject<GeometryType::GeometryCollectionWithSingleRigid>::Init<Traits>(Params)->As<TGeometryCollectionWrapper<Traits>>();
+
+		TArray<FTransform> Transform;
+		GeometryCollectionAlgo::GlobalMatrices(Collection->RestCollection->Transform, Collection->RestCollection->Parent, Transform);
+
+		// group ?
+		const TManagedArray<bool>& Visible = Collection->RestCollection->Visible;
+
+		// VerticesGroup
+		const TManagedArray<FVector>& Vertex = Collection->RestCollection->Vertex;
+
+		// GeometryGroup
+		const int32 NumGeometries = Collection->RestCollection->NumElements(FGeometryCollection::GeometryGroup);
+		const TManagedArray<int32>& VertexCount = Collection->RestCollection->VertexCount;
+		const TManagedArray<int32>& VertexStart = Collection->RestCollection->VertexStart;
+		const TManagedArray<int32>& FaceCount = Collection->RestCollection->FaceCount;
+		const TManagedArray<int32>& FaceStart = Collection->RestCollection->FaceStart;
+		const TManagedArray<int32>& TransformIndex = Collection->RestCollection->TransformIndex;
+		const TManagedArray<FIntVector>& Indices = Collection->RestCollection->Indices;
+		const TManagedArray<int32>& BoneMap = Collection->RestCollection->BoneMap;
+		int GeometryIndex = 0;
+
+		TUniquePtr<TTriangleMesh<float>> TriMesh(
+			CreateTriangleMesh(
+				FaceStart[GeometryIndex],
+				FaceCount[GeometryIndex],
+				Visible,
+				Indices));
+
+
+		TArray<TMassProperties<float, 3>> MassPropertiesArray;
+		MassPropertiesArray.AddUninitialized(NumGeometries);
+		TMassProperties<float, 3>& MassProperties = MassPropertiesArray[GeometryIndex];
+		MassProperties.Mass = 1.0f;
+		MassProperties.CenterOfMass = FVector(0);
+
+		TParticles<float, 3> MassSpaceParticles;
+		MassSpaceParticles.AddParticles(Vertex.Num());
+		for (int32 Idx = 0; Idx < Vertex.Num(); ++Idx)
+		{
+			MassSpaceParticles.X(Idx) = Transform[BoneMap[Idx]].TransformPosition(Vertex[Idx]);
+		}
+
+		CalculateVolumeAndCenterOfMass(MassSpaceParticles, TriMesh->GetElements(), MassProperties.Volume, MassProperties.CenterOfMass);
+
+		EXPECT_NEAR(MassProperties.Volume - 2.666666, 0.0f, KINDA_SMALL_NUMBER); // Tetrahedron with edge lengths 2.8284
+		EXPECT_NEAR(MassProperties.CenterOfMass.X - GlobalTranslation[0], 0.0f, KINDA_SMALL_NUMBER);
+		EXPECT_NEAR(MassProperties.CenterOfMass.Y - GlobalTranslation[1], 0.0f, KINDA_SMALL_NUMBER);
+		EXPECT_NEAR(MassProperties.CenterOfMass.Z - GlobalTranslation[2], 0.0f, KINDA_SMALL_NUMBER);
+
+		for (int32 Idx = 0; Idx < Vertex.Num(); ++Idx)
+		{
+			MassSpaceParticles.X(Idx) -= MassProperties.CenterOfMass;
+		}
+
+		float Density = 0.01f;
+		TVector<float,3> ZeroVec(0);
+		CalculateInertiaAndRotationOfMass(MassSpaceParticles, TriMesh->GetSurfaceElements(), Density, ZeroVec, MassProperties.InertiaTensor, MassProperties.RotationOfMass);
+
+		// todo(chaos) : Check this. 
+		// This is just measured data to let us know when it changes. Ideally this would be derived. 
+		FVector EulerAngle = MassProperties.RotationOfMass.Euler();
+		//EXPECT_TRUE((MassProperties.RotationOfMass.Euler() - FVector(115.8153, -12.4347, 1.9705)).Size() > SMALL_NUMBER);
+		//EXPECT_TRUE(MassProperties.InertiaTensor.M[0][0] - 4.99521351 < SMALL_NUMBER);
+		//EXPECT_TRUE(MassProperties.InertiaTensor.M[1][1] - 4.07145357 < SMALL_NUMBER);
+		//EXPECT_TRUE(MassProperties.InertiaTensor.M[2][2] - 4.26666689 < SMALL_NUMBER);
+	}
+
+
+
+
+	TYPED_TEST(AllTraits, DISABLED_GeometryCollection_MassProperties_ScaledSphere)
+	{
+		// This test has points that are scaled, rotated and translated within mass space. 
+		// So the resulting surface is not about the center of mass and needs to be
+		// moved for simulation. 
+
+		using Traits = TypeParam;
+		using namespace Chaos;
+		FVector GlobalTranslation(10); FQuat GlobalRotation = FQuat::MakeFromEuler(FVector(45,0,0));
+		CreationParameters Params; Params.SimplicialType = ESimplicialType::Chaos_Simplicial_Sphere; Params.ImplicitType = EImplicitTypeEnum::Chaos_Implicit_Sphere;
+		Params.GeomTransform = FTransform(GlobalRotation,GlobalTranslation, FVector(1, 5, 11)); Params.NestedTransforms = { FTransform::Identity, FTransform::Identity,  FTransform::Identity };
+		TGeometryCollectionWrapper<Traits>* Collection = TNewSimulationObject<GeometryType::GeometryCollectionWithSingleRigid>::Init<Traits>(Params)->As<TGeometryCollectionWrapper<Traits>>();
+
+		TArray<FTransform> Transform;
+		GeometryCollectionAlgo::GlobalMatrices(Collection->RestCollection->Transform, Collection->RestCollection->Parent, Transform);
+
+		// group ?
+		const TManagedArray<bool>& Visible = Collection->RestCollection->Visible;
+
+		// VerticesGroup
+		TManagedArray<FVector>& Vertex = Collection->RestCollection->Vertex;
+
+		// GeometryGroup
+		const int32 NumGeometries = Collection->RestCollection->NumElements(FGeometryCollection::GeometryGroup);
+		const TManagedArray<int32>& VertexCount = Collection->RestCollection->VertexCount;
+		const TManagedArray<int32>& VertexStart = Collection->RestCollection->VertexStart;
+		const TManagedArray<int32>& FaceCount = Collection->RestCollection->FaceCount;
+		const TManagedArray<int32>& FaceStart = Collection->RestCollection->FaceStart;
+		const TManagedArray<int32>& TransformIndex = Collection->RestCollection->TransformIndex;
+		const TManagedArray<FIntVector>& Indices = Collection->RestCollection->Indices;
+		const TManagedArray<int32>& BoneMap = Collection->RestCollection->BoneMap;
+		int GeometryIndex = 0;
+
+
+		TUniquePtr<TTriangleMesh<float>> TriMesh(
+			CreateTriangleMesh(
+				FaceStart[GeometryIndex],
+				FaceCount[GeometryIndex],
+				Visible,
+				Indices));
+
+
+		TArray<TMassProperties<float, 3>> MassPropertiesArray;
+		MassPropertiesArray.AddUninitialized(NumGeometries);
+		TMassProperties<float, 3>& MassProperties = MassPropertiesArray[GeometryIndex];
+
+		TArray<FVector> SomeVec;
+		TParticles<float, 3> MassSpaceParticles;
+		MassSpaceParticles.AddParticles(Vertex.Num());
+		for (int32 Idx = 0; Idx < Vertex.Num(); ++Idx)
+		{
+			FVector VertexPoint = Vertex[Idx];
+			MassSpaceParticles.X(Idx) = Transform[BoneMap[Idx]].TransformPosition(Vertex[Idx]);
+			FVector MassSpacePoint = FVector(MassSpaceParticles.X(Idx)[0], MassSpaceParticles.X(Idx)[1], MassSpaceParticles.X(Idx)[2]);
+			SomeVec.Add(MassSpacePoint);
+			//UE_LOG(LogChaos, Log, TEXT("...(%3.5f,%3.5f,%3.5f)"), MassSpacePoint[0], MassSpacePoint[1], MassSpacePoint[2]);
+			UE_LOG(LogChaos, Log, TEXT("...(%3.5f,%3.5f,%3.5f)"), VertexPoint[0], VertexPoint[1], VertexPoint[2]);
+		}
+
+		FBox Bounds(SomeVec);
+		UE_LOG(LogChaos, Log, TEXT("Bounds[%3.5f,%3.5f,%3.5f] -> [%3.5f,%3.5f,%3.5f]"),
+			Bounds.Min[0], Bounds.Min[1], Bounds.Min[2], Bounds.Max[0], Bounds.Max[1], Bounds.Max[2]);
+		CalculateVolumeAndCenterOfMass(MassSpaceParticles, TriMesh->GetElements(), MassProperties.Volume, MassProperties.CenterOfMass);
+
+		UE_LOG(LogChaos, Log, TEXT("CenterOfMass[%3.5f,%3.5f,%3.5f]"),
+			MassProperties.CenterOfMass[0], MassProperties.CenterOfMass[1], MassProperties.CenterOfMass[2]);
+		UE_LOG(LogChaos, Log, TEXT("Volume: %3.5f"), MassProperties.Volume);
+		
+		EXPECT_NEAR(MassProperties.CenterOfMass.X - GlobalTranslation[0], 0.0f, KINDA_SMALL_NUMBER);
+		EXPECT_NEAR(MassProperties.CenterOfMass.Y - GlobalTranslation[1], 0.0f, KINDA_SMALL_NUMBER);
+		EXPECT_NEAR(MassProperties.CenterOfMass.Z - GlobalTranslation[2], 0.0f, KINDA_SMALL_NUMBER);
+
+
+		for (int32 Idx = 0; Idx < Vertex.Num(); ++Idx)
+		{
+			MassSpaceParticles.X(Idx) -= MassProperties.CenterOfMass;
+		}
+
+		float Density = 0.01f;
+		TVector<float, 3> ZeroVec(0);
+		CalculateInertiaAndRotationOfMass(MassSpaceParticles, TriMesh->GetSurfaceElements(), Density, ZeroVec, MassProperties.InertiaTensor, MassProperties.RotationOfMass);
+
+
+		UE_LOG(LogChaos, Log, TEXT("InertiaTensor[%3.5f,%3.5f,%3.5f]"),
+			MassProperties.InertiaTensor.M[0][0], MassProperties.InertiaTensor.M[1][1], MassProperties.InertiaTensor.M[2][2]);
+		UE_LOG(LogChaos, Log, TEXT("RotationOfMass[%3.5f,%3.5f,%3.5f]"),
+			MassProperties.RotationOfMass.Euler()[0], MassProperties.RotationOfMass.Euler()[1], MassProperties.RotationOfMass.Euler()[2]);
+
+		// rotational alignment.
+		EXPECT_NEAR(MassProperties.RotationOfMass.Euler()[0], 135.f, SMALL_NUMBER);
+		EXPECT_NEAR(MassProperties.RotationOfMass.Euler()[1], 0, SMALL_NUMBER);
+		EXPECT_NEAR(MassProperties.RotationOfMass.Euler()[2], 0, SMALL_NUMBER);
+		// X dominate inertia tensor
+		EXPECT_GT(MassProperties.InertiaTensor.M[0][0], MassProperties.InertiaTensor.M[2][2]);
+		EXPECT_GT(MassProperties.InertiaTensor.M[0][0], MassProperties.InertiaTensor.M[1][1]);
+	}
+
+}
\ No newline at end of file
diff --git a/Engine/Source/Programs/HeadlessChaos/Private/GeometryCollection/GeometryCollectionTestPlugin.cpp b/Engine/Source/Programs/HeadlessChaos/Private/GeometryCollection/GeometryCollectionTestPlugin.cpp
new file mode 100644
--- /dev/null
+++ b/Engine/Source/Programs/HeadlessChaos/Private/GeometryCollection/GeometryCollectionTestPlugin.cpp
@@ -0,0 +1,31 @@
+// Copyright Epic Games, Inc. All Rights Reserved.
+
+#include "GeometryCollection/GeometryCollectionTestPlugin.h"
+
+#include "CoreMinimal.h"
+#include "Modules/ModuleManager.h"
+
+
+
+
+class FGeometryCollectionTestPlugin : public IGeometryCollectionTestPlugin
+{
+	/** IModuleInterface implementation */
+	virtual void StartupModule() override;
+	virtual void ShutdownModule() override;
+};
+
+IMPLEMENT_MODULE(FGeometryCollectionTestPlugin, GeometryCollectionTestCore)
+
+
+
+void FGeometryCollectionTestPlugin::StartupModule()
+{
+
+}
+
+
+void FGeometryCollectionTestPlugin::ShutdownModule()
+{
+
+}
diff --git a/Engine/Source/Programs/HeadlessChaos/Private/GeometryCollection/GeometryCollectionTestSimulation.cpp b/Engine/Source/Programs/HeadlessChaos/Private/GeometryCollection/GeometryCollectionTestSimulation.cpp
new file mode 100644
--- /dev/null
+++ b/Engine/Source/Programs/HeadlessChaos/Private/GeometryCollection/GeometryCollectionTestSimulation.cpp
@@ -0,0 +1,401 @@
+// Copyright Epic Games, Inc. All Rights Reserved.
+
+#include "GeometryCollection/GeometryCollectionTestSimulation.h"
+#include "GeometryCollection/GeometryCollectionTestUtility.h"
+#include "GeometryCollection/GeometryCollectionTestFramework.h"
+
+#include "GeometryCollection/GeometryCollection.h"
+#include "GeometryCollection/GeometryCollectionUtility.h"
+#include "GeometryCollection/TransformCollection.h"
+#include "UObject/Package.h"
+#include "UObject/UObjectGlobals.h"
+
+#include "GeometryCollectionProxyData.h"
+#include "GeometryCollection/GeometryCollectionSimulationTypes.h"
+#include "PhysicsProxy/PhysicsProxies.h"
+#include "Chaos/ErrorReporter.h"
+#include "ChaosSolversModule.h"
+#include "PBDRigidsSolver.h"
+#include "HeadlessChaosTestUtility.h"
+
+#define SMALL_THRESHOLD 1e-4
+
+// #TODO Lots of duplication in here, anyone making solver or object changes
+// has to go and fix up so many callsites here and they're all pretty much
+// Identical. The similar code should be pulled out
+
+namespace GeometryCollectionTest
+{
+	using namespace ChaosTest;
+	
+	TYPED_TEST(AllTraits, GeometryCollection_RigidBodies_SingleFallingUnderGravity)
+	{
+		using Traits = TypeParam;
+		TGeometryCollectionWrapper<Traits>* Collection = TNewSimulationObject<GeometryType::GeometryCollectionWithSingleRigid>::Init<Traits>()->As<TGeometryCollectionWrapper<Traits>>();
+
+		TFramework<Traits> UnitTest;
+		UnitTest.AddSimulationObject(Collection);
+		UnitTest.Initialize();
+		UnitTest.Advance();
+
+		{ // test results
+			EXPECT_LT(FMath::Abs(Collection->RestCollection->Transform[0].GetTranslation().Z), SMALL_THRESHOLD); // rest never touched
+			EXPECT_EQ(Collection->DynamicCollection->Transform.Num(), 1); // simulated is falling
+			EXPECT_LT(Collection->DynamicCollection->Transform[0].GetTranslation().Z, 0.f);
+			EXPECT_NEAR(Collection->DynamicCollection->Transform[0].GetTranslation().Z, -980.f * UnitTest.Dt * UnitTest.Dt, 1e-2);// we seem to be twice gravity
+		}
+	}
+
+	
+	TYPED_TEST(AllTraits, GeometryCollection_RigidBodies_SingleBodyCollidingWithGroundPlane)
+	{
+		using Traits = TypeParam;
+		CreationParameters Params; Params.ImplicitType = EImplicitTypeEnum::Chaos_Implicit_Box; Params.SimplicialType = ESimplicialType::Chaos_Simplicial_Box;
+		TGeometryCollectionWrapper<Traits>* Collection = TNewSimulationObject<GeometryType::GeometryCollectionWithSingleRigid>::Init<Traits>(Params)->As<TGeometryCollectionWrapper<Traits>>();
+		RigidBodyWrapper* Floor = TNewSimulationObject<GeometryType::RigidFloor>::Init<Traits>()->As<RigidBodyWrapper>();
+
+		TFramework<Traits> UnitTest;
+		UnitTest.AddSimulationObject(Collection);
+		UnitTest.AddSimulationObject(Floor);
+		UnitTest.Initialize();
+		for (int i = 0; i < 10; i++) UnitTest.Advance();
+
+		{
+			EXPECT_LT(FMath::Abs(Collection->RestCollection->Transform[0].GetTranslation().Z), SMALL_THRESHOLD);
+			EXPECT_EQ(Collection->DynamicCollection->Transform.Num(), 1);
+			EXPECT_LT(FMath::Abs(Collection->DynamicCollection->Transform[0].GetTranslation().Z - 1.0), SMALL_THRESHOLD);
+		}
+	}
+
+	
+	TYPED_TEST(AllTraits, GeometryCollection_RigidBodies_SingleSphereCollidingWithSolverFloor)
+	{
+		using Traits = TypeParam;
+		FVector Scale(0.5f);
+		CreationParameters Params; 
+		Params.ImplicitType = EImplicitTypeEnum::Chaos_Implicit_Sphere; 
+		Params.GeomTransform.SetScale3D(Scale); // Sphere radius
+		TGeometryCollectionWrapper<Traits>* Collection =
+			TNewSimulationObject<GeometryType::GeometryCollectionWithSingleRigid>::Init<Traits>(Params)->As<TGeometryCollectionWrapper<Traits>>();
+		RigidBodyWrapper* Floor = 
+			TNewSimulationObject<GeometryType::RigidFloor>::Init<Traits>()->As<RigidBodyWrapper>();
+
+		TFramework<Traits> UnitTest;
+		UnitTest.AddSimulationObject(Collection);
+		UnitTest.AddSimulationObject(Floor);
+		UnitTest.Initialize();
+		for (int i = 0; i < 10; i++) UnitTest.Advance();
+
+		{ // test results
+			EXPECT_LT(FMath::Abs(Collection->RestCollection->Transform[0].GetTranslation().Z), SMALL_THRESHOLD);
+			EXPECT_EQ(Collection->DynamicCollection->Transform.Num(), 1);
+			EXPECT_LT(FMath::Abs(Collection->DynamicCollection->Transform[0].GetTranslation().Z) - Scale[0], SMALL_THRESHOLD);
+		}
+	}
+
+
+	
+	TYPED_TEST(AllTraits, GeometryCollection_RigidBodies_SingleCubeIntersectingWithSolverFloor)
+	{
+		using Traits = TypeParam;
+		FVector Scale(2.f);
+		CreationParameters Params; Params.ImplicitType = EImplicitTypeEnum::Chaos_Implicit_Box;  Params.SimplicialType = ESimplicialType::Chaos_Simplicial_Box;
+		Params.GeomTransform.SetScale3D(Scale); // Sphere radius
+
+		TGeometryCollectionWrapper<Traits>* Collection = TNewSimulationObject<GeometryType::GeometryCollectionWithSingleRigid>::Init<Traits>(Params)->As<TGeometryCollectionWrapper<Traits>>();
+		RigidBodyWrapper* Floor = TNewSimulationObject<GeometryType::RigidFloor>::Init<Traits>()->As<RigidBodyWrapper>();
+
+		TFramework<Traits> UnitTest;
+		UnitTest.AddSimulationObject(Collection);
+		UnitTest.AddSimulationObject(Floor);
+		UnitTest.Initialize();
+		for (int i = 0; i < 10; i++) UnitTest.Advance();
+
+		{
+			EXPECT_LT(FMath::Abs(Collection->RestCollection->Transform[0].GetTranslation().Z), SMALL_THRESHOLD);
+			EXPECT_EQ(Collection->DynamicCollection->Transform.Num(), 1);
+			EXPECT_LT(FMath::Abs(Collection->DynamicCollection->Transform[0].GetTranslation().Z - Scale[2]), SMALL_THRESHOLD);
+		}
+	}
+
+
+	
+	TYPED_TEST(AllTraits, GeometryCollection_RigidBodies_SingleKinematicBody)
+	{
+		using Traits = TypeParam;
+		CreationParameters Params; Params.DynamicState = EObjectStateTypeEnum::Chaos_Object_Kinematic;
+		TGeometryCollectionWrapper<Traits>* Collection = TNewSimulationObject<GeometryType::GeometryCollectionWithSingleRigid>::Init<Traits>(Params)->As<TGeometryCollectionWrapper<Traits>>();
+
+		TFramework<Traits> UnitTest;
+		UnitTest.AddSimulationObject(Collection);
+		UnitTest.Initialize();
+		for (int i = 0; i < 3; i++)
+			UnitTest.Advance();
+
+		{
+			TManagedArray<FTransform>& Transform = Collection->DynamicCollection->Transform;
+			EXPECT_EQ(Transform.Num(), 1);
+			//UE_LOG(LogTest, Verbose, TEXT("Position : (%3.5f,%3.5f,%3.5f)"), Transform[0].GetTranslation().X, Transform[0].GetTranslation().Y, Transform[0].GetTranslation().Z);
+			EXPECT_EQ(Transform[0].GetTranslation().Z, 0.f);
+			EXPECT_EQ(Collection->DynamicCollection->DynamicState[0], (int32)EObjectStateTypeEnum::Chaos_Object_Kinematic);
+		}
+	}
+
+
+	
+	TYPED_TEST(AllTraits, GeometryCollection_RigidBodies_SleepingDontMove)
+	{
+		using Traits = TypeParam;
+		CreationParameters Params;
+		Params.DynamicState = EObjectStateTypeEnum::Chaos_Object_Sleeping;
+		Params.ImplicitType = EImplicitTypeEnum::Chaos_Implicit_Box;
+		float InitialStartHeight = 5.0f;
+		Params.RootTransform.SetLocation(FVector(0.f, 0.f, InitialStartHeight));
+		TGeometryCollectionWrapper<Traits>* SleepingCollection = TNewSimulationObject<GeometryType::GeometryCollectionWithSingleRigid>::Init<Traits>(Params)->As<TGeometryCollectionWrapper<Traits>>();
+
+		TFramework<Traits> UnitTest;
+		UnitTest.AddSimulationObject(SleepingCollection);
+		UnitTest.Initialize();
+
+		const auto& Transform0 = SleepingCollection->DynamicCollection->Transform[0];
+		for (int i = 0; i < 3; i++)
+		{
+			UnitTest.Advance();
+			//UE_LOG(LogTest, Verbose, TEXT("Position[0] : (%3.5f,%3.5f,%3.5f)"), Transform0.GetTranslation().X, Transform0.GetTranslation().Y, Transform0.GetTranslation().Z);
+		}
+
+		{
+			// particle doesn't fall due to sleeping state
+			EXPECT_EQ(SleepingCollection->DynamicCollection->DynamicState[0], (int32)EObjectStateTypeEnum::Chaos_Object_Sleeping);
+			EXPECT_LT(FMath::Abs(SleepingCollection->DynamicCollection->Transform[0].GetTranslation().Z - InitialStartHeight), SMALL_THRESHOLD);
+		}
+
+	}
+
+
+	
+	TYPED_TEST(AllTraits, GeometryCollection_RigidBodies_SleepingActivation)
+	{
+		using Traits = TypeParam;
+		CreationParameters Params; 
+		Params.SimplicialType = ESimplicialType::Chaos_Simplicial_Box;
+		Params.ImplicitType = EImplicitTypeEnum::Chaos_Implicit_Box;
+
+		Params.DynamicState = EObjectStateTypeEnum::Chaos_Object_Dynamic;
+		Params.RootTransform.SetLocation(FVector(0.f, 0.f, 15.f));
+		TGeometryCollectionWrapper<Traits>* MovingCollection = TNewSimulationObject<GeometryType::GeometryCollectionWithSingleRigid>::Init<Traits>(Params)->As<TGeometryCollectionWrapper<Traits>>();
+
+		float InitialStartHeight = 5.0f;
+		Params.DynamicState = EObjectStateTypeEnum::Chaos_Object_Sleeping;
+		Params.RootTransform.SetLocation(FVector(0.f, 0.f, InitialStartHeight));
+		TGeometryCollectionWrapper<Traits>* SleepingCollection = TNewSimulationObject<GeometryType::GeometryCollectionWithSingleRigid>::Init<Traits>(Params)->As<TGeometryCollectionWrapper<Traits>>();
+
+		TFramework<Traits> UnitTest;
+		UnitTest.AddSimulationObject(SleepingCollection);
+		UnitTest.AddSimulationObject(MovingCollection);
+		UnitTest.Initialize();
+
+		const auto& Transform0 = MovingCollection->DynamicCollection->Transform[0];
+		const auto& Transform1 = SleepingCollection->DynamicCollection->Transform[0];
+		for (int i = 0; i < 15; i++)
+		{
+			UnitTest.Advance();
+
+			//UE_LOG(LogTest, Verbose, TEXT("Position[0] : (%3.5f,%3.5f,%3.5f)"), Transform0.GetTranslation().X, Transform0.GetTranslation().Y, Transform0.GetTranslation().Z);
+			//UE_LOG(LogTest, Verbose, TEXT("Position[1] : (%3.5f,%3.5f,%3.5f)"), Transform1.GetTranslation().X, Transform1.GetTranslation().Y, Transform1.GetTranslation().Z);
+		}
+
+		{
+			// Is now dynamic and has moved from initial position
+			EXPECT_EQ(SleepingCollection->DynamicCollection->DynamicState[0], (int32)EObjectStateTypeEnum::Chaos_Object_Dynamic);
+			EXPECT_LT(Transform0.GetTranslation().Z, InitialStartHeight - 2.0f);
+		}
+
+	}
+
+	
+	TYPED_TEST(AllTraits, GeometryCollection_RigidBodies_CollisionGroup)
+	{
+		using Traits = TypeParam;
+		/*
+		TUniquePtr<Chaos::FChaosPhysicsMaterial> PhysicalMaterial = nullptr;
+		TSharedPtr<FGeometryCollection> RestCollection = nullptr;
+		TSharedPtr<FGeometryDynamicCollection> DynamicCollection = nullptr;
+
+		//
+		//  Rigid Body Setup
+		//
+		auto RestInitFunc = [](TSharedPtr<FGeometryCollection>& RestCollection)
+		{
+			RestCollection->AppendGeometry(*GeometryCollection::MakeCubeElement(FTransform(FVector(0, 0, 210.0)), FVector(100.0)));
+			RestCollection->AppendGeometry(*GeometryCollection::MakeCubeElement(FTransform(FVector(0, 0, 320.0)), FVector(100.0)));
+			RestCollection->AppendGeometry(*GeometryCollection::MakeCubeElement(FTransform(FVector(0, 0, 430.0)), FVector(100.0)));
+		};
+
+		////InitCollectionsParameters InitParams = { FTransform(FVector(0, 0, 100.0)), FVector(100.0), RestInitFunc, (int32)EObjectStateTypeEnum::Chaos_Object_Kinematic };
+		//InitCollections(PhysicalMaterial, RestCollection, DynamicCollection, InitParams);
+
+		//
+		// Solver setup
+		//
+		auto CustomFunc = [&RestCollection, &DynamicCollection, &PhysicalMaterial](FSimulationParameters& InParams)
+		{
+			InParams.Shared.SizeSpecificData[0].ImplicitType = EImplicitTypeEnum::Chaos_Implicit_Box;
+		};
+
+		FGeometryCollectionPhysicsProxy* PhysObject = RigidBodySetup(PhysicalMaterial, RestCollection, DynamicCollection, CustomFunc);
+		PhysObject->SetCollisionParticlesPerObjectFraction( 1.0 );
+
+		Chaos::FPBDRigidsSolver* Solver = FChaosSolversModule::GetModule()->CreateSolver(nullptr, ESolverFlags::Standalone);
+#if CHAOS_PARTICLEHANDLE_TODO
+		Solver->RegisterObject(PhysObject);
+#endif
+		//Solver->SetHasFloor(true);
+		Solver->SetEnabled(true);
+		//PhysObject->ActivateBodies();
+
+		Solver->AdvanceSolverBy(1 / 24.);
+#if TODO_REIMPLEMENT_GET_RIGID_PARTICLES
+		Chaos::TPBDRigidParticles<float, 3>& Particles = Solver->GetRigidParticles();
+
+		for (int Frame = 1; Frame < 200; Frame++)
+		{
+			Solver->AdvanceSolverBy(1 / 24.);
+			//FinalizeSolver(*Solver);
+			
+			if (Frame == 1)
+			{
+				Particles.CollisionGroup(0)=  0;
+				Particles.CollisionGroup(1)=  1;
+				Particles.CollisionGroup(2)=  1;
+				Particles.CollisionGroup(3)=  3;
+				Particles.CollisionGroup(4)= -1;
+			}
+			if (Frame == 13)
+			{
+				EXPECT_LT(FMath::Abs(Particles.X(0).Z), SMALL_NUMBER);
+				EXPECT_LT(FMath::Abs(Particles.X(1).Z - 50.f), 10.f);
+				EXPECT_LT(FMath::Abs(Particles.X(2).Z - 150.f), 10.f);
+			}
+			if( Frame == 30 )
+			{
+				EXPECT_LT(FMath::Abs(Particles.X(0).Z), SMALL_NUMBER);
+				EXPECT_LT(FMath::Abs(Particles.X(1).Z - 50.f), 10.f);
+				EXPECT_LT(FMath::Abs(Particles.X(2).Z - 150.f), 10.f);
+				EXPECT_GT(Particles.X(3).Z, 50.f);
+				EXPECT_LT(Particles.X(4).Z, -100);
+			}
+			if (Frame == 31)
+			{
+				Particles.CollisionGroup(0) = 0;
+				Particles.CollisionGroup(1) = -1;
+				Particles.CollisionGroup(2) = 1;
+				Particles.CollisionGroup(3) = -1;
+				Particles.CollisionGroup(4) = -1;
+			}
+		}
+
+		EXPECT_LT(FMath::Abs(Particles.X(0).Z), SMALL_NUMBER);
+		EXPECT_LT(Particles.X(1).Z, -10000);
+		EXPECT_GT(Particles.X(2).Z, 50.0);
+		EXPECT_LT(Particles.X(3).Z, -10000);
+		EXPECT_LT(Particles.X(4).Z, -10000);
+#endif
+
+		FChaosSolversModule::GetModule()->DestroySolver(Solver);
+		delete PhysObject;
+		*/
+	}
+
+
+
+	
+	TYPED_TEST(AllTraits, GeometryCollection_TestImplicitCollisionGeometry)
+	{
+		using Traits = TypeParam;
+		typedef Chaos::TVector<FReal, 3> Vec;
+
+		CreationParameters Params; 
+		Params.SimplicialType = ESimplicialType::Chaos_Simplicial_GriddleBox;
+		Params.ImplicitType = EImplicitTypeEnum::Chaos_Implicit_LevelSet;
+		Params.CollisionType = ECollisionTypeEnum::Chaos_Surface_Volumetric;
+
+		TGeometryCollectionWrapper<Traits>* Collection =
+			TNewSimulationObject<GeometryType::GeometryCollectionWithSingleRigid>::Init<Traits>(
+				Params)->As<TGeometryCollectionWrapper<Traits>>();
+
+		const TManagedArray<TUniquePtr<Chaos::TBVHParticles<float,3>>>& Simplicials =
+			Collection->RestCollection->template GetAttribute<TUniquePtr<Chaos::TBVHParticles<float,3>>>(
+				FGeometryDynamicCollection::SimplicialsAttribute, FTransformCollection::TransformGroup);
+		EXPECT_EQ(Simplicials.Num(), 1);
+		const Chaos::TBVHParticles<float, 3>& Simplicial = *Simplicials[0];
+
+		const TManagedArray<FGeometryDynamicCollection::FSharedImplicit>& Implicits = 
+			Collection->RestCollection->template GetAttribute<FGeometryDynamicCollection::FSharedImplicit>(
+				FGeometryDynamicCollection::ImplicitsAttribute, FTransformCollection::TransformGroup);
+		EXPECT_EQ(Implicits.Num(), 1);
+		check(Implicits[0]);
+		const Chaos::FImplicitObject& Implicit = *Implicits[0];
+
+		// Ensure all simplicial particles are on the surface of the implicit shape.
+		check(Implicit.GetType() == Chaos::ImplicitObjectType::LevelSet);
+		const Chaos::TLevelSet<float, 3>* LevelSet = static_cast<const Chaos::TLevelSet<float, 3>*>(&Implicit);
+		const FReal DxSize = LevelSet->GetGrid().Dx().Size();
+
+		FReal MinX = TNumericLimits<FReal>::Max();
+		FReal MinY = TNumericLimits<FReal>::Max();
+		FReal MinZ = TNumericLimits<FReal>::Max();
+
+		FReal MaxX = -TNumericLimits<FReal>::Max();
+		FReal MaxY = -TNumericLimits<FReal>::Max();
+		FReal MaxZ = -TNumericLimits<FReal>::Max();
+		for (uint32 Idx = 0; Idx < Simplicial.Size(); ++Idx)
+		{
+			const FReal phi = Implicit.SignedDistance(Simplicial.X(Idx));
+			EXPECT_LT(FMath::Abs(phi), DxSize);
+			//EXPECT_LT(FMath::Abs(phi), 0.01f);
+
+			const auto& Pos = Simplicial.X(Idx);
+			MinX = MinX < Pos[0] ? MinX : Pos[0];
+			MinY = MinY < Pos[1] ? MinY : Pos[1];
+			MinZ = MinZ < Pos[2] ? MinZ : Pos[2];
+
+			MaxX = MaxX > Pos[0] ? MaxX : Pos[0];
+			MaxY = MaxY > Pos[1] ? MaxY : Pos[1];
+			MaxZ = MaxZ > Pos[2] ? MaxZ : Pos[2];
+		}
+
+		// Make sure the geometry occupies a volume.
+		check(MinX < MaxX);
+		check(MinY < MaxY);
+		check(MinZ < MaxZ);
+
+		// Cast a ray through the level set, and make sure it's as we expect.
+		for(float x = 2*MinX; x < 2*MaxX; x += (MaxX-MinX)/10)
+		{
+			Vec Normal;
+			const FReal phi = Implicit.PhiWithNormal(Vec(x, 0, 0), Normal);
+			if (x < MinX || MaxX < x)
+			{
+				check(phi >= -0.01f);
+				EXPECT_GT(phi, -0.01f);
+			}
+			else
+			{
+				check(phi <= 0.01f);
+				EXPECT_LT(phi, 0.01f);
+			}
+
+			if (x < MinX/4)
+			{
+				EXPECT_LT((Normal-Vec(-1,0,0)).Size(), KINDA_SMALL_NUMBER);
+			}
+			else if (x > MaxX/4)
+			{
+				EXPECT_LT((Normal - Vec(1, 0, 0)).Size(), KINDA_SMALL_NUMBER);
+			}
+		}
+	}
+
+}
+
