Merge geometry collection into headless chaos.

diff --git a/Engine/Source/Programs/HeadlessChaos/Private/GeometryCollection/GeometryCollectionTestSimulationField.cpp b/Engine/Source/Programs/HeadlessChaos/Private/GeometryCollection/GeometryCollectionTestSimulationField.cpp
new file mode 100644
--- /dev/null
+++ b/Engine/Source/Programs/HeadlessChaos/Private/GeometryCollection/GeometryCollectionTestSimulationField.cpp
@@ -0,0 +1,956 @@
+// Copyright Epic Games, Inc. All Rights Reserved.
+
+#include "GeometryCollection/GeometryCollectionTestSimulationField.h"
+#include "GeometryCollection/GeometryCollectionTestUtility.h"
+#include "GeometryCollection/GeometryCollectionTestFramework.h"
+
+#include "GeometryCollection/GeometryCollection.h"
+#include "GeometryCollection/GeometryCollectionUtility.h"
+#include "GeometryCollection/TransformCollection.h"
+#include "UObject/Package.h"
+#include "UObject/UObjectGlobals.h"
+
+#include "Field/FieldSystem.h"
+#include "Field/FieldSystemNodes.h"
+#include "GeometryCollectionProxyData.h"
+#include "GeometryCollection/GeometryCollectionSimulationTypes.h"
+#include "PhysicsProxy/PhysicsProxies.h"
+#include "Chaos/ErrorReporter.h"
+#include "Chaos/PBDRigidClustering.h"
+#include "PBDRigidsSolver.h"
+#include "ChaosSolversModule.h"
+#include "HeadlessChaosTestUtility.h"
+
+extern bool bSkipGeomFailing;
+
+//#include "GeometryCollection/GeometryCollectionAlgo.h"
+
+#define SMALL_THRESHOLD 1e-4
+
+// #TODO Lots of duplication in here, anyone making solver or object changes
+// has to go and fix up so many callsites here and they're all pretty much
+// Identical. The similar code should be pulled out
+
+namespace GeometryCollectionTest
+{
+
+	using namespace ChaosTest;
+
+	TYPED_TEST(AllTraits, GeometryCollection_RigidBodies_Field_KinematicActivationOnProxyDuringInit)
+	{
+		using Traits = TypeParam;
+		const FVector Translation0(0, 0, 1);
+
+		CreationParameters Params;
+		Params.DynamicState = EObjectStateTypeEnum::Chaos_Object_Kinematic;
+		Params.RootTransform.SetLocation(Translation0);
+		TGeometryCollectionWrapper<Traits>* Collection = TNewSimulationObject<GeometryType::GeometryCollectionWithSingleRigid>::Init<Traits>(Params)->As<TGeometryCollectionWrapper<Traits>>();
+		Params.RootTransform.SetLocation(FVector(100,0,0));
+		TGeometryCollectionWrapper<Traits>* CollectionOther = TNewSimulationObject<GeometryType::GeometryCollectionWithSingleRigid>::Init<Traits>(Params)->As<TGeometryCollectionWrapper<Traits>>();
+
+		TFramework<Traits> UnitTest;
+		UnitTest.AddSimulationObject(CollectionOther);
+		UnitTest.AddSimulationObject(Collection);
+
+		FFieldSystemPhysicsProxy* FieldObject = new FFieldSystemPhysicsProxy(nullptr);
+
+		FRadialIntMask* RadialMaskTmp = new FRadialIntMask();
+		RadialMaskTmp->Position = FVector(0.0, 0.0, 0.0);
+		RadialMaskTmp->Radius = 100.0;
+		RadialMaskTmp->InteriorValue = (int32)EObjectStateTypeEnum::Chaos_Object_Dynamic;
+		RadialMaskTmp->ExteriorValue = (int32)EObjectStateTypeEnum::Chaos_Object_Kinematic;
+		RadialMaskTmp->SetMaskCondition = ESetMaskConditionType::Field_Set_IFF_NOT_Interior;
+		FName TargetNameTmp = GetFieldPhysicsName(EFieldPhysicsType::Field_DynamicState);
+		Collection->PhysObject->BufferCommand(UnitTest.Solver, { TargetNameTmp, RadialMaskTmp });
+
+		UnitTest.Initialize();
+
+		EXPECT_EQ(CollectionOther->DynamicCollection->DynamicState[0], (int32)EObjectStateTypeEnum::Chaos_Object_Kinematic);
+		EXPECT_EQ(Collection->DynamicCollection->DynamicState[0], (int32)EObjectStateTypeEnum::Chaos_Object_Dynamic);
+	}
+
+	
+	TYPED_TEST(AllTraits, GeometryCollection_RigidBodies_Field_KinematicActivationOnProxyDuringUpdate)
+	{
+		using Traits = TypeParam;
+		const FVector Translation0(0, 0, 1);
+
+		CreationParameters Params;
+		Params.DynamicState = EObjectStateTypeEnum::Chaos_Object_Kinematic;
+		Params.RootTransform.SetLocation(Translation0);
+		TGeometryCollectionWrapper<Traits>* Collection = TNewSimulationObject<GeometryType::GeometryCollectionWithSingleRigid>::Init<Traits>(Params)->As<TGeometryCollectionWrapper<Traits>>();
+
+		TFramework<Traits> UnitTest; UnitTest.Dt = 1/24.0;
+		UnitTest.AddSimulationObject(Collection);
+
+		FFieldSystemPhysicsProxy* FieldObject = new FFieldSystemPhysicsProxy(nullptr);
+
+		UnitTest.Initialize();
+
+		{
+			UnitTest.Advance();
+		}
+
+		TManagedArray<int32>& DynamicState = Collection->DynamicCollection->DynamicState;
+		EXPECT_EQ(DynamicState[0], (int32)EObjectStateTypeEnum::Chaos_Object_Kinematic);
+
+		// simulated
+		const TManagedArray<FTransform>& Transform = Collection->DynamicCollection->Transform;
+		EXPECT_EQ(Transform.Num(), 1);
+		const FVector Translation1 = Transform[0].GetTranslation();
+		EXPECT_NEAR( (Translation0-Translation1).Size(), 0.f, KINDA_SMALL_NUMBER);
+		EXPECT_NEAR(Transform[0].GetTranslation().Z, 1.f, KINDA_SMALL_NUMBER);
+
+		FRadialIntMask* RadialMask = new FRadialIntMask();
+		RadialMask->Position = FVector(0.0, 0.0, 0.0);
+		RadialMask->Radius = 100.0;
+		RadialMask->InteriorValue = (int32)EObjectStateTypeEnum::Chaos_Object_Dynamic;
+		RadialMask->ExteriorValue = (int32)EObjectStateTypeEnum::Chaos_Object_Kinematic;
+		RadialMask->SetMaskCondition = ESetMaskConditionType::Field_Set_IFF_NOT_Interior;
+		FName TargetName = GetFieldPhysicsName(EFieldPhysicsType::Field_DynamicState);
+		Collection->PhysObject->BufferCommand(UnitTest.Solver, { TargetName, RadialMask });
+
+		{
+			UnitTest.Advance();
+		}
+		EXPECT_EQ(DynamicState[0], (int32)EObjectStateTypeEnum::Chaos_Object_Dynamic);
+
+		const FVector Translation2 = Transform[0].GetTranslation();
+		EXPECT_NE(Translation1, Translation2);
+		EXPECT_LE(Transform[0].GetTranslation().Z, 0.f);
+	}
+
+
+	
+	TYPED_TEST(AllTraits, GeometryCollection_RigidBodies_Field_KinematicActivation)
+	{
+		if(bSkipGeomFailing)
+		{
+			return;
+		}
+		using Traits = TypeParam;
+		const FVector Translation0(0, 0, 1);
+
+		CreationParameters Params; 
+		Params.DynamicState = EObjectStateTypeEnum::Chaos_Object_Kinematic;
+		Params.RootTransform.SetLocation(Translation0);
+		TGeometryCollectionWrapper<Traits>* Collection = TNewSimulationObject<GeometryType::GeometryCollectionWithSingleRigid>::Init<Traits>(Params)->As<TGeometryCollectionWrapper<Traits>>();
+
+		TFramework<Traits> UnitTest;
+		UnitTest.AddSimulationObject(Collection);
+
+		FFieldSystemPhysicsProxy* FieldObject = new FFieldSystemPhysicsProxy(nullptr);
+		UnitTest.AddFieldObject(FieldObject);
+
+		UnitTest.Initialize();
+
+		for (int i = 0; i < 100; i++)
+		{
+			UnitTest.Advance();
+		}
+
+		// simulated
+		const TManagedArray<FTransform>& Transform = Collection->DynamicCollection->Transform;
+		EXPECT_EQ(Transform.Num(), 1);
+		const FVector Translation1 = Transform[0].GetTranslation();
+		EXPECT_NEAR( (Translation0-Translation1).Size(), 0.f, KINDA_SMALL_NUMBER);
+		EXPECT_NEAR(Transform[0].GetTranslation().Z, 1.f, KINDA_SMALL_NUMBER);
+
+		FRadialIntMask* RadialMask = new FRadialIntMask();
+		RadialMask->Position = FVector(0.0, 0.0, 0.0);
+		RadialMask->Radius = 100.0;
+		RadialMask->InteriorValue = (int32)EObjectStateTypeEnum::Chaos_Object_Dynamic;
+		RadialMask->ExteriorValue = (int32)EObjectStateTypeEnum::Chaos_Object_Kinematic;
+		RadialMask->SetMaskCondition = ESetMaskConditionType::Field_Set_IFF_NOT_Interior;
+		FName TargetName = GetFieldPhysicsName(EFieldPhysicsType::Field_DynamicState);
+		FieldObject->BufferCommand(UnitTest.Solver, { TargetName, RadialMask });
+
+		for (int i = 0; i < 100; i++)
+		{
+			UnitTest.Advance();
+		}
+
+		const FVector Translation2 = Transform[0].GetTranslation();
+		EXPECT_NE(Translation1, Translation2);
+		EXPECT_LE(Transform[0].GetTranslation().Z, 0.f);
+	}
+
+	
+	TYPED_TEST(AllTraits, GeometryCollection_RigidBodies_Field_InitialLinearVelocity)
+	{
+		if(bSkipGeomFailing)
+		{
+			return;
+		}
+		using Traits = TypeParam;
+		TFramework<Traits> UnitTest;
+
+		// Physics Object Setup
+		CreationParameters Params;
+		Params.DynamicState = EObjectStateTypeEnum::Chaos_Object_Kinematic;
+		Params.RootTransform.SetLocation(FVector(0.0, 0.0, 0.0));
+
+		Params.InitialVelocityType = EInitialVelocityTypeEnum::Chaos_Initial_Velocity_User_Defined;
+		Params.InitialLinearVelocity = FVector(0.f, 100.f, 0.f);
+		TGeometryCollectionWrapper<Traits>* Collection = TNewSimulationObject<GeometryType::GeometryCollectionWithSingleRigid>::Init<Traits>(Params)->As<TGeometryCollectionWrapper<Traits>>();
+		UnitTest.AddSimulationObject(Collection);
+
+		// Field setup
+		FRadialIntMask* RadialMask = new FRadialIntMask();
+		RadialMask->Position = FVector(0.0, 0.0, 0.0);
+		RadialMask->Radius = 5.0f;
+		RadialMask->InteriorValue = (int32)EObjectStateTypeEnum::Chaos_Object_Dynamic;
+		RadialMask->ExteriorValue = (int32)EObjectStateTypeEnum::Chaos_Object_Dynamic;
+		RadialMask->SetMaskCondition = ESetMaskConditionType::Field_Set_Always;
+
+		FFieldSystemPhysicsProxy* FieldObject = new FFieldSystemPhysicsProxy(nullptr);
+		UnitTest.AddFieldObject(FieldObject);
+
+		UnitTest.Initialize();
+
+		TManagedArray<FTransform>& Transform = Collection->DynamicCollection->Transform;
+		TManagedArray<int32>& DynamicState = Collection->DynamicCollection->DynamicState;
+
+
+		float PreviousY = 0.f;
+		EXPECT_EQ(Transform[0].GetTranslation().X, 0);
+		EXPECT_EQ(Transform[0].GetTranslation().Y, 0);
+
+		for (int Frame = 0; Frame < 10; Frame++)
+		{
+			UnitTest.Advance();
+
+			if (Frame == 1)
+			{
+				FName TargetName = GetFieldPhysicsName(EFieldPhysicsType::Field_DynamicState);
+				Collection->PhysObject->BufferCommand(UnitTest.Solver, { TargetName, RadialMask });
+			}
+
+			if (Frame >= 2)
+			{
+				EXPECT_EQ(DynamicState[0], (int32)EObjectStateTypeEnum::Chaos_Object_Dynamic);
+				EXPECT_EQ(Transform[0].GetTranslation().X, 0);
+				EXPECT_GT(Transform[0].GetTranslation().Y, PreviousY);
+				EXPECT_LT(Transform[0].GetTranslation().Z, 0);
+			}
+			else
+			{
+				EXPECT_EQ(DynamicState[0], (int32)EObjectStateTypeEnum::Chaos_Object_Kinematic);
+				EXPECT_EQ(Transform[0].GetTranslation().X, 0);
+				EXPECT_EQ(Transform[0].GetTranslation().Y, 0);
+				EXPECT_EQ(Transform[0].GetTranslation().Z, 0);
+			}
+			PreviousY = Transform[0].GetTranslation().Y;
+		}
+
+	}
+
+	
+	TYPED_TEST(AllTraits, GeometryCollection_RigidBodies_Field_StayDynamic)
+	{
+		if(bSkipGeomFailing)
+		{
+			return;
+		}
+		using Traits = TypeParam;
+		TFramework<Traits> UnitTest;
+		float PreviousHeight = 5.f;
+
+		// Physics Object Setup
+		CreationParameters Params;
+		Params.DynamicState = EObjectStateTypeEnum::Chaos_Object_Static;
+		Params.RootTransform.SetLocation(FVector(0.f, 0.f, PreviousHeight));
+		TGeometryCollectionWrapper<Traits>* Collection = TNewSimulationObject<GeometryType::GeometryCollectionWithSingleRigid>::Init<Traits>(Params)->As<TGeometryCollectionWrapper<Traits>>();
+		UnitTest.AddSimulationObject(Collection);
+
+		// Field setup
+		FRadialIntMask * RadialMask = new FRadialIntMask();
+		RadialMask->Position = FVector(0.0, 0.0, PreviousHeight);
+		RadialMask->Radius = 5.0;
+		RadialMask->InteriorValue = (int32)EObjectStateTypeEnum::Chaos_Object_Dynamic;
+		RadialMask->ExteriorValue = (int32)EObjectStateTypeEnum::Chaos_Object_Kinematic;
+		RadialMask->SetMaskCondition = ESetMaskConditionType::Field_Set_IFF_NOT_Interior;
+
+		FFieldSystemPhysicsProxy* FieldObject = new FFieldSystemPhysicsProxy(nullptr);
+		UnitTest.AddFieldObject(FieldObject);
+
+		UnitTest.Initialize();
+
+		TManagedArray<FTransform>& Transform = Collection->DynamicCollection->Transform;
+		for (int Frame = 0; Frame < 10; Frame++)
+		{
+			//UE_LOG(LogTest, Verbose, TEXT("Frame[%d]"), Frame);
+
+			if (Frame == 5)
+			{
+				FName TargetName = GetFieldPhysicsName(EFieldPhysicsType::Field_DynamicState);
+				FieldObject->BufferCommand(UnitTest.Solver, { TargetName, RadialMask });
+			}
+
+			UnitTest.Advance();
+
+			if (Frame < 5)
+			{
+				EXPECT_LT(FMath::Abs(Transform[0].GetTranslation().Z - 5.f), SMALL_THRESHOLD);
+			}
+			else
+			{
+				EXPECT_LT(Transform[0].GetTranslation().Z, PreviousHeight);
+			}
+			PreviousHeight = Transform[0].GetTranslation().Z;
+		}
+
+	}
+
+
+	
+	TYPED_TEST(AllTraits, GeometryCollection_RigidBodies_Field_LinearForce)
+	{
+		if(bSkipGeomFailing)
+		{
+			return;
+		}
+		using Traits = TypeParam;
+		TFramework<Traits> UnitTest;
+
+		// Physics Object Setup
+		CreationParameters Params;
+		Params.DynamicState = EObjectStateTypeEnum::Chaos_Object_Dynamic;
+		Params.RootTransform.SetLocation(FVector(0.f, 0.f, 5.f));
+		TGeometryCollectionWrapper<Traits>* Collection = TNewSimulationObject<GeometryType::GeometryCollectionWithSingleRigid>::Init<Traits>(Params)->As<TGeometryCollectionWrapper<Traits>>();
+		UnitTest.AddSimulationObject(Collection);
+
+		// Field setup
+		FUniformVector * UniformVector = new FUniformVector();
+		UniformVector->Direction = FVector(0.0, 1.0, 0.0);
+		UniformVector->Magnitude = 1000.0;
+
+		FFieldSystemPhysicsProxy* FieldObject = new FFieldSystemPhysicsProxy(nullptr);
+		UnitTest.AddFieldObject(FieldObject);
+
+		UnitTest.Initialize();
+
+		TManagedArray<FTransform>& Transform = Collection->DynamicCollection->Transform;
+		float PreviousY = 0.f;
+		for (int Frame = 0; Frame < 10; Frame++)
+		{
+			if (Frame >= 5)
+			{
+				FName TargetName = GetFieldPhysicsName(EFieldPhysicsType::Field_LinearForce);
+				FieldObject->BufferCommand(UnitTest.Solver, { TargetName, UniformVector->NewCopy() });
+			}
+
+			UnitTest.Advance();
+
+			if (Frame < 5)
+			{
+				EXPECT_LT(FMath::Abs(Transform[0].GetTranslation().Y), SMALL_THRESHOLD);
+			}
+			else
+			{
+				EXPECT_GT(Transform[0].GetTranslation().Y, PreviousY);
+			}
+
+			PreviousY = Transform[0].GetTranslation().Y;
+
+		}
+
+		delete UniformVector;
+	}
+
+	
+	TYPED_TEST(AllTraits, GeometryCollection_RigidBodies_Field_Torque)
+	{
+		if(bSkipGeomFailing)
+		{
+			return;
+		}
+		using Traits = TypeParam;
+		TFramework<Traits> UnitTest;
+
+		// Physics Object Setup
+		FVector Scale = FVector(10.0);
+		CreationParameters Params;
+		Params.DynamicState = EObjectStateTypeEnum::Chaos_Object_Dynamic;
+		Params.RootTransform.SetLocation(FVector(0.f, 0.f, 5.f));
+		Params.GeomTransform.SetScale3D(Scale);
+		TGeometryCollectionWrapper<Traits>* Collection = TNewSimulationObject<GeometryType::GeometryCollectionWithSingleRigid>::Init<Traits>(Params)->As<TGeometryCollectionWrapper<Traits>>();
+		UnitTest.AddSimulationObject(Collection);
+
+		// Field setup
+		FUniformVector * UniformVector = new FUniformVector();
+		UniformVector->Direction = FVector(0.0, 1.0, 0.0);
+		UniformVector->Magnitude = 100.0;
+
+		FFieldSystemPhysicsProxy* FieldObject = new FFieldSystemPhysicsProxy(nullptr);
+		UnitTest.AddFieldObject(FieldObject);
+
+		UnitTest.Initialize();
+
+		TManagedArray<FTransform>& Transform = Collection->DynamicCollection->Transform;
+		float PreviousY = 0.f;
+		for (int Frame = 0; Frame < 10; Frame++)
+		{
+
+			if (Frame >= 5)
+			{
+				FName TargetName = GetFieldPhysicsName(EFieldPhysicsType::Field_AngularTorque);
+				FieldObject->BufferCommand(UnitTest.Solver, { TargetName, UniformVector->NewCopy() });
+			}
+
+			UnitTest.Advance();
+
+			auto& Particles = UnitTest.Solver->GetParticles().GetGeometryCollectionParticles();
+			if (Frame < 5)
+			{
+				EXPECT_LT(FMath::Abs(Transform[0].GetRotation().Euler().Y), SMALL_THRESHOLD);
+			}
+			else
+			{
+				EXPECT_NE(FMath::Abs(Transform[0].GetRotation().Euler().Y), SMALL_THRESHOLD);
+				EXPECT_GT(Particles.W(0).Y, PreviousY);
+			}
+
+			PreviousY = Particles.W(0).Y;
+
+
+		}
+
+	}
+
+
+
+	
+	TYPED_TEST(AllTraits, GeometryCollection_RigidBodies_Field_Kill)
+	{
+		if(bSkipGeomFailing)
+		{
+			return;
+		}
+		using Traits = TypeParam;
+		TFramework<Traits> UnitTest;
+
+		// Physics Object Setup
+		CreationParameters Params;
+		Params.DynamicState = EObjectStateTypeEnum::Chaos_Object_Dynamic;
+		Params.RootTransform.SetLocation(FVector(0.f, 0.f, 20.f));
+		TGeometryCollectionWrapper<Traits>* Collection = TNewSimulationObject<GeometryType::GeometryCollectionWithSingleRigid>::Init<Traits>(Params)->As<TGeometryCollectionWrapper<Traits>>();
+		UnitTest.AddSimulationObject(Collection);
+
+		// Field setup
+		FPlaneFalloff * FalloffField = new FPlaneFalloff();
+		FalloffField->Magnitude = 1.0;
+		FalloffField->Distance = 10.0f;
+		FalloffField->Position = FVector(0.0, 0.0, 5.0);
+		FalloffField->Normal = FVector(0.0, 0.0, 1.0);
+		FalloffField->Falloff = EFieldFalloffType::Field_Falloff_Linear;
+
+		FFieldSystemPhysicsProxy* FieldObject = new FFieldSystemPhysicsProxy(nullptr);
+		UnitTest.AddFieldObject(FieldObject);
+
+		UnitTest.Initialize();
+
+		TManagedArray<FTransform>& Transform = Collection->DynamicCollection->Transform;
+		TManagedArray<bool>& Active = Collection->DynamicCollection->Active;
+		auto& Particles = UnitTest.Solver->GetParticles().GetGeometryCollectionParticles();
+		for (int Frame = 0; Frame < 20; Frame++)
+		{
+			//UE_LOG(LogTest, Verbose, TEXT("Frame[%d]"), Frame);
+
+			FName TargetName = GetFieldPhysicsName(EFieldPhysicsType::Field_Kill);
+			FieldObject->BufferCommand(UnitTest.Solver, { TargetName, FalloffField->NewCopy() });
+
+			UnitTest.Advance();
+
+			if (Particles.Disabled(0))
+				break;
+		}
+
+		EXPECT_EQ(Particles.Disabled(0), true);
+		//EXPECT_EQ(Active[0], false);
+
+		// hasn't fallen any further than this due to being disabled
+		EXPECT_LT(Transform[0].GetTranslation().Z, 5.f);
+		EXPECT_GT(Transform[0].GetTranslation().Z, -5.0f);
+
+	}
+
+	
+	TYPED_TEST(AllTraits, GeometryCollection_RigidBodies_Field_LinearVelocity)
+	{
+		if(bSkipGeomFailing)
+		{
+			return;
+		}
+		using Traits = TypeParam;
+		TFramework<Traits> UnitTest;
+
+		// Physics Object Setup
+		CreationParameters Params;
+		Params.DynamicState = EObjectStateTypeEnum::Chaos_Object_Dynamic;
+		Params.RootTransform.SetLocation(FVector(0.f, 0.f, 20.f));
+		TGeometryCollectionWrapper<Traits>* Collection = TNewSimulationObject<GeometryType::GeometryCollectionWithSingleRigid>::Init<Traits>(Params)->As<TGeometryCollectionWrapper<Traits>>();
+		UnitTest.AddSimulationObject(Collection);
+
+		// Field setup
+		FUniformVector * VectorField = new FUniformVector();
+		VectorField->Magnitude = 100.0;
+		VectorField->Direction = FVector(1.0, 0.0, 0.0);
+
+		FFieldSystemPhysicsProxy* FieldObject = new FFieldSystemPhysicsProxy(nullptr);
+		UnitTest.AddFieldObject(FieldObject);
+
+		UnitTest.Initialize();
+
+		FName TargetName = GetFieldPhysicsName(EFieldPhysicsType::Field_LinearVelocity);
+		FieldObject->BufferCommand(UnitTest.Solver, { TargetName, VectorField->NewCopy() });
+		UnitTest.Advance();
+
+		float PreviousX = 0.f;
+		TManagedArray<FTransform>& Transform = Collection->DynamicCollection->Transform;
+		for (int Frame = 1; Frame < 10; Frame++)
+		{
+			FieldObject->BufferCommand(UnitTest.Solver, { GetFieldPhysicsName(EFieldPhysicsType::Field_LinearVelocity), VectorField->NewCopy() });
+
+			UnitTest.Advance();
+
+			EXPECT_GT(Transform[0].GetTranslation().X, PreviousX);
+			PreviousX = Transform[0].GetTranslation().X;
+		}
+	}
+
+
+	/**
+	 * Create a stack of boxes on the ground and verify that we we change their collision
+	 * group, they drop through the ground.
+	 */
+	
+	TYPED_TEST(AllTraits, GeometryCollection_RigidBodies_Field_CollisionGroup)
+	{
+		if(bSkipGeomFailing)
+		{
+			return;
+		}
+		using Traits = TypeParam;
+		TFramework<Traits> UnitTest; UnitTest.Dt = 1/24.0;
+
+		RigidBodyWrapper* Floor = TNewSimulationObject<GeometryType::RigidFloor>::Init<Traits>()->As<RigidBodyWrapper>();
+		UnitTest.AddSimulationObject(Floor);
+
+		// Generate Geometry - a stack of boxes.
+		// The bottom box is on the ground, and the others are dropped into it.
+		CreationParameters Params;
+		Params.DynamicState = EObjectStateTypeEnum::Chaos_Object_Dynamic;
+		FVector Scale = FVector(100);
+		Params.GeomTransform.SetScale3D(Scale);
+
+		TGeometryCollectionWrapper<Traits>* Collection[4];
+		for (int n=0; n<3;n++)
+		{
+			Params.RootTransform.SetLocation(FVector(0.f, 0.f, n*200.0f + 100.0f));
+			Params.CollisionType = ECollisionTypeEnum::Chaos_Volumetric;
+			Collection[n+1] = TNewSimulationObject<GeometryType::GeometryCollectionWithSingleRigid>::Init<Traits>(Params)->As<TGeometryCollectionWrapper<Traits>>();
+			UnitTest.AddSimulationObject(Collection[n+1]);
+		}
+
+		// Field setup
+		FRadialIntMask * RadialMask = new FRadialIntMask();
+		RadialMask->Position = FVector(0.0, 0.0, 0.0);
+		RadialMask->Radius = 0;
+		RadialMask->InteriorValue = -1;
+		RadialMask->ExteriorValue = -1;
+		RadialMask->SetMaskCondition = ESetMaskConditionType::Field_Set_Always;
+
+		UnitTest.Initialize();
+
+		for (int Frame = 0; Frame < 60; Frame++)
+		{
+			UnitTest.Advance();
+
+			auto& Particles = UnitTest.Solver->GetParticles().GetGeometryCollectionParticles();
+			if (Frame == 30)
+			{
+				// The boxes should have landed on each other and settled by now
+				EXPECT_NEAR(Collection[1]->DynamicCollection->Transform[0].GetTranslation().Z, (FReal)100, (FReal)20);
+				EXPECT_NEAR(Collection[2]->DynamicCollection->Transform[0].GetTranslation().Z, (FReal)300, (FReal)20);
+				EXPECT_NEAR(Collection[3]->DynamicCollection->Transform[0].GetTranslation().Z, (FReal)500, (FReal)20);
+			}
+			if (Frame == 31)
+			{
+				FName TargetName = GetGeometryCollectionPhysicsTypeName(EGeometryCollectionPhysicsTypeEnum::Chaos_CollisionGroup);
+				Collection[1]->PhysObject->BufferCommand(UnitTest.Solver, { TargetName, RadialMask });
+			}
+		}
+		// The bottom boxes should have fallen below the ground level, box 2 now on the ground with box 3 on top
+		auto& Particles = UnitTest.Solver->GetParticles().GetGeometryCollectionParticles();
+		EXPECT_LT(Collection[1]->DynamicCollection->Transform[0].GetTranslation().Z, 0);
+		EXPECT_TRUE(FMath::IsNearlyEqual(Collection[2]->DynamicCollection->Transform[0].GetTranslation().Z, (FReal)100, (FReal)20));
+		EXPECT_TRUE(FMath::IsNearlyEqual(Collection[3]->DynamicCollection->Transform[0].GetTranslation().Z, (FReal)300, (FReal)20));
+
+	}
+
+	
+	TYPED_TEST(AllTraits, GeometryCollection_RigidBodies_Field_ClusterBreak_StrainModel_Test1)
+	{
+		if(bSkipGeomFailing)
+		{
+			return;
+		}
+		using Traits = TypeParam;
+		TFramework<Traits> UnitTest;		
+
+		TSharedPtr<FGeometryCollection> RestCollection = CreateClusteredBody_TwoByTwo_ThreeTransform(FVector(0));
+
+		CreationParameters Params;
+		Params.RestCollection = RestCollection;
+		Params.DynamicState = EObjectStateTypeEnum::Chaos_Object_Dynamic;
+		Params.CollisionType = ECollisionTypeEnum::Chaos_Surface_Volumetric;
+		Params.ImplicitType = EImplicitTypeEnum::Chaos_Implicit_Box;
+		Params.Simulating = true;
+		Params.EnableClustering = true;
+		Params.DamageThreshold = { 1.0 };
+		Params.MaxClusterLevel = 1000;
+		Params.ClusterConnectionMethod = Chaos::FClusterCreationParameters<FReal>::EConnectionMethod::DelaunayTriangulation;
+		Params.ClusterGroupIndex = 0;
+
+		TGeometryCollectionWrapper<Traits>* Collection = TNewSimulationObject<GeometryType::GeometryCollectionWithSuppliedRestCollection>::Init<Traits>(Params)->As<TGeometryCollectionWrapper<Traits>>();
+
+		UnitTest.AddSimulationObject(Collection);
+
+		//GeometryCollectionAlgo::PrintParentHierarchy(Collection->RestCollection.Get());
+
+		FRadialFalloff * FalloffField = new FRadialFalloff();
+		FalloffField->Magnitude = 1.5;
+		FalloffField->Radius = 100.0;
+		FalloffField->Position = FVector(0.0, 0.0, 0.0);
+		FalloffField->Falloff = EFieldFalloffType::Field_FallOff_None;
+		
+		FFieldSystemPhysicsProxy* FieldObject = new FFieldSystemPhysicsProxy(nullptr);
+		UnitTest.AddFieldObject(FieldObject);
+
+		UnitTest.Initialize();		
+
+		auto& Clustering = UnitTest.Solver->GetEvolution()->GetRigidClustering();
+		const auto& ClusterMap = Clustering.GetChildrenMap();
+
+		UnitTest.Advance();
+		
+		TArray<Chaos::TPBDRigidClusteredParticleHandle<float, 3>*>& ParticleHandles = Collection->PhysObject->GetSolverParticleHandles();
+		{
+			FName TargetName = GetFieldPhysicsName(EFieldPhysicsType::Field_ExternalClusterStrain);
+			FieldObject->BufferCommand(UnitTest.Solver, { TargetName, FalloffField->NewCopy() });
+
+			EXPECT_EQ(ClusterMap.Num(), 3);
+			EXPECT_EQ(ClusterMap[ParticleHandles[4]].Num(), 2);
+			EXPECT_TRUE(ClusterMap[ParticleHandles[4]].Contains(ParticleHandles[0]));
+			EXPECT_TRUE(ClusterMap[ParticleHandles[4]].Contains(ParticleHandles[1]));
+			EXPECT_EQ(ClusterMap[ParticleHandles[5]].Num(), 2);
+			EXPECT_TRUE(ClusterMap[ParticleHandles[5]].Contains(ParticleHandles[2]));
+			EXPECT_TRUE(ClusterMap[ParticleHandles[5]].Contains(ParticleHandles[3]));
+			EXPECT_EQ(ClusterMap[ParticleHandles[6]].Num(), 2);
+			EXPECT_TRUE(ClusterMap[ParticleHandles[6]].Contains(ParticleHandles[5]));
+			EXPECT_TRUE(ClusterMap[ParticleHandles[6]].Contains(ParticleHandles[4]));
+
+			EXPECT_TRUE(ParticleHandles[0]->Disabled());
+			EXPECT_TRUE(ParticleHandles[1]->Disabled());
+			EXPECT_TRUE(ParticleHandles[2]->Disabled());
+			EXPECT_TRUE(ParticleHandles[3]->Disabled());
+			EXPECT_TRUE(ParticleHandles[4]->Disabled());
+			EXPECT_TRUE(ParticleHandles[5]->Disabled());
+			EXPECT_FALSE(ParticleHandles[6]->Disabled());			
+
+			UnitTest.Advance();			
+
+			// todo: indices here might seem odd, particles 4 & 5 are swapped
+			EXPECT_EQ(ClusterMap.Num(), 2);
+			EXPECT_EQ(ClusterMap[ParticleHandles[4]].Num(), 2);
+			EXPECT_TRUE(ClusterMap[ParticleHandles[4]].Contains(ParticleHandles[0]));
+			EXPECT_TRUE(ClusterMap[ParticleHandles[4]].Contains(ParticleHandles[1]));
+			EXPECT_EQ(ClusterMap[ParticleHandles[5]].Num(), 2);
+			EXPECT_TRUE(ClusterMap[ParticleHandles[5]].Contains(ParticleHandles[2]));
+			EXPECT_TRUE(ClusterMap[ParticleHandles[5]].Contains(ParticleHandles[3]));
+
+			EXPECT_TRUE(ParticleHandles[0]->Disabled());
+			EXPECT_TRUE(ParticleHandles[1]->Disabled());
+			EXPECT_TRUE(ParticleHandles[2]->Disabled());
+			EXPECT_TRUE(ParticleHandles[3]->Disabled());
+			EXPECT_FALSE(ParticleHandles[4]->Disabled());
+			EXPECT_FALSE(ParticleHandles[5]->Disabled());
+			EXPECT_TRUE(ParticleHandles[6]->Disabled());
+		}
+
+		delete FalloffField;
+
+	}
+
+
+
+	
+	TYPED_TEST(AllTraits, GeometryCollection_RigidBodies_Field_ClusterBreak_StrainModel_Test2)
+	{
+		if(bSkipGeomFailing)
+		{
+			return;
+		}
+		using Traits = TypeParam;
+		TFramework<Traits> UnitTest;
+
+		TSharedPtr<FGeometryCollection> RestCollection = CreateClusteredBody_ThreeByTwo_ThreeTransform(FVector(0));
+
+		CreationParameters Params;
+		Params.RestCollection = RestCollection;
+		Params.DynamicState = EObjectStateTypeEnum::Chaos_Object_Dynamic;
+		Params.CollisionType = ECollisionTypeEnum::Chaos_Surface_Volumetric;
+		Params.ImplicitType = EImplicitTypeEnum::Chaos_Implicit_Box;
+		Params.Simulating = true;
+		Params.EnableClustering = true;
+		Params.DamageThreshold = { 1.0 };
+		Params.MaxClusterLevel = 1000;		
+		Params.ClusterGroupIndex = 0;
+
+		TGeometryCollectionWrapper<Traits>* Collection = TNewSimulationObject<GeometryType::GeometryCollectionWithSuppliedRestCollection>::Init<Traits>(Params)->As<TGeometryCollectionWrapper<Traits>>();
+
+		UnitTest.AddSimulationObject(Collection);
+
+		FRadialFalloff * FalloffField = new FRadialFalloff();
+		FalloffField->Magnitude = 1.5;
+		FalloffField->Radius = 200.0;
+		FalloffField->Position = FVector(0.0, 0.0, 0.0);
+		FalloffField->Falloff = EFieldFalloffType::Field_FallOff_None;
+
+		FFieldSystemPhysicsProxy* FieldObject = new FFieldSystemPhysicsProxy(nullptr);
+		UnitTest.AddFieldObject(FieldObject);
+
+		UnitTest.Initialize();	
+		UnitTest.Advance();
+
+		TArray<Chaos::TPBDRigidClusteredParticleHandle<float, 3>*>& ParticleHandles = Collection->PhysObject->GetSolverParticleHandles();
+		auto& Clustering = UnitTest.Solver->GetEvolution()->GetRigidClustering();
+		const auto& ClusterMap = Clustering.GetChildrenMap();
+
+		{	
+			FName TargetName = GetFieldPhysicsName(EFieldPhysicsType::Field_ExternalClusterStrain);
+			FFieldSystemCommand Command(TargetName, FalloffField->NewCopy());
+			FFieldSystemMetaDataProcessingResolution* ResolutionData = new FFieldSystemMetaDataProcessingResolution(EFieldResolutionType::Field_Resolution_Maximum);
+			Command.MetaData.Add(FFieldSystemMetaData::EMetaType::ECommandData_ProcessingResolution, TUniquePtr< FFieldSystemMetaDataProcessingResolution >(ResolutionData));
+			FieldObject->BufferCommand(UnitTest.Solver, Command);
+
+			EXPECT_EQ(ClusterMap.Num(), 3);
+			EXPECT_EQ(ClusterMap[ParticleHandles[6]].Num(), 3);
+			EXPECT_TRUE(ClusterMap[ParticleHandles[6]].Contains(ParticleHandles[0]));
+			EXPECT_TRUE(ClusterMap[ParticleHandles[6]].Contains(ParticleHandles[1]));
+			EXPECT_TRUE(ClusterMap[ParticleHandles[6]].Contains(ParticleHandles[2]));
+			EXPECT_EQ(ClusterMap[ParticleHandles[7]].Num(), 3);
+			EXPECT_TRUE(ClusterMap[ParticleHandles[7]].Contains(ParticleHandles[3]));
+			EXPECT_TRUE(ClusterMap[ParticleHandles[7]].Contains(ParticleHandles[4]));
+			EXPECT_TRUE(ClusterMap[ParticleHandles[7]].Contains(ParticleHandles[5]));
+			EXPECT_EQ(ClusterMap[ParticleHandles[8]].Num(), 2);
+			EXPECT_TRUE(ClusterMap[ParticleHandles[8]].Contains(ParticleHandles[7]));
+			EXPECT_TRUE(ClusterMap[ParticleHandles[8]].Contains(ParticleHandles[6]));
+
+			EXPECT_TRUE(ParticleHandles[0]->Disabled());
+			EXPECT_TRUE(ParticleHandles[1]->Disabled());
+			EXPECT_TRUE(ParticleHandles[2]->Disabled());
+			EXPECT_TRUE(ParticleHandles[3]->Disabled());
+			EXPECT_TRUE(ParticleHandles[4]->Disabled());
+			EXPECT_TRUE(ParticleHandles[5]->Disabled());
+			EXPECT_TRUE(ParticleHandles[6]->Disabled());
+			EXPECT_TRUE(ParticleHandles[7]->Disabled());
+			EXPECT_FALSE(ParticleHandles[8]->Disabled());
+
+			UnitTest.Advance();
+			FieldObject->BufferCommand(UnitTest.Solver, Command);
+			UnitTest.Advance();
+
+			EXPECT_EQ(ClusterMap.Num(), 1);
+			EXPECT_EQ(ClusterMap[ParticleHandles[7]].Num(), 3);
+			EXPECT_TRUE(ClusterMap[ParticleHandles[7]].Contains(ParticleHandles[3]));
+			EXPECT_TRUE(ClusterMap[ParticleHandles[7]].Contains(ParticleHandles[4]));
+			EXPECT_TRUE(ClusterMap[ParticleHandles[7]].Contains(ParticleHandles[5]));
+
+			EXPECT_FALSE(ParticleHandles[0]->Disabled());
+			EXPECT_FALSE(ParticleHandles[1]->Disabled());
+			EXPECT_FALSE(ParticleHandles[2]->Disabled());
+			EXPECT_TRUE(ParticleHandles[3]->Disabled());
+			EXPECT_TRUE(ParticleHandles[4]->Disabled());
+			EXPECT_TRUE(ParticleHandles[5]->Disabled());
+			EXPECT_TRUE(ParticleHandles[6]->Disabled());
+			EXPECT_FALSE(ParticleHandles[7]->Disabled());
+			EXPECT_TRUE(ParticleHandles[8]->Disabled());
+		}
+
+		delete FalloffField;
+	}
+
+	
+	TYPED_TEST(AllTraits, GeometryCollection_RigidBodies_Field_ClusterBreak_StrainModel_Test3)
+	{
+		if(bSkipGeomFailing)
+		{
+			return;
+		}
+		using Traits = TypeParam;
+		TFramework<Traits> UnitTest;
+
+		TSharedPtr<FGeometryCollection> RestCollection = CreateClusteredBody_ThreeByTwo_ThreeTransform(FVector(0));
+
+		CreationParameters Params;
+		Params.RestCollection = RestCollection;
+		Params.DynamicState = EObjectStateTypeEnum::Chaos_Object_Dynamic;
+		Params.CollisionType = ECollisionTypeEnum::Chaos_Surface_Volumetric;
+		Params.ImplicitType = EImplicitTypeEnum::Chaos_Implicit_Box;
+		Params.Simulating = true;
+		Params.EnableClustering = true;
+		Params.DamageThreshold = { 1.0 };
+		Params.MaxClusterLevel = 1000;		
+		Params.ClusterGroupIndex = 0;
+
+		TGeometryCollectionWrapper<Traits>* Collection = TNewSimulationObject<GeometryType::GeometryCollectionWithSuppliedRestCollection>::Init<Traits>(Params)->As<TGeometryCollectionWrapper<Traits>>();
+
+		UnitTest.AddSimulationObject(Collection);
+
+		FRadialFalloff * FalloffField = new FRadialFalloff();
+		FalloffField->Magnitude = 1.1;
+		FalloffField->Radius = 200.0;
+		FalloffField->Position = FVector(350.0, 0.0, 0.0);
+		FalloffField->Falloff = EFieldFalloffType::Field_FallOff_None;
+
+		FFieldSystemPhysicsProxy* FieldObject = new FFieldSystemPhysicsProxy(nullptr);
+		UnitTest.AddFieldObject(FieldObject);
+
+		UnitTest.Initialize();
+		UnitTest.Advance();
+
+		auto& Clustering = UnitTest.Solver->GetEvolution()->GetRigidClustering();
+		const auto& ClusterMap = Clustering.GetChildrenMap();
+
+		TArray<Chaos::TPBDRigidClusteredParticleHandle<float, 3>*>& ParticleHandles = Collection->PhysObject->GetSolverParticleHandles();
+
+		UnitTest.Advance();
+
+		{
+			FName TargetName = GetFieldPhysicsName(EFieldPhysicsType::Field_ExternalClusterStrain);
+			FFieldSystemCommand Command(TargetName, FalloffField->NewCopy());
+			FFieldSystemMetaDataProcessingResolution* ResolutionData = new FFieldSystemMetaDataProcessingResolution(EFieldResolutionType::Field_Resolution_Maximum);
+			Command.MetaData.Add(FFieldSystemMetaData::EMetaType::ECommandData_ProcessingResolution, TUniquePtr< FFieldSystemMetaDataProcessingResolution >(ResolutionData));
+			FieldObject->BufferCommand(UnitTest.Solver, Command);
+		
+			EXPECT_TRUE(ParticleHandles[0]->Disabled());
+			EXPECT_TRUE(ParticleHandles[1]->Disabled());
+			EXPECT_TRUE(ParticleHandles[2]->Disabled());
+			EXPECT_TRUE(ParticleHandles[3]->Disabled());
+			EXPECT_TRUE(ParticleHandles[4]->Disabled());
+			EXPECT_TRUE(ParticleHandles[5]->Disabled());
+			EXPECT_TRUE(ParticleHandles[6]->Disabled());
+			EXPECT_TRUE(ParticleHandles[7]->Disabled());
+			EXPECT_FALSE(ParticleHandles[8]->Disabled());
+
+			UnitTest.Advance();
+
+			// todo: indices here might be off but the test crashes before this so we can't validate yet
+			EXPECT_EQ(ClusterMap.Num(), 2);
+			EXPECT_EQ(ClusterMap[ParticleHandles[7]].Num(), 3);
+			EXPECT_TRUE(ClusterMap[ParticleHandles[7]].Contains(ParticleHandles[3]));
+			EXPECT_TRUE(ClusterMap[ParticleHandles[7]].Contains(ParticleHandles[4]));
+			EXPECT_TRUE(ClusterMap[ParticleHandles[7]].Contains(ParticleHandles[5]));
+			EXPECT_EQ(ClusterMap[ParticleHandles[6]].Num(), 3);
+			EXPECT_TRUE(ClusterMap[ParticleHandles[6]].Contains(ParticleHandles[0]));
+			EXPECT_TRUE(ClusterMap[ParticleHandles[6]].Contains(ParticleHandles[1]));
+			EXPECT_TRUE(ClusterMap[ParticleHandles[6]].Contains(ParticleHandles[2]));
+
+			EXPECT_TRUE(ParticleHandles[0]->Disabled());
+			EXPECT_TRUE(ParticleHandles[1]->Disabled());
+			EXPECT_TRUE(ParticleHandles[2]->Disabled());
+			EXPECT_TRUE(ParticleHandles[3]->Disabled());
+			EXPECT_TRUE(ParticleHandles[4]->Disabled());
+			EXPECT_TRUE(ParticleHandles[5]->Disabled());
+			EXPECT_FALSE(ParticleHandles[6]->Disabled());
+			EXPECT_FALSE(ParticleHandles[7]->Disabled());
+			EXPECT_TRUE(ParticleHandles[8]->Disabled());
+		}
+			
+		delete FalloffField;
+
+	}
+
+
+	
+	TYPED_TEST(AllTraits, GeometryCollection_RigidBodies_Field_ClusterBreak_StrainModel_Test4)
+	{
+		if(bSkipGeomFailing)
+		{
+			return;
+		}
+		using Traits = TypeParam;
+		TFramework<Traits> UnitTest;
+
+		TSharedPtr<FGeometryCollection> RestCollection = CreateClusteredBody_TwoByTwo_ThreeTransform(FVector(0));
+
+		CreationParameters Params;
+		Params.RestCollection = RestCollection;
+		Params.DynamicState = EObjectStateTypeEnum::Chaos_Object_Dynamic;
+		Params.CollisionType = ECollisionTypeEnum::Chaos_Surface_Volumetric;
+		Params.ImplicitType = EImplicitTypeEnum::Chaos_Implicit_Box;
+		Params.Simulating = true;
+		Params.EnableClustering = true;
+		Params.DamageThreshold = { 1.0 };
+		Params.MaxClusterLevel = 1000;		
+		Params.ClusterGroupIndex = 0;
+
+		TGeometryCollectionWrapper<Traits>* Collection = TNewSimulationObject<GeometryType::GeometryCollectionWithSuppliedRestCollection>::Init<Traits>(Params)->As<TGeometryCollectionWrapper<Traits>>();
+
+		UnitTest.AddSimulationObject(Collection);
+		
+		FRadialFalloff* FalloffField = new FRadialFalloff();
+		FalloffField->Magnitude = 1.5;
+		FalloffField->Radius = 100.0;
+		FalloffField->Position = FVector(0.0, 0.0, 0.0);
+		FalloffField->Falloff = EFieldFalloffType::Field_FallOff_None;
+
+		FFieldSystemPhysicsProxy* FieldObject = new FFieldSystemPhysicsProxy(nullptr);
+		UnitTest.AddFieldObject(FieldObject);
+
+		UnitTest.Initialize();
+		UnitTest.Advance();
+
+		auto& Clustering = UnitTest.Solver->GetEvolution()->GetRigidClustering();
+		const auto& ClusterMap = Clustering.GetChildrenMap();
+
+		TArray<Chaos::TPBDRigidClusteredParticleHandle<float, 3>*>& ParticleHandles = Collection->PhysObject->GetSolverParticleHandles();
+		{
+			FName TargetName = GetFieldPhysicsName(EFieldPhysicsType::Field_ExternalClusterStrain);
+			FieldObject->BufferCommand(UnitTest.Solver, { TargetName, FalloffField->NewCopy() });
+
+			EXPECT_TRUE(ParticleHandles[0]->Disabled());
+			EXPECT_TRUE(ParticleHandles[1]->Disabled());
+			EXPECT_TRUE(ParticleHandles[2]->Disabled());
+			EXPECT_TRUE(ParticleHandles[3]->Disabled());
+			EXPECT_TRUE(ParticleHandles[4]->Disabled());
+			EXPECT_TRUE(ParticleHandles[5]->Disabled());
+			EXPECT_FALSE(ParticleHandles[6]->Disabled());
+
+			UnitTest.Advance();
+
+			EXPECT_EQ(ClusterMap.Num(), 2);
+			EXPECT_EQ(ClusterMap[ParticleHandles[4]].Num(), 2);
+			EXPECT_TRUE(ClusterMap[ParticleHandles[4]].Contains(ParticleHandles[0]));
+			EXPECT_TRUE(ClusterMap[ParticleHandles[4]].Contains(ParticleHandles[1]));
+			EXPECT_EQ(ClusterMap[ParticleHandles[5]].Num(), 2);
+			EXPECT_TRUE(ClusterMap[ParticleHandles[5]].Contains(ParticleHandles[2]));
+			EXPECT_TRUE(ClusterMap[ParticleHandles[5]].Contains(ParticleHandles[3]));
+			
+			EXPECT_TRUE(ParticleHandles[0]->Disabled());
+			EXPECT_TRUE(ParticleHandles[1]->Disabled());
+			EXPECT_TRUE(ParticleHandles[2]->Disabled());
+			EXPECT_TRUE(ParticleHandles[3]->Disabled());
+			EXPECT_FALSE(ParticleHandles[4]->Disabled());
+			EXPECT_FALSE(ParticleHandles[5]->Disabled());
+			EXPECT_TRUE(ParticleHandles[6]->Disabled());
+		}
+		
+		delete FalloffField;
+	}
+
+}
+
+
diff --git a/Engine/Source/Programs/HeadlessChaos/Private/GeometryCollection/GeometryCollectionTestSimulationSolver.cpp b/Engine/Source/Programs/HeadlessChaos/Private/GeometryCollection/GeometryCollectionTestSimulationSolver.cpp
new file mode 100644
--- /dev/null
+++ b/Engine/Source/Programs/HeadlessChaos/Private/GeometryCollection/GeometryCollectionTestSimulationSolver.cpp
@@ -0,0 +1,422 @@
+// Copyright Epic Games, Inc. All Rights Reserved.
+
+#include "GeometryCollection/GeometryCollectionTestSimulationSolver.h"
+#include "GeometryCollection/GeometryCollectionTestFramework.h"
+#include "GeometryCollection/GeometryCollectionTestUtility.h"
+
+#include "GeometryCollection/GeometryCollection.h"
+#include "GeometryCollection/GeometryCollectionAlgo.h"
+#include "GeometryCollection/GeometryCollectionClusteringUtility.h"
+#include "GeometryCollection/GeometryCollectionUtility.h"
+#include "GeometryCollection/TransformCollection.h"
+#include "UObject/Package.h"
+#include "UObject/UObjectGlobals.h"
+
+
+#include "Field/FieldSystem.h"
+#include "Field/FieldSystemNodes.h"
+#include "GeometryCollectionProxyData.h"
+#include "GeometryCollection/GeometryCollectionSimulationTypes.h"
+#include "PhysicsProxy/GeometryCollectionPhysicsProxy.h"
+#include "PhysicsProxy/PhysicsProxies.h"
+#include "Chaos/ErrorReporter.h"
+#include "Chaos/Defines.h"
+#include "ChaosSolversModule.h"
+#include "PhysicsSolver.h"
+#include "EventsData.h"
+#include "HeadlessChaosTestUtility.h"
+
+#define SMALL_THRESHOLD 1e-4
+
+
+namespace GeometryCollectionTest
+{
+using namespace ChaosTest;
+	
+	TYPED_TEST(AllTraits, GeometryCollection_Solver_AdvanceNoObjects)
+	{
+		using Traits = TypeParam;
+		TFramework<Traits> UnitTest;
+		UnitTest.Initialize();
+		UnitTest.Advance();
+
+		{
+			// just makiung sure we did not crash in this test
+		}
+	}
+
+
+	
+	TYPED_TEST(AllTraits, GeometryCollection_Solver_AdvanceDisabledObjects)
+	{
+		using Traits = TypeParam;
+		CreationParameters Params; Params.Simulating = false;
+		TGeometryCollectionWrapper<Traits>* Collection = TNewSimulationObject<GeometryType::GeometryCollectionWithSingleRigid>::Init<Traits>(Params)->As<TGeometryCollectionWrapper<Traits>>();
+
+		TFramework<Traits> UnitTest;
+		UnitTest.AddSimulationObject(Collection);
+		UnitTest.Initialize();
+		UnitTest.Advance();
+
+		{
+			// never touched
+			TManagedArray<FTransform>& RestTransform = Collection->RestCollection->Transform;
+			EXPECT_LT(FMath::Abs(RestTransform[0].GetTranslation().Z), SMALL_THRESHOLD);
+
+			// simulated
+			TManagedArray<FTransform>& Transform = Collection->DynamicCollection->Transform;
+			EXPECT_EQ(Transform.Num(), 1);
+			EXPECT_LT(FMath::Abs(Transform[0].GetTranslation().Z), SMALL_THRESHOLD);
+		}
+	}
+
+	
+	TYPED_TEST(AllTraits, GeometryCollection_Solver_AdvanceDisabledClusteredObjects)
+	{
+		using Traits = TypeParam;
+		TFramework<Traits> UnitTest;
+
+		RigidBodyWrapper* Floor = TNewSimulationObject<GeometryType::RigidFloor>::Init<Traits>()->As<RigidBodyWrapper>();
+		UnitTest.AddSimulationObject(Floor);
+
+		TSharedPtr<FGeometryCollection> RestCollection = GeometryCollection::MakeCubeElement(FTransform(FQuat::MakeFromEuler(FVector(0, 0, 0.)), FVector(0, -10, 10)), FVector(1.0));
+		RestCollection->AppendGeometry(*GeometryCollection::MakeCubeElement(FTransform(FQuat::MakeFromEuler(FVector(0, 0, 0.)), FVector(0, 10, 10)), FVector(1.0)));
+		EXPECT_EQ(RestCollection->Transform.Num(), 2);
+
+		FGeometryCollectionClusteringUtility::ClusterAllBonesUnderNewRoot(RestCollection.Get());
+		EXPECT_EQ(RestCollection->Transform.Num(), 3);
+		RestCollection->Transform[2] = FTransform(FQuat::MakeFromEuler(FVector(90.f, 0, 0.)), FVector(0, 0, 40));
+
+		//GeometryCollectionAlgo::PrintParentHierarchy(RestCollection.Get());
+
+		CreationParameters Params;
+		Params.RestCollection = RestCollection;
+		Params.DynamicState = EObjectStateTypeEnum::Chaos_Object_Dynamic;
+		Params.CollisionType = ECollisionTypeEnum::Chaos_Surface_Volumetric;
+		Params.Simulating = true;
+		Params.EnableClustering = true;
+		Params.DamageThreshold = { 1000.f };
+		TGeometryCollectionWrapper<Traits>* Collection = TNewSimulationObject<GeometryType::GeometryCollectionWithSuppliedRestCollection>::Init<Traits>(Params)->As<TGeometryCollectionWrapper<Traits>>();					
+
+		UnitTest.AddSimulationObject(Collection);
+		UnitTest.Initialize();
+
+		TManagedArray<int32>& SimulationType = Collection->DynamicCollection->SimulationType;
+		EXPECT_EQ(SimulationType[0], FGeometryCollection::ESimulationTypes::FST_Clustered);
+		EXPECT_EQ(SimulationType[1], FGeometryCollection::ESimulationTypes::FST_Clustered);
+		EXPECT_EQ(SimulationType[2], FGeometryCollection::ESimulationTypes::FST_Rigid);
+
+		TManagedArray<int32>& Parent = Collection->DynamicCollection->Parent;
+		EXPECT_EQ(Parent[0], 2);
+		EXPECT_EQ(Parent[1], 2);
+		EXPECT_EQ(Parent[2], -1);
+
+		// Set the one cluster to disabled
+		Collection->PhysObject->GetSolverClusterHandles()[0]->SetDisabled(true);
+
+		TManagedArray<FTransform>& Transform = Collection->DynamicCollection->Transform;
+		float StartingRigidDistance = (Transform[1].GetTranslation() - Transform[0].GetTranslation()).Size();
+		EXPECT_LT(StartingRigidDistance - 20.0f, SMALL_THRESHOLD);
+
+		FVector StartingClusterPosition = Transform[2].GetTranslation();
+
+		float CurrentRigidDistance = 0.f;
+
+		for (int Frame = 0; Frame < 10; Frame++)
+		{
+			UnitTest.Advance();
+
+			// Distance between gc cubes remains the same
+			CurrentRigidDistance = (Transform[1].GetTranslation() - Transform[0].GetTranslation()).Size();
+			EXPECT_LT(StartingRigidDistance-CurrentRigidDistance, SMALL_THRESHOLD);
+
+			// Clustered particle doesn't move
+			EXPECT_LT((StartingClusterPosition - Transform[2].GetTranslation()).Size(), SMALL_THRESHOLD);
+		}
+		
+	}
+
+
+	
+	TYPED_TEST(AllTraits, DISABLED_GeometryCollection_Solver_ValidateReverseMapping)
+	{
+		using Traits = TypeParam;
+		// Todo: this test is missing code for InitCollections that seems to be no longer in the code base.  
+		// This makes it unclear how this test is setup
+
+		/*
+		TUniquePtr<Chaos::FChaosPhysicsMaterial> PhysicalMaterial = MakeUnique<Chaos::FChaosPhysicsMaterial>();
+		InitMaterialToZero(PhysicalMaterial.Get());
+
+		Chaos::FPBDRigidsSolver* Solver = FChaosSolversModule::GetModule()->CreateSolver(nullptr, ESolverFlags::Standalone);
+		//Solver->SetHasFloor(true);
+		Solver->SetEnabled(true);
+		TArray<TSharedPtr<FGeometryCollection> > RestArray;
+		TArray<TSharedPtr<FGeometryDynamicCollection> > DynamicArray;
+
+		for (int32 i = 0; i < 10; i++)
+		{
+			TSharedPtr<FGeometryCollection> RestCollection = nullptr;
+			TSharedPtr<FGeometryDynamicCollection> DynamicCollection = nullptr;
+
+			//InitCollectionsParameters InitParams = { FTransform::Identity, FVector(1.0), nullptr, (int32)EObjectStateTypeEnum::Chaos_Object_Kinematic };
+			//InitCollections(PhysicalMaterial, RestCollection, DynamicCollection, InitParams);
+
+			FGeometryCollectionPhysicsProxy* PhysObject = RigidBodySetup(PhysicalMaterial, RestCollection, DynamicCollection);
+
+#if CHAOS_PARTICLEHANDLE_TODO
+			Solver->RegisterObject(PhysObject);
+#endif
+			//PhysObject->ActivateBodies();
+
+			RestArray.Add(RestCollection);
+			DynamicArray.Add(DynamicCollection);
+		}
+
+
+		Solver->AdvanceSolverBy(1 / 24.);
+
+#if TODO_REIMPLEMENT_PHYSICS_PROXY_REVERSE_MAPPING
+		const Chaos::TArrayCollectionArray<PhysicsProxyWrapper> & SolverMapping = Solver->GetPhysicsProxyReverseMapping();
+		const Chaos::TArrayCollectionArray<int32> & ParticleMapping = Solver->GetParticleIndexReverseMapping();
+
+		EXPECT_EQ(SolverMapping.Num(), 11);
+		EXPECT_EQ(ParticleMapping.Num(), 11);
+
+		EXPECT_EQ(ParticleMapping[0], INDEX_NONE);
+		EXPECT_EQ(ParticleMapping[1], 0);
+
+		EXPECT_EQ(SolverMapping[0].PhysicsProxy, nullptr);
+		EXPECT_EQ(SolverMapping[0].Type, EPhysicsProxyType::NoneType);
+
+		EXPECT_EQ(SolverMapping[5].PhysicsProxy != nullptr);
+		EXPECT_EQ(SolverMapping[5].Type, EPhysicsProxyType::GeometryCollectionType);
+
+		const TManagedArray<int32> & RigidBodyID = static_cast<FGeometryCollectionPhysicsProxy*>(SolverMapping[5].PhysicsProxy)->
+			GetGeometryDynamicCollection_PhysicsThread()->GetAttribute<int32>("RigidBodyID", FGeometryCollection::TransformGroup);
+		EXPECT_EQ(RigidBodyID.Num(), 1);
+		EXPECT_EQ(RigidBodyID[0], 5);
+#endif
+
+		FChaosSolversModule::GetModule()->DestroySolver(Solver);
+		*/
+	}
+	
+	TSharedPtr<FGeometryCollection> CommonInit(int32 NumObjects, bool UseClusters)
+	{	
+		TSharedPtr<FGeometryCollection> RestCollectionOut;
+
+		for (int32 i = 0; i < NumObjects; i++)
+		{
+
+			TSharedPtr<FGeometryCollection> RestCollection;
+
+			if (UseClusters)
+			{
+				RestCollection = GeometryCollection::MakeCubeElement(FTransform(FQuat::MakeFromEuler(FVector(0, 0, 0.)), FVector(0, -10, 10)), FVector(1.0));
+				RestCollection->AppendGeometry(*GeometryCollection::MakeCubeElement(FTransform(FQuat::MakeFromEuler(FVector(0, 0, 0.)), FVector(0, 10, 10)), FVector(1.0)));
+				RestCollection->AppendGeometry(*GeometryCollection::MakeCubeElement(FTransform(FQuat::MakeFromEuler(FVector(0, 0, 0.)), FVector(0, 0, 0)), FVector(1.0)));
+				RestCollection->AppendGeometry(*GeometryCollection::MakeCubeElement(FTransform(FQuat::MakeFromEuler(FVector(0, 0, 0.)), FVector(0, 0, 10)), FVector(1.0)));
+				FGeometryCollectionClusteringUtility::ClusterAllBonesUnderNewRoot(RestCollection.Get());
+				RestCollection->Transform[4] = FTransform(FQuat::MakeFromEuler(FVector(90.f, 0, 0.)), FVector(0, 0, 40));
+			}
+			else
+			{
+				RestCollection = GeometryCollection::MakeCubeElement(FTransform(FQuat::MakeFromEuler(FVector(0, 0, 0.)), FVector(i * 3, 0, 5)), FVector(1.0));
+			}
+						
+
+			if (i == 0)
+			{
+				RestCollectionOut = RestCollection;
+			}
+			else
+			{
+				RestCollectionOut->AppendGeometry(*RestCollection);			
+			}
+		}		
+
+		return RestCollectionOut;
+	}
+
+	template <typename Traits>
+	class TEventHarvester
+	{
+	public:
+		TEventHarvester(Chaos::TPBDRigidsSolver<Traits>* Solver)
+		{
+			Solver->GetEventManager()->RegisterHandler<Chaos::FCollisionEventData>(Chaos::EEventType::Collision, this, &TEventHarvester<Traits>::HandleCollisionEvents);
+			Solver->GetEventManager()->RegisterHandler<Chaos::FBreakingEventData>(Chaos::EEventType::Breaking, this, &TEventHarvester<Traits>::HandleBreakingEvents);
+		}
+
+		void HandleCollisionEvents(const Chaos::FCollisionEventData& Events)
+		{
+			CollisionEventData = Events;
+		}
+
+		void HandleBreakingEvents(const Chaos::FBreakingEventData& Events)
+		{
+			BreakingEventData = Events;
+		}
+
+		Chaos::FCollisionEventData CollisionEventData;
+		Chaos::FBreakingEventData BreakingEventData;
+	};
+
+	
+	TYPED_TEST(AllTraits, GeometryCollection_Solver_CollisionEventFilter)
+	{
+		using Traits = TypeParam;
+		float TestMassThreshold = 7.0f;
+
+		TGeometryCollectionWrapper<Traits>* Collection[10];
+		for (int i=0; i<10; i++)
+		{
+			TSharedPtr<FGeometryCollection> RestCollection = GeometryCollection::MakeCubeElement(FTransform(FQuat::MakeFromEuler(FVector(0, 0, 0)), FVector(0,0,0)), FVector(1.0));
+
+			CreationParameters Params;
+			Params.RestCollection = RestCollection;
+			Params.SimplicialType = ESimplicialType::Chaos_Simplicial_Box;
+			Params.ImplicitType = EImplicitTypeEnum::Chaos_Implicit_Box;
+			Params.DynamicState = EObjectStateTypeEnum::Chaos_Object_Dynamic;
+			Params.RootTransform.SetLocation(FVector(i * 10.0f, 0.f, 10.f));
+			Params.Simulating = true;
+			
+			Collection[i] = TNewSimulationObject<GeometryType::GeometryCollectionWithSingleRigid>::Init<Traits>(Params)->As<TGeometryCollectionWrapper<Traits>>();
+		}
+		RigidBodyWrapper* Floor = TNewSimulationObject<GeometryType::RigidFloor>::Init<Traits>()->As<RigidBodyWrapper>();
+
+		TFramework<Traits> UnitTest;
+		for (int i=0; i<10; i++)
+		{
+			TManagedArray<float>& Mass = Collection[i]->RestCollection->GetAttribute<float>("Mass", FTransformCollection::TransformGroup);
+			Mass[0] = i + 1;
+			UnitTest.AddSimulationObject(Collection[i]);
+		}
+		UnitTest.AddSimulationObject(Floor);
+		UnitTest.Initialize();
+
+		// setup collision filter
+		FSolverCollisionFilterSettings CollisionFilterSettings;
+		CollisionFilterSettings.FilterEnabled = true;
+		CollisionFilterSettings.MinImpulse = 0;
+		CollisionFilterSettings.MinMass = TestMassThreshold;
+		CollisionFilterSettings.MinSpeed = 0;
+
+		UnitTest.Solver->SetGenerateCollisionData(true);
+		UnitTest.Solver->SetCollisionFilterSettings(CollisionFilterSettings);
+		TEventHarvester<Traits> Events(UnitTest.Solver);
+
+		bool Impact = false;
+		for(int LoopCount=0; LoopCount<10; LoopCount++)
+		{
+			// Events data on physics thread is appended until the game thread has had a chance to Tick & read it
+			UnitTest.Advance();
+
+			const auto& AllCollisionsArray = Events.CollisionEventData.CollisionData.AllCollisionsArray;
+			Impact = (AllCollisionsArray.Num() > 0);
+
+			if (Impact)
+			{
+				// any objects with a mass of less than 6 are removed from returned collision data
+				EXPECT_EQ(AllCollisionsArray.Num(), 4);
+
+				for (const auto& Collision : AllCollisionsArray)
+				{ 
+					EXPECT_GE(Collision.Mass1, TestMassThreshold);
+					EXPECT_EQ(Collision.Mass2, 0.0f);
+					EXPECT_LT(Collision.Velocity2.Z, SMALL_NUMBER);
+				}
+				break;
+			}
+		} 
+
+		EXPECT_EQ(Impact, true);
+	}
+
+	
+	TYPED_TEST(AllTraits, GeometryCollection_Solver_BreakingEventFilter)
+	{
+		using Traits = TypeParam;
+		TFramework<Traits> UnitTest;
+		
+		TSharedPtr<FGeometryCollection> RestCollection = CreateClusteredBody_OnePartent_FourBodies(FVector(0));
+
+		CreationParameters Params;
+		Params.RootTransform.SetLocation(FVector(0, 0, 20));
+		Params.RestCollection = RestCollection;
+		Params.DynamicState = EObjectStateTypeEnum::Chaos_Object_Dynamic;
+		Params.CollisionType = ECollisionTypeEnum::Chaos_Surface_Volumetric;
+		Params.ImplicitType = EImplicitTypeEnum::Chaos_Implicit_LevelSet;
+		Params.Simulating = true;
+		Params.EnableClustering = true;
+		Params.DamageThreshold = { 1.0 };
+		Params.MaxClusterLevel = 1000;
+		Params.ClusterConnectionMethod = Chaos::FClusterCreationParameters<FReal>::EConnectionMethod::DelaunayTriangulation;
+		Params.ClusterGroupIndex = 0;
+
+		RigidBodyWrapper* Floor = TNewSimulationObject<GeometryType::RigidFloor>::Init<Traits>()->As<RigidBodyWrapper>();
+		UnitTest.AddSimulationObject(Floor);
+
+		TGeometryCollectionWrapper<Traits>* Collection = TNewSimulationObject<GeometryType::GeometryCollectionWithSuppliedRestCollection>::Init<Traits>(Params)->As<TGeometryCollectionWrapper<Traits>>();
+		UnitTest.AddSimulationObject(Collection);
+		Collection->PhysObject->SetCollisionParticlesPerObjectFraction(1.0);		
+
+		UnitTest.Initialize();
+
+	//	UnitTest.Advance();
+
+		float TestMass = 7.0f;
+
+		// setup breaking filter
+		FSolverBreakingFilterSettings BreakingFilterSettings;
+		BreakingFilterSettings.FilterEnabled = true;
+		BreakingFilterSettings.MinMass = TestMass;
+		BreakingFilterSettings.MinSpeed = 0;
+		BreakingFilterSettings.MinVolume = 0;
+
+		UnitTest.Solver->SetGenerateBreakingData(true);
+		UnitTest.Solver->SetBreakingFilterSettings(BreakingFilterSettings);
+
+		TEventHarvester<Traits> Events(UnitTest.Solver);
+
+		TArray<Chaos::TPBDRigidClusteredParticleHandle<float, 3>*>& ParticleHandles = Collection->PhysObject->GetSolverParticleHandles();
+		
+		ParticleHandles[0]->SetM(TestMass + 1.0f);
+		ParticleHandles[1]->SetM(TestMass - 1.0f);
+		ParticleHandles[2]->SetM(TestMass - 2.0f);
+		ParticleHandles[3]->SetM(TestMass + 2.0f);
+	
+		bool Impact = false;
+		for (int Loop=0; Loop<50; Loop++)
+		{
+			// Events data on physics thread is appended until the game thread has had a chance to Tick & read it
+			UnitTest.Advance();
+
+			const auto& AllBreakingsArray = Events.BreakingEventData.BreakingData.AllBreakingsArray;
+			Impact = (AllBreakingsArray.Num() > 0);
+
+			if (Impact)
+			{
+				EXPECT_EQ(ParticleHandles[0]->Disabled(), false); // piece1 active 6 mass
+				EXPECT_EQ(ParticleHandles[1]->Disabled(), false); // piece2 active 0.5 mass
+				EXPECT_EQ(ParticleHandles[2]->Disabled(), false); // piece3 active 0.5 mass
+				EXPECT_EQ(ParticleHandles[3]->Disabled(), false); // cluster active 7 mass
+				EXPECT_EQ(ParticleHandles[4]->Disabled(), true); // cluster
+
+				// breaking data
+				EXPECT_EQ(AllBreakingsArray.Num(), 2); // 2 pieces filtered out of 4
+
+				EXPECT_LT(AllBreakingsArray[0].Mass - (TestMass + 2.0f), SMALL_THRESHOLD);
+				EXPECT_LT(AllBreakingsArray[1].Mass - (TestMass + 1.0f), SMALL_THRESHOLD);
+				break;
+			}
+
+
+		}
+	}
+}
+
+
+
+
diff --git a/Engine/Source/Programs/HeadlessChaos/Private/GeometryCollection/GeometryCollectionTestSimulationStreaming.cpp b/Engine/Source/Programs/HeadlessChaos/Private/GeometryCollection/GeometryCollectionTestSimulationStreaming.cpp
new file mode 100644
--- /dev/null
+++ b/Engine/Source/Programs/HeadlessChaos/Private/GeometryCollection/GeometryCollectionTestSimulationStreaming.cpp
@@ -0,0 +1,216 @@
+// Copyright Epic Games, Inc. All Rights Reserved.
+
+#include "GeometryCollection/GeometryCollectionTestSimulationStreaming.h"
+#include "GeometryCollection/GeometryCollectionTestFramework.h"
+#include "GeometryCollection/GeometryCollectionTestUtility.h"
+
+#include "GeometryCollection/GeometryCollection.h"
+#include "GeometryCollection/GeometryCollectionUtility.h"
+#include "GeometryCollection/TransformCollection.h"
+#include "UObject/Package.h"
+#include "UObject/UObjectGlobals.h"
+
+#include "Field/FieldSystem.h"
+#include "Field/FieldSystemNodes.h"
+#include "GeometryCollectionProxyData.h"
+#include "GeometryCollection/GeometryCollectionSimulationTypes.h"
+#include "PhysicsProxy/PhysicsProxies.h"
+#include "Chaos/ErrorReporter.h"
+#include "ChaosSolversModule.h"
+#include "PBDRigidsSolver.h"
+
+#define SMALL_THRESHOLD 1e-4
+
+namespace GeometryCollectionTest
+{
+	template<typename Traits>
+	void RigidBodies_Streaming_StartSolverEmpty()
+	{
+		TFramework<Traits> UnitTest;
+		
+		// no floor
+		UnitTest.Initialize();
+
+		UnitTest.Advance();
+								
+		for (int32 Frame = 1; Frame < 1000; Frame++)
+		{
+			UnitTest.Advance();
+			if (Frame % 100 == 0)
+			{
+				TSharedPtr<FGeometryCollection> RestCollection = GeometryCollection::MakeCubeElement(FTransform(FQuat::MakeFromEuler(FVector(0, 0, 0.)), FVector(0, 0, 0)), FVector(1.0));
+
+				CreationParameters Params;
+				Params.RestCollection = RestCollection;
+				Params.DynamicState = EObjectStateTypeEnum::Chaos_Object_Dynamic;
+				Params.CollisionType = ECollisionTypeEnum::Chaos_Surface_Volumetric;
+				Params.ImplicitType = EImplicitTypeEnum::Chaos_Implicit_Box;
+				Params.Simulating = true;				
+
+				TGeometryCollectionWrapper<Traits>* Collection = TNewSimulationObject<GeometryType::GeometryCollectionWithSuppliedRestCollection>::Init<Traits>(Params)->As<TGeometryCollectionWrapper<Traits>>();
+				UnitTest.AddSimulationObject(Collection);
+							
+				// todo: this is crashing
+				UnitTest.Solver->RegisterObject(Collection->PhysObject);
+				UnitTest.Solver->AddDirtyProxy(Collection->PhysObject);
+				Collection->PhysObject->Initialize();
+			}
+		}
+		
+		EXPECT_EQ(UnitTest.PhysicsObjects.Num(), 9);
+		for (int32 i = 0; i < UnitTest.PhysicsObjects.Num() - 1; i++)
+		{
+			TGeometryCollectionWrapper<Traits>* GCW = UnitTest.PhysicsObjects[i]->As<TGeometryCollectionWrapper<Traits>>();
+			TGeometryCollectionWrapper<Traits>* GCW2 = UnitTest.PhysicsObjects[i + 1]->As<TGeometryCollectionWrapper<Traits>>();
+			EXPECT_EQ(GCW->PhysObject->GetSolverParticleHandles().Num(), 1);
+			EXPECT_EQ(GCW2->PhysObject->GetSolverParticleHandles().Num(), 1);
+			EXPECT_LT(GCW->PhysObject->GetSolverParticleHandles()[0]->X().Z, GCW2->PhysObject->GetSolverParticleHandles()[0]->X().Z);
+		}
+	}
+
+
+	template<typename Traits>
+	void RigidBodies_Streaming_BulkInitialization()
+	{
+		TFramework<Traits> UnitTest;
+
+		// no floor
+		UnitTest.Initialize();
+
+		UnitTest.Advance();	
+		
+		for (int32 Frame = 1; Frame < 1000; Frame++)
+		{
+			UnitTest.Advance();
+			if (Frame % 100 == 0)
+			{
+				TSharedPtr<FGeometryCollection> RestCollection = GeometryCollection::MakeCubeElement(FTransform(FQuat::MakeFromEuler(FVector(0, 0, 0.)), FVector(0, 0, 0)), FVector(1.0));
+
+				CreationParameters Params;
+				Params.RestCollection = RestCollection;
+				Params.DynamicState = EObjectStateTypeEnum::Chaos_Object_Dynamic;
+				Params.CollisionType = ECollisionTypeEnum::Chaos_Surface_Volumetric;
+				Params.ImplicitType = EImplicitTypeEnum::Chaos_Implicit_Box;
+				Params.Simulating = true;
+				Params.CollisionGroup = -1;
+
+				TGeometryCollectionWrapper<Traits>* Collection = TNewSimulationObject<GeometryType::GeometryCollectionWithSuppliedRestCollection>::Init<Traits>(Params)->As<TGeometryCollectionWrapper<Traits>>();
+				UnitTest.AddSimulationObject(Collection);
+
+				// todo: this is crashing
+				UnitTest.Solver->RegisterObject(Collection->PhysObject);
+				UnitTest.Solver->AddDirtyProxy(Collection->PhysObject);
+				Collection->PhysObject->Initialize();
+			}
+		}
+		
+		/*
+		todo: what is the replacement for ActivateBodies?
+		for (auto Obj : Collections)
+			Obj->PhysicsProxy->ActivateBodies();
+		*/
+
+		for (int32 Frame = 1; Frame < 100; Frame++)
+		{
+			UnitTest.Advance();
+		}
+
+		EXPECT_EQ(UnitTest.PhysicsObjects.Num(), 9);
+		for (int32 i = 0; i < UnitTest.PhysicsObjects.Num() - 1; i++)
+		{
+			TGeometryCollectionWrapper<Traits>* GCW = UnitTest.PhysicsObjects[i]->As<TGeometryCollectionWrapper<Traits>>();
+			TGeometryCollectionWrapper<Traits>* GCW2 = UnitTest.PhysicsObjects[i + 1]->As<TGeometryCollectionWrapper<Traits>>();
+			EXPECT_EQ(GCW->PhysObject->GetSolverParticleHandles().Num(), 1);
+			EXPECT_EQ(GCW2->PhysObject->GetSolverParticleHandles().Num(), 1);
+			EXPECT_LT(FMath::Abs(GCW->PhysObject->GetSolverParticleHandles()[0]->X().Z - GCW2->PhysObject->GetSolverParticleHandles()[0]->X().Z), KINDA_SMALL_NUMBER);
+		}
+	}
+
+
+
+	template<typename Traits>
+	void RigidBodies_Streaming_DeferedClusteringInitialization()
+	{
+		TFramework<Traits> UnitTest;
+
+		// no floor
+		UnitTest.Initialize();
+
+		UnitTest.Advance();
+		
+		for (int32 Frame = 1; Frame < 1000; Frame++)
+		{
+			UnitTest.Advance();
+
+			if (Frame % 100 == 0)
+			{
+				TSharedPtr<FGeometryCollection> RestCollection = GeometryCollection::MakeCubeElement(FTransform(FQuat::MakeFromEuler(FVector(0, 0, 0.)), FVector(0, 0, 0)), FVector(1.0));
+
+				CreationParameters Params;
+				Params.RestCollection = RestCollection;
+				Params.DynamicState = EObjectStateTypeEnum::Chaos_Object_Dynamic;
+				Params.CollisionType = ECollisionTypeEnum::Chaos_Surface_Volumetric;
+				Params.ImplicitType = EImplicitTypeEnum::Chaos_Implicit_Box;
+				Params.Simulating = true;
+				Params.CollisionGroup = -1;
+				Params.EnableClustering = true;
+				Params.ClusterGroupIndex = 1;
+
+				TGeometryCollectionWrapper<Traits>* Collection = TNewSimulationObject<GeometryType::GeometryCollectionWithSuppliedRestCollection>::Init<Traits>(Params)->As<TGeometryCollectionWrapper<Traits>>();
+				UnitTest.AddSimulationObject(Collection);
+
+				// todo: this is crashing
+				UnitTest.Solver->RegisterObject(Collection->PhysObject);
+				UnitTest.Solver->AddDirtyProxy(Collection->PhysObject);
+				Collection->PhysObject->Initialize();
+			}
+		}
+	
+		// all particles should be disabled
+		EXPECT_EQ(UnitTest.PhysicsObjects.Num(), 9);
+		for (int32 i = 0; i < UnitTest.PhysicsObjects.Num(); i++)
+		{
+			TGeometryCollectionWrapper<Traits>* GCW = UnitTest.PhysicsObjects[i]->As<TGeometryCollectionWrapper<Traits>>();
+			EXPECT_EQ(GCW->PhysObject->GetSolverParticleHandles().Num(), 1);
+			EXPECT_EQ(GCW->PhysObject->GetSolverParticleHandles()[0]->Disabled(), true);
+		}
+
+		/*
+		todo: what is the replacement for ActivateBodies
+		for (auto Obj : Collections)
+			Obj->PhysicsProxy->ActivateBodies();
+		*/
+
+		// all particles should be enabled
+		EXPECT_EQ(UnitTest.PhysicsObjects.Num(), 9);
+		for (int32 i = 0; i < UnitTest.PhysicsObjects.Num(); i++)
+		{
+			TGeometryCollectionWrapper<Traits>* GCW = UnitTest.PhysicsObjects[i]->As<TGeometryCollectionWrapper<Traits>>();
+			EXPECT_EQ(GCW->PhysObject->GetSolverParticleHandles().Num(), 1);
+			EXPECT_EQ(GCW->PhysObject->GetSolverParticleHandles()[0]->Disabled(), false);
+		}
+
+		for (int32 Frame = 1; Frame < 100; Frame++)
+		{
+			UnitTest.Advance();			
+		}
+
+		// new cluster parent should be falling
+		// todo: the logic here is probably off.  How is the new cluster made? Is this at odds with each particle being in a separate GC
+		// Difficult to tell based on missing code in initialization
+		EXPECT_EQ(UnitTest.PhysicsObjects.Num(), 10);
+		for (int32 i = 0; i < UnitTest.PhysicsObjects.Num() - 1; i++)
+		{
+			TGeometryCollectionWrapper<Traits>* GCW = UnitTest.PhysicsObjects[i]->As<TGeometryCollectionWrapper<Traits>>();
+			EXPECT_EQ(GCW->PhysObject->GetSolverParticleHandles().Num(), 1);
+			EXPECT_EQ(GCW->PhysObject->GetSolverParticleHandles()[0]->Disabled(), true);
+		}
+		TGeometryCollectionWrapper<Traits>* GCW = UnitTest.PhysicsObjects[UnitTest.PhysicsObjects.Num() - 1]->As<TGeometryCollectionWrapper<Traits>>();
+		EXPECT_EQ(GCW->PhysObject->GetSolverParticleHandles()[0]->Disabled(), false);
+
+		EXPECT_LT(GCW->PhysObject->GetSolverParticleHandles()[0]->X().Z, -1.f);					
+	}
+
+
+}
+
diff --git a/Engine/Source/Programs/HeadlessChaos/Private/GeometryCollection/GeometryCollectionTestSkeletalMeshPhysicsProxy.cpp b/Engine/Source/Programs/HeadlessChaos/Private/GeometryCollection/GeometryCollectionTestSkeletalMeshPhysicsProxy.cpp
new file mode 100644
--- /dev/null
+++ b/Engine/Source/Programs/HeadlessChaos/Private/GeometryCollection/GeometryCollectionTestSkeletalMeshPhysicsProxy.cpp
@@ -0,0 +1,334 @@
+// Copyright Epic Games, Inc. All Rights Reserved.
+
+#include "GeometryCollection/GeometryCollectionTestSkeletalMeshPhysicsProxy.h"
+
+#include "Chaos/Sphere.h"
+#include "PhysicsSolver.h"
+#include "PhysicsProxy/PhysicsProxies.h"
+#include "ChaosSolversModule.h"
+#include "Chaos/ParticleHandle.h"
+#include "GeometryCollection/GeometryCollection.h"
+#include "GeometryCollectionProxyData.h"
+#include "GeometryCollection/GeometryCollectionTestFramework.h"
+#include "GeometryCollection/GeometryCollectionTestUtility.h"
+
+#include "gtest/gtest.h"
+
+
+namespace GeometryCollectionTest
+{
+
+	using namespace Chaos;
+
+	/**
+	 * Create an empty SkeletalMeshPhysicsProxy and register it with a solver.
+	 * Check that the appropriate callbacks are called and nothing crashes.
+	 */
+	template<typename Traits>
+	void TestSkeletalMeshPhysicsProxy_Register()
+	{
+		const float Dt = (FReal)1 / 60;
+
+		FrameworkParameters Fp;
+		Fp.Dt = Dt;
+		TFramework<Traits> UnitTest(Fp);
+
+		int32 InitCallCount = 0;
+		int32 InputCallCount = 0;
+		auto InitFunc = [&](FSkeletalMeshPhysicsProxyParams & OutPhysicsParams) -> void
+		{
+			++InitCallCount;
+		};
+		auto InputFunc = [&](const float Dt, FSkeletalMeshPhysicsProxyParams& OutPhysicsParams) -> bool
+		{
+			++InputCallCount;
+			return true;
+		};
+
+		FSkeletalMeshPhysicsProxy SkeletalMeshPhysicsProxy(nullptr, InitFunc);		
+
+		EXPECT_EQ(InitCallCount, 0);
+		EXPECT_EQ(InputCallCount, 0);
+		EXPECT_EQ(SkeletalMeshPhysicsProxy.GetOutputs(), nullptr);
+
+		// todo
+		// UnitTest.Solver->RegisterObject(&SkeletalMeshPhysicsProxy);
+
+
+		// Should call Init func
+		SkeletalMeshPhysicsProxy.Initialize();
+		EXPECT_EQ(InitCallCount, 1);
+		EXPECT_EQ(InputCallCount, 0);
+		EXPECT_EQ(SkeletalMeshPhysicsProxy.GetOutputs(), nullptr);
+
+		UnitTest.Initialize();
+		{
+			// Should call Input func
+			SkeletalMeshPhysicsProxy.CaptureInputs(Dt, InputFunc);
+			EXPECT_EQ(InitCallCount, 1);
+			EXPECT_EQ(InputCallCount, 1);
+			EXPECT_EQ(SkeletalMeshPhysicsProxy.GetOutputs(), nullptr);
+
+			UnitTest.Advance();			
+
+			// Should give us some outputs
+			SkeletalMeshPhysicsProxy.BufferPhysicsResults();
+			SkeletalMeshPhysicsProxy.FlipBuffer();
+			SkeletalMeshPhysicsProxy.PullFromPhysicsState();
+			EXPECT_NE(SkeletalMeshPhysicsProxy.GetOutputs(), nullptr);
+		}
+	}
+
+	/**
+	 * A minimal class representing a component that uses the SkeletalMeshPhysicsProxy to
+	 * implements its physics. E.g., USkeletalMeshComponent or FAnimNode_RigidBody
+	 */
+	template <typename Traits>
+	class TFakeSkeletalMeshPhysicsComponent
+	{
+	public:
+		TSharedPtr<FSkeletalMeshPhysicsProxy> SkeletalMeshPhysicsProxy;		
+
+		float BoneRadius = 50.0f;
+		EObjectStateTypeEnum ObjectState = EObjectStateTypeEnum::Chaos_Object_Kinematic;
+		TArray<int32> Parents;
+		TArray< EObjectStateTypeEnum> BoneStates;
+
+		TFramework<Traits> UnitTest;
+
+		// Local-space transform inputs to physics (from animation)
+		TArray<FTransform> InputWorldTransforms;
+		TArray<FVector> InputLinearVelocities;
+		TArray<FVector> InputAngularVelocities;
+
+		// World-space transform outputs from physics (to animation)
+		TArray<FTransform> OutputWorldTransforms;
+		TArray<FVector> OutputLinearVelocities;
+		TArray<FVector> OutputAngularVelocities;
+
+
+		TFakeSkeletalMeshPhysicsComponent(float dt)
+		: UnitTest(FrameworkParameters(dt))
+		{			
+			SkeletalMeshPhysicsProxy = MakeShared<FSkeletalMeshPhysicsProxy>(
+				nullptr,
+				[this](FSkeletalMeshPhysicsProxyParams & OutPhysicsParams) -> void
+				{
+					InitCallback(OutPhysicsParams);
+				});
+		}
+
+		void InitCallback(FSkeletalMeshPhysicsProxyParams & OutParams)
+		{
+			// Check setup is valid
+			check(Parents.Num() == InputWorldTransforms.Num());
+
+			// Fill in data that was not provided
+			InputLinearVelocities.SetNumZeroed(Parents.Num());
+			InputAngularVelocities.SetNumZeroed(Parents.Num());
+			while (BoneStates.Num() < Parents.Num())
+			{
+				BoneStates.Add(EObjectStateTypeEnum::Chaos_Object_UserDefined);
+			}
+
+			// Build the bone hierarchy
+			FBoneHierarchy Hierarchy;
+			{
+				// Build rigid body setup
+				Hierarchy.InitPreAdd(InputWorldTransforms.Num());
+				for (int32 BoneIndex = 0; BoneIndex < InputWorldTransforms.Num(); ++BoneIndex)
+				{
+					const FName NAME_Bone = FName(*FString::Printf(TEXT("Bone_%03d"), BoneIndex));
+					TUniquePtr<FAnalyticImplicitGroup> Group(new FAnalyticImplicitGroup(NAME_Bone, BoneIndex));
+					Group->SetParentBoneIndex(Parents[BoneIndex]);
+					Group->Add(FTransform::Identity, new TSphere<float, 3>(TVector<float, 3>(0), BoneRadius));
+					Group->SetRigidBodyState(BoneStates[BoneIndex]);
+					Hierarchy.Add(MoveTemp(Group));
+				}
+				Hierarchy.InitPostAdd();
+
+				// Set initial transforms
+				Hierarchy.PrepareForUpdate();
+				for (int32 BoneIndex = 0; BoneIndex < InputWorldTransforms.Num(); ++BoneIndex)
+				{
+					const int32 ParentBoneIndex = Parents[BoneIndex];
+					FTransform ParentTransform = (ParentBoneIndex >= 0)? InputWorldTransforms[ParentBoneIndex] : InputWorldTransforms[0];
+					FTransform LocalTransform = InputWorldTransforms[BoneIndex] * ParentTransform.Inverse();
+					Hierarchy.SetAnimLocalSpaceTransform(BoneIndex, LocalTransform);
+				}
+				Hierarchy.SetActorWorldSpaceTransform(InputWorldTransforms[0]);
+				Hierarchy.PrepareAnimWorldSpaceTransforms();
+			}
+
+			// Set up the params structure
+			OutParams.Name = TEXT("Test_Solve");
+			OutParams.InitialTransform = InputWorldTransforms[0];
+			OutParams.InitialLinearVelocity = FVector::ZeroVector;
+			OutParams.InitialAngularVelocity = FVector::ZeroVector;
+			OutParams.BoneHierarchy = MoveTemp(Hierarchy);
+			OutParams.ObjectType = ObjectState;
+			OutParams.bSimulating = true;
+		};
+
+		void Initialize()
+		{
+
+			SkeletalMeshPhysicsProxy->Initialize();
+
+			RigidBodyWrapper* Floor = NewSimulationObject<GeometryType::RigidFloor>()->As<RigidBodyWrapper>();
+			UnitTest.AddSimulationObject(Floor);
+
+			// todo
+			// UnitTest.Solver->RegisterObject(SkeletalMeshPhysicsProxy.Get());
+		}
+
+		void Uninitialize()
+		{
+			// todo
+			// UnitTest.Solver->UnregisterObject(SkeletalMeshPhysicsProxy.Get());
+		}
+
+		void Tick(float Dt)
+		{
+			if (UnitTest.Solver->Enabled())
+			{
+				SkeletalMeshPhysicsProxy->CaptureInputs(Dt, 
+					[this](const float Dt, FSkeletalMeshPhysicsProxyParams & OutParams) -> bool
+					{
+						FBoneHierarchy& Hierarchy = OutParams.BoneHierarchy;
+						Hierarchy.PrepareForUpdate();
+						for (int32 BoneIndex = 0; BoneIndex < InputWorldTransforms.Num(); ++BoneIndex)
+						{
+							const int32 ParentBoneIndex = Parents[BoneIndex];
+							FTransform ParentTransform = (ParentBoneIndex >= 0) ? InputWorldTransforms[ParentBoneIndex] : InputWorldTransforms[0];
+							FTransform LocalTransform = InputWorldTransforms[BoneIndex].GetRelativeTransform(ParentTransform);
+							Hierarchy.SetAnimLocalSpaceTransform(BoneIndex, LocalTransform);
+						}
+						Hierarchy.SetActorWorldSpaceTransform(InputWorldTransforms[0]);
+						Hierarchy.PrepareAnimWorldSpaceTransforms();
+						return true;
+					});
+								
+				UnitTest.Advance();
+
+				if (const FSkeletalMeshPhysicsProxyOutputs * PhysicsOutputs = SkeletalMeshPhysicsProxy->GetOutputs())
+				{
+					OutputWorldTransforms = PhysicsOutputs->Transforms;
+					OutputLinearVelocities = PhysicsOutputs->LinearVelocities;
+					OutputAngularVelocities = PhysicsOutputs->AngularVelocities;
+				};
+			}
+		}
+	};
+
+	/**
+	 * Check that the SkeletalMeshPhysicsProxy is able to provide input and receive correct simulated output from the Solver.
+	 * Check that kinematic body state is correctly reproduces the input animation pose.
+	 */
+	template<typename Traits>
+	void TestSkeletalMeshPhysicsProxy_Kinematic()
+	{
+		const float Dt = (FReal)1 / 30;
+
+		// Two kinematic bodies
+		TFakeSkeletalMeshPhysicsComponent<T> Component(Dt);
+		Component.InputWorldTransforms =
+		{
+			FTransform(FVector(0, 0, 100)),
+			FTransform(FVector(0, 100, 150)),
+		};
+		Component.Parents =
+		{
+			INDEX_NONE,
+			0,
+		};
+
+		Component.Initialize();
+		
+
+		float Time = (FReal)0;		
+
+		for (int32 TickIndex = 0; TickIndex < 100; ++TickIndex)
+		{
+			// Animate both bodies
+			FVector Offset0 = 100.0f * FMath::Sin(2 * PI * Time / 1.0f) * FVector(1, 0, 0);
+			Component.InputWorldTransforms[0] = FTransform(Offset0) * FTransform(FVector(0, 0, 100));
+			Component.InputWorldTransforms[1] = FTransform(FVector(0, 100, 150)) * Component.InputWorldTransforms[0];
+
+			Component.Tick(Dt);
+
+			// All bodies are kinematic. The output pose should match the input pose
+			for (int32 TransformIndex = 0; TransformIndex < Component.InputWorldTransforms.Num(); ++TransformIndex)
+			{
+				EXPECT_NEAR(Component.OutputWorldTransforms[0].GetTranslation().X, Component.InputWorldTransforms[0].GetTranslation().X, KINDA_SMALL_NUMBER);
+				EXPECT_NEAR(Component.OutputWorldTransforms[0].GetTranslation().Y, Component.InputWorldTransforms[0].GetTranslation().Y, KINDA_SMALL_NUMBER);
+				EXPECT_NEAR(Component.OutputWorldTransforms[0].GetTranslation().Z, Component.InputWorldTransforms[0].GetTranslation().Z, KINDA_SMALL_NUMBER);
+			}
+
+			Time += Dt;
+		}
+	}
+
+	/**
+	 * Check that the SkeletalMeshPhysicsProxy is able to provide input and receive correct simulated output from the Solver.
+	 * Check that kinematic and dynamic body state is correctly reproduces the input animation pose.
+	 */
+	template<typename Traits>
+	void TestSkeletalMeshPhysicsProxy_Dynamic()
+	{
+		const float Dt = (FReal)1 / 30;
+
+		// One kinematic, one dynamic body
+		TFakeSkeletalMeshPhysicsComponent<T> Component(Dt);
+		
+		Component.ObjectState = EObjectStateTypeEnum::Chaos_Object_Dynamic;
+		Component.InputWorldTransforms =
+		{
+			FTransform(FVector(0, 0, 300)),
+			//FTransform(FVector(0, 100, 350)),
+			FTransform(FVector(0, 200, 300)),
+		};
+		Component.Parents =
+		{
+			INDEX_NONE,
+			0,
+		};
+		Component.BoneStates =
+		{
+			EObjectStateTypeEnum::Chaos_Object_Kinematic,
+			EObjectStateTypeEnum::Chaos_Object_Dynamic,
+		};
+		
+		// todo
+		// Component.NumIterations = 6;
+
+		Component.Initialize();
+
+		const TArray<FTransform> InitialTransforms = Component.InputWorldTransforms;
+		const float InitialDistance = (InitialTransforms[1].GetTranslation() - InitialTransforms[0].GetTranslation()).Size();
+		
+
+		float Time = (FReal)0;		
+
+		for (int32 TickIndex = 0; TickIndex < 100; ++TickIndex)
+		{
+			// Amimate both bodies. Without pose matching or constraints, only the kinematic body should be affected
+			FVector Offset0 = 500.0f * FMath::Sin(2 * PI * Time / 1.0f) * FVector(1, 0, 0);
+			Component.InputWorldTransforms[0] = FTransform(InitialTransforms[0].GetTranslation() + Offset0);
+			Component.InputWorldTransforms[1] = FTransform(InitialTransforms[1].GetTranslation() + Offset0);
+
+			Component.Tick(Dt);
+
+			Time += Dt;
+
+			// Kinematic body output pose should match the input pose
+			EXPECT_NEAR(Component.OutputWorldTransforms[0].GetTranslation().X, Component.InputWorldTransforms[0].GetTranslation().X, KINDA_SMALL_NUMBER);
+			EXPECT_NEAR(Component.OutputWorldTransforms[0].GetTranslation().Y, Component.InputWorldTransforms[0].GetTranslation().Y, KINDA_SMALL_NUMBER);
+			EXPECT_NEAR(Component.OutputWorldTransforms[0].GetTranslation().Z, Component.InputWorldTransforms[0].GetTranslation().Z, KINDA_SMALL_NUMBER);
+
+			// Dynamic body should swing around the kinematic body at fixed distance
+			const T Distance = (Component.OutputWorldTransforms[1].GetTranslation() - Component.OutputWorldTransforms[0].GetTranslation()).Size();
+			EXPECT_NEAR(Distance, InitialDistance, (FReal)3.0);
+		}
+	}
+}
