Merge geometry collection into headless chaos.

diff --git a/Engine/Source/Programs/HeadlessChaos/Resource/CylinderGeometry.h b/Engine/Source/Programs/HeadlessChaos/Resource/CylinderGeometry.h
new file mode 100644
--- /dev/null
+++ b/Engine/Source/Programs/HeadlessChaos/Resource/CylinderGeometry.h
@@ -0,0 +1,738 @@
+// Copyright Epic Games, Inc. All Rights Reserved.
+
+#pragma once
+namespace GeometryCollectionTest
+{
+	class CylinderGeometry
+	{
+
+	public:
+		CylinderGeometry() {}
+		~CylinderGeometry() {}
+
+		static const TArray<float>	RawVertexArray;
+		static const TArray<int32>	RawIndicesArray;
+	};
+
+	const TArray<float> CylinderGeometry::RawVertexArray = {
+																100.000000, 0.000000, 100.000000,
+																95.105652, 30.901701, 100.000000,
+																80.901703, 58.778526, 100.000000,
+																58.778526, 80.901703, 100.000000,
+																30.901697, 95.105652, 100.000000,
+																-0.000004, 100.000000, 100.000000,
+																-30.901703, 95.105652, 100.000000,
+																-58.778519, 80.901703, 100.000000,
+																-80.901703, 58.778519, 100.000000,
+																-95.105659, 30.901680, 100.000000,
+																-100.000000, -0.000009, 100.000000,
+																-95.105652, -30.901697, 100.000000,
+																-80.901695, -58.778538, 100.000000,
+																-58.778507, -80.901711, 100.000000,
+																-30.901709, -95.105652, 100.000000,
+																0.000001, -100.000000, 100.000000,
+																30.901712, -95.105652, 100.000000,
+																58.778549, -80.901680, 100.000000,
+																80.901726, -58.778496, 100.000000,
+																95.105652, -30.901695, 100.000000,
+																100.000000, 0.000000, 60.000004,
+																95.105652, 30.901701, 60.000004,
+																80.901703, 58.778526, 60.000004,
+																58.778526, 80.901703, 60.000004,
+																30.901697, 95.105652, 60.000004,
+																-0.000004, 100.000000, 60.000004,
+																-30.901703, 95.105652, 60.000004,
+																-58.778519, 80.901703, 60.000004,
+																-80.901703, 58.778519, 60.000004,
+																-95.105659, 30.901680, 60.000004,
+																-100.000000, -0.000009, 60.000004,
+																-95.105652, -30.901697, 60.000004,
+																-80.901695, -58.778538, 60.000004,
+																-58.778507, -80.901711, 60.000004,
+																-30.901709, -95.105652, 60.000004,
+																0.000001, -100.000000, 60.000004,
+																30.901712, -95.105652, 60.000004,
+																58.778549, -80.901680, 60.000004,
+																80.901726, -58.778496, 60.000004,
+																95.105652, -30.901695, 60.000004,
+																100.000000, 0.000000, 20.000002,
+																95.105652, 30.901701, 20.000002,
+																80.901703, 58.778526, 20.000002,
+																58.778526, 80.901703, 20.000002,
+																30.901697, 95.105652, 20.000002,
+																-0.000004, 100.000000, 20.000002,
+																-30.901703, 95.105652, 20.000002,
+																-58.778519, 80.901703, 20.000002,
+																-80.901703, 58.778519, 20.000002,
+																-95.105659, 30.901680, 20.000002,
+																-100.000000, -0.000009, 20.000002,
+																-95.105652, -30.901697, 20.000002,
+																-80.901695, -58.778538, 20.000002,
+																-58.778507, -80.901711, 20.000002,
+																-30.901709, -95.105652, 20.000002,
+																0.000001, -100.000000, 20.000002,
+																30.901712, -95.105652, 20.000002,
+																58.778549, -80.901680, 20.000002,
+																80.901726, -58.778496, 20.000002,
+																95.105652, -30.901695, 20.000002,
+																100.000000, -0.000000, -19.999998,
+																95.105652, 30.901701, -19.999998,
+																80.901703, 58.778526, -19.999998,
+																58.778526, 80.901703, -19.999998,
+																30.901697, 95.105652, -19.999998,
+																-0.000004, 100.000000, -19.999998,
+																-30.901703, 95.105652, -19.999998,
+																-58.778519, 80.901703, -19.999998,
+																-80.901703, 58.778519, -19.999998,
+																-95.105659, 30.901680, -19.999998,
+																-100.000000, -0.000009, -19.999998,
+																-95.105652, -30.901697, -19.999998,
+																-80.901695, -58.778538, -19.999998,
+																-58.778507, -80.901711, -19.999998,
+																-30.901709, -95.105652, -19.999998,
+																0.000001, -100.000000, -19.999998,
+																30.901712, -95.105652, -19.999998,
+																58.778549, -80.901680, -19.999998,
+																80.901726, -58.778496, -19.999998,
+																95.105652, -30.901695, -19.999998,
+																100.000000, -0.000000, -60.000004,
+																95.105652, 30.901701, -60.000004,
+																80.901703, 58.778526, -60.000004,
+																58.778526, 80.901703, -60.000004,
+																30.901697, 95.105652, -60.000004,
+																-0.000004, 100.000000, -60.000004,
+																-30.901703, 95.105652, -60.000004,
+																-58.778519, 80.901703, -60.000004,
+																-80.901703, 58.778519, -60.000004,
+																-95.105659, 30.901680, -60.000004,
+																-100.000000, -0.000009, -60.000004,
+																-95.105652, -30.901697, -60.000004,
+																-80.901695, -58.778538, -60.000004,
+																-58.778507, -80.901711, -60.000004,
+																-30.901709, -95.105652, -60.000004,
+																0.000001, -100.000000, -60.000004,
+																30.901712, -95.105652, -60.000004,
+																58.778549, -80.901680, -60.000004,
+																80.901726, -58.778496, -60.000004,
+																95.105652, -30.901695, -60.000004,
+																100.000000, -0.000000, -100.000000,
+																95.105652, 30.901701, -100.000000,
+																80.901703, 58.778526, -100.000000,
+																58.778526, 80.901703, -100.000000,
+																30.901697, 95.105652, -100.000000,
+																-0.000004, 100.000000, -100.000000,
+																-30.901703, 95.105652, -100.000000,
+																-58.778519, 80.901703, -100.000000,
+																-80.901703, 58.778519, -100.000000,
+																-95.105659, 30.901680, -100.000000,
+																-100.000000, -0.000009, -100.000000,
+																-95.105652, -30.901697, -100.000000,
+																-80.901695, -58.778538, -100.000000,
+																-58.778507, -80.901711, -100.000000,
+																-30.901709, -95.105652, -100.000000,
+																0.000001, -100.000000, -100.000000,
+																30.901712, -95.105652, -100.000000,
+																58.778549, -80.901680, -100.000000,
+																80.901726, -58.778496, -100.000000,
+																95.105652, -30.901695, -100.000000,
+																100.000000, 0.0000003, 100.000000,
+																95.105652, 30.9017013, 100.000000,
+																80.901703, 58.7785263, 100.000000,
+																58.778526, 80.9017033, 100.000000,
+																30.901697, 95.1056523, 100.000000,
+																-0.000004, 100.0000003, 100.000000,
+																-30.901703, 95.1056523, 100.000000,
+																-58.778519, 80.9017033, 100.000000,
+																-80.901703, 58.7785193, 100.000000,
+																-95.105659, 30.9016803, 100.000000,
+																-100.000000, -0.0000093, 100.000000,
+																-95.105652, -30.9016973, 100.000000,
+																-80.901695, -58.7785383, 100.000000,
+																-58.778507, -80.9017113, 100.000000,
+																-30.901709, -95.1056523, 100.000000,
+																0.000001, -100.0000003, 100.000000,
+																30.901712, -95.1056523, 100.000000,
+																58.778549, -80.9016803, 100.000000,
+																80.901726, -58.7784963, 100.000000,
+																95.105652, -30.9016953, 100.000000,
+																100.000000, 0.0000003, 60.000004,
+																95.105652, 30.9017013, 60.000004,
+																80.901703, 58.7785263, 60.000004,
+																58.778526, 80.9017033, 60.000004,
+																30.901697, 95.1056523, 60.000004,
+																-0.000004, 100.0000003, 60.000004,
+																-30.901703, 95.1056523, 60.000004,
+																-58.778519, 80.9017033, 60.000004,
+																-80.901703, 58.7785193, 60.000004,
+																-95.105659, 30.9016803, 60.000004,
+																-100.000000, -0.0000093, 60.000004,
+																-95.105652, -30.9016973, 60.000004,
+																-80.901695, -58.7785383, 60.000004,
+																-58.778507, -80.9017113, 60.000004,
+																-30.901709, -95.1056523, 60.000004,
+																0.000001, -100.0000003, 60.000004,
+																30.901712, -95.1056523, 60.000004,
+																58.778549, -80.9016803, 60.000004,
+																80.901726, -58.7784963, 60.000004,
+																95.105652, -30.9016953, 60.000004,
+																100.000000, 0.0000003, 20.000002,
+																95.105652, 30.9017013, 20.000002,
+																80.901703, 58.7785263, 20.000002,
+																58.778526, 80.9017033, 20.000002,
+																30.901697, 95.1056523, 20.000002,
+																-0.000004, 100.0000003, 20.000002,
+																-30.901703, 95.1056523, 20.000002,
+																-58.778519, 80.9017033, 20.000002,
+																-80.901703, 58.7785193, 20.000002,
+																-95.105659, 30.9016803, 20.000002,
+																-100.000000, -0.0000093, 20.000002,
+																-95.105652, -30.9016973, 20.000002,
+																-80.901695, -58.7785383, 20.000002,
+																-58.778507, -80.9017113, 20.000002,
+																-30.901709, -95.1056523, 20.000002,
+																0.000001, -100.0000003, 20.000002,
+																30.901712, -95.1056523, 20.000002,
+																58.778549, -80.9016803, 20.000002,
+																80.901726, -58.7784963, 20.000002,
+																95.105652, -30.9016953, 20.000002,
+																100.000000, -0.0000003, -19.999998,
+																95.105652, 30.9017013, -19.999998,
+																80.901703, 58.7785263, -19.999998,
+																58.778526, 80.9017033, -19.999998,
+																30.901697, 95.1056523, -19.999998,
+																-0.000004, 100.0000003, -19.999998,
+																-30.901703, 95.1056523, -19.999998,
+																-58.778519, 80.9017033, -19.999998,
+																-80.901703, 58.7785193, -19.999998,
+																-95.105659, 30.9016803, -19.999998,
+																-100.000000, -0.0000093, -19.999998,
+																-95.105652, -30.9016973, -19.999998,
+																-80.901695, -58.7785383, -19.999998,
+																-58.778507, -80.9017113, -19.999998,
+																-30.901709, -95.1056523, -19.999998,
+																0.000001, -100.0000003, -19.999998,
+																30.901712, -95.1056523, -19.999998,
+																58.778549, -80.9016803, -19.999998,
+																80.901726, -58.7784963, -19.999998,
+																95.105652, -30.9016953, -19.999998,
+																100.000000, -0.0000003, -60.000004,
+																95.105652, 30.9017013, -60.000004,
+																80.901703, 58.7785263, -60.000004,
+																58.778526, 80.9017033, -60.000004,
+																30.901697, 95.1056523, -60.000004,
+																-0.000004, 100.0000003, -60.000004,
+																-30.901703, 95.1056523, -60.000004,
+																-58.778519, 80.9017033, -60.000004,
+																-80.901703, 58.7785193, -60.000004,
+																-95.105659, 30.9016803, -60.000004,
+																-100.000000, -0.0000093, -60.000004,
+																-95.105652, -30.9016973, -60.000004,
+																-80.901695, -58.7785383, -60.000004,
+																-58.778507, -80.9017113, -60.000004,
+																-30.901709, -95.1056523, -60.000004,
+																0.000001, -100.0000003, -60.000004,
+																30.901712, -95.1056523, -60.000004,
+																58.778549, -80.9016803, -60.000004,
+																80.901726, -58.7784963, -60.000004,
+																95.105652, -30.9016953, -60.000004,
+																100.000000, -0.0000003, -100.000000,
+																95.105652, 30.9017013, -100.000000,
+																80.901703, 58.7785263, -100.000000,
+																58.778526, 80.9017033, -100.000000,
+																30.901697, 95.1056523, -100.000000,
+																-0.000004, 100.0000003, -100.000000,
+																-30.901703, 95.1056523, -100.000000,
+																-58.778519, 80.9017033, -100.000000,
+																-80.901703, 58.7785193, -100.000000,
+																-95.105659, 30.9016803, -100.000000,
+																-100.000000, -0.0000093, -100.000000,
+																-95.105652, -30.9016973, -100.000000,
+																-80.901695, -58.7785383, -100.000000,
+																-58.778507, -80.9017113, -100.000000,
+																-30.901709, -95.1056523, -100.000000,
+																0.000001, -100.0000003, -100.000000,
+																30.901712, -95.1056523, -100.000000,
+																58.778549, -80.9016803, -100.000000,
+																80.901726, -58.7784963, -100.000000,
+																95.105652, -30.9016953, -100.000000,
+																100.000000,       0.0000, 100.000000,
+																95.105652,      30.9017, 100.000000,
+																80.901703,      58.7785, 100.000000,
+																58.778526,      80.9017, 100.000000,
+																30.901697,      95.1057, 100.000000,
+																-0.000004,     100.0000, 100.000000,
+																-30.901703,      95.1057, 100.000000,
+																-58.778519,      80.9017, 100.000000,
+																-80.901703,      58.7785, 100.000000,
+																-95.105659,      30.9017, 100.000000,
+																-100.000000,      -0.0000, 100.000000,
+																-95.105652,     -30.9017, 100.000000,
+																-80.901695,     -58.7785, 100.000000,
+																-58.778507,     -80.9017, 100.000000,
+																-30.901709,     -95.1057, 100.000000,
+																0.000001,    -100.0000, 100.000000,
+																30.901712,     -95.1057, 100.000000,
+																58.778549,     -80.9017, 100.000000,
+																80.901726,     -58.7785, 100.000000,
+																95.105652,     -30.9017, 100.000000,
+																100.000000,       0.0000, 60.000004,
+																95.105652,      30.9017, 60.000004,
+																80.901703,      58.7785, 60.000004,
+																58.778526,      80.9017, 60.000004,
+																30.901697,      95.1057, 60.000004,
+																-0.000004,     100.0000, 60.000004,
+																-30.901703,      95.1057, 60.000004,
+																-58.778519,      80.9017, 60.000004,
+																-80.901703,      58.7785, 60.000004,
+																-95.105659,      30.9017, 60.000004,
+																-100.000000,      -0.0000, 60.000004,
+																-95.105652,     -30.9017, 60.000004,
+																-80.901695,     -58.7785, 60.000004,
+																-58.778507,     -80.9017, 60.000004,
+																-30.901709,     -95.1057, 60.000004,
+																0.000001,    -100.0000, 60.000004,
+																30.901712,     -95.1057, 60.000004,
+																58.778549,     -80.9017, 60.000004,
+																80.901726,     -58.7785, 60.000004,
+																95.105652,     -30.9017, 60.000004,
+																100.000000,       0.0000, 20.000002,
+																95.105652,      30.9017, 20.000002,
+																80.901703,      58.7785, 20.000002,
+																58.778526,      80.9017, 20.000002,
+																30.901697,      95.1057, 20.000002,
+																-0.000004,     100.0000, 20.000002,
+																-30.901703,      95.1057, 20.000002,
+																-58.778519,      80.9017, 20.000002,
+																-80.901703,      58.7785, 20.000002,
+																-95.105659,      30.9017, 20.000002,
+																-100.000000,      -0.0000, 20.000002,
+																-95.105652,     -30.9017, 20.000002,
+																-80.901695,     -58.7785, 20.000002,
+																-58.778507,     -80.9017, 20.000002,
+																-30.901709,     -95.1057, 20.000002,
+																0.000001,    -100.0000, 20.000002,
+																30.901712,     -95.1057, 20.000002,
+																58.778549,     -80.9017, 20.000002,
+																80.901726,     -58.7785, 20.000002,
+																95.105652,     -30.9017, 20.000002,
+																100.000000,      -0.0000, -19.999998,
+																95.105652,      30.9017, -19.999998,
+																80.901703,      58.7785, -19.999998,
+																58.778526,      80.9017, -19.999998,
+																30.901697,      95.1057, -19.999998,
+																-0.000004,     100.0000, -19.999998,
+																-30.901703,      95.1057, -19.999998,
+																-58.778519,      80.9017, -19.999998,
+																-80.901703,      58.7785, -19.999998,
+																-95.105659,      30.9017, -19.999998,
+																-100.000000,      -0.0000, -19.999998,
+																-95.105652,     -30.9017, -19.999998,
+																-80.901695,     -58.7785, -19.999998,
+																-58.778507,     -80.9017, -19.999998,
+																-30.901709,     -95.1057, -19.999998,
+																0.000001,    -100.0000, -19.999998,
+																30.901712,     -95.1057, -19.999998,
+																58.778549,     -80.9017, -19.999998,
+																80.901726,     -58.7785, -19.999998,
+																95.105652,     -30.9017, -19.999998,
+																100.000000,      -0.0000, -60.000004,
+																95.105652,      30.9017, -60.000004,
+																80.901703,      58.7785, -60.000004,
+																58.778526,      80.9017, -60.000004,
+																30.901697,      95.1057, -60.000004,
+																-0.000004,     100.0000, -60.000004,
+																-30.901703,      95.1057, -60.000004,
+																-58.778519,      80.9017, -60.000004,
+																-80.901703,      58.7785, -60.000004,
+																-95.105659,      30.9017, -60.000004,
+																-100.000000,      -0.0000, -60.000004,
+																-95.105652,     -30.9017, -60.000004,
+																-80.901695,     -58.7785, -60.000004,
+																-58.778507,     -80.9017, -60.000004,
+																-30.901709,     -95.1057, -60.000004,
+																0.000001,    -100.0000, -60.000004,
+																30.901712,     -95.1057, -60.000004,
+																58.778549,     -80.9017, -60.000004,
+																80.901726,     -58.7785, -60.000004,
+																95.105652,     -30.9017, -60.000004,
+																100.000000,      -0.0000, -100.000000,
+																95.105652,      30.9017, -100.000000,
+																80.901703,      58.7785, -100.000000,
+																58.778526,      80.9017, -100.000000,
+																30.901697,      95.1057, -100.000000,
+																-0.000004,     100.0000, -100.000000,
+																-30.901703,      95.1057, -100.000000,
+																-58.778519,      80.9017, -100.000000,
+																-80.901703,      58.7785, -100.000000,
+																-95.105659,      30.9017, -100.000000,
+																-100.000000,      -0.0000, -100.000000,
+																-95.105652,     -30.9017, -100.000000,
+																-80.901695,     -58.7785, -100.000000,
+																-58.778507,     -80.9017, -100.000000,
+																-30.901709,     -95.1057, -100.000000,
+																0.000001,    -100.0000, -100.000000,
+																30.901712,     -95.1057, -100.000000,
+																58.778549,     -80.9017, -100.000000,
+																80.901726,     -58.7785, -100.000000,
+																95.105652,     -30.9017, -100.000000,
+																100.000000, 0.000000, 100.000000,
+																95.105652, 30.901701, 100.000000,
+																80.901703, 58.778526, 100.000000,
+																58.778526, 80.901703, 100.000000,
+																30.901697, 95.105652, 100.000000,
+																-0.000004, 100.000000, 100.000000,
+																-30.901703, 95.105652, 100.000000,
+																-58.778519, 80.901703, 100.000000,
+																-80.901703, 58.778519, 100.000000,
+																-95.105659, 30.901680, 100.000000,
+																-100.000000, -0.000009, 100.000000,
+																-95.105652, -30.901697, 100.000000,
+																-80.901695, -58.778538, 100.000000,
+																-58.778507, -80.901711, 100.000000,
+																-30.901709, -95.105652, 100.000000,
+																0.000001, -100.000000, 100.000000,
+																30.901712, -95.105652, 100.000000,
+																58.778549, -80.901680, 100.000000,
+																80.901726, -58.778496, 100.000000,
+																95.105652, -30.901695, 100.000000,
+																100.000000, 0.000000, 60.000004,
+																95.105652, 30.901701, 60.000004,
+																80.901703, 58.778526, 60.000004,
+																58.778526, 80.901703, 60.000004,
+																30.901697, 95.105652, 60.000004,
+																-0.000004, 100.000000, 60.000004,
+																-30.901703, 95.105652, 60.000004,
+																-58.778519, 80.901703, 60.000004,
+																-80.901703, 58.778519, 60.000004,
+																-95.105659, 30.901680, 60.000004,
+																-100.000000, -0.000009, 60.000004,
+																-95.105652, -30.901697, 60.000004,
+																-80.901695, -58.778538, 60.000004,
+																-58.778507, -80.901711, 60.000004,
+																-30.901709, -95.105652, 60.000004,
+																0.000001, -100.000000, 60.000004,
+																30.901712, -95.105652, 60.000004,
+																58.778549, -80.901680, 60.000004,
+																80.901726, -58.778496, 60.000004,
+																95.105652, -30.901695, 60.000004,
+																100.000000, 0.000000, 20.000002,
+																95.105652, 30.901701, 20.000002,
+																80.901703, 58.778526, 20.000002,
+																58.778526, 80.901703, 20.000002,
+																30.901697, 95.105652, 20.000002,
+																-0.000004, 100.000000, 20.000002,
+																-30.901703, 95.105652, 20.000002,
+																-58.778519, 80.901703, 20.000002,
+																-80.901703, 58.778519, 20.000002,
+																-95.105659, 30.901680, 20.000002,
+																-100.000000, -0.000009, 20.000002,
+																-95.105652, -30.901697, 20.000002,
+																-80.901695, -58.778538, 20.000002,
+																-58.778507, -80.901711, 20.000002,
+																-30.901709, -95.105652, 20.000002,
+																0.000001, -100.000000, 20.000002,
+																30.901712, -95.105652, 20.000002,
+																58.778549, -80.901680, 20.000002,
+																80.901726, -58.778496, 20.000002,
+																95.105652, -30.901695, 20.000002,
+																100.000000, -0.000000, -19.999998,
+																95.105652, 30.901701, -19.999998,
+																80.901703, 58.778526, -19.999998,
+																58.778526, 80.901703, -19.999998,
+																30.901697, 95.105652, -19.999998,
+																-0.000004, 100.000000, -19.999998,
+																-30.901703, 95.105652, -19.999998,
+																-58.778519, 80.901703, -19.999998,
+																-80.901703, 58.778519, -19.999998,
+																-95.105659, 30.901680, -19.999998,
+																-100.000000, -0.000009, -19.999998,
+																-95.105652, -30.901697, -19.999998,
+																-80.901695, -58.778538, -19.999998,
+																-58.778507, -80.901711, -19.999998,
+																-30.901709, -95.105652, -19.999998,
+																0.000001, -100.000000, -19.999998,
+																30.901712, -95.105652, -19.999998,
+																58.778549, -80.901680, -19.999998,
+																80.901726, -58.778496, -19.999998,
+																95.105652, -30.901695, -19.999998,
+																100.000000, -0.000000, -60.000004,
+																95.105652, 30.901701, -60.000004,
+																80.901703, 58.778526, -60.000004,
+																58.778526, 80.901703, -60.000004,
+																30.901697, 95.105652, -60.000004,
+																-0.000004, 100.000000, -60.000004,
+																-30.901703, 95.105652, -60.000004,
+																-58.778519, 80.901703, -60.000004,
+																-80.901703, 58.778519, -60.000004,
+																-95.105659, 30.901680, -60.000004,
+																-100.000000, -0.000009, -60.000004,
+																-95.105652, -30.901697, -60.000004,
+																-80.901695, -58.778538, -60.000004,
+																-58.778507, -80.901711, -60.000004,
+																-30.901709, -95.105652, -60.000004,
+																0.000001, -100.000000, -60.000004,
+																30.901712, -95.105652, -60.000004,
+																58.778549, -80.901680, -60.000004,
+																80.901726, -58.778496, -60.000004,
+																95.105652, -30.901695, -60.000004,
+																100.000000, -0.000000, -100.000000,
+																95.105652, 30.901701, -100.000000,
+																80.901703, 58.778526, -100.000000,
+																58.778526, 80.901703, -100.000000,
+																30.901697, 95.105652, -100.000000,
+																-0.000004, 100.000000, -100.000000,
+																-30.901703, 95.105652, -100.000000,
+																-58.778519, 80.901703, -100.000000,
+																-80.901703, 58.778519, -100.000000,
+																-95.105659, 30.901680, -100.000000,
+																-100.000000, -0.000009, -100.000000,
+																-95.105652, -30.901697, -100.000000,
+																-80.901695, -58.778538, -100.000000,
+																-58.778507, -80.901711, -100.000000,
+																-30.901709, -95.105652, -100.000000,
+																0.000001, -100.000000, -100.000000,
+																30.901712, -95.105652, -100.000000,
+																58.778549, -80.901680, -100.000000,
+																80.901726, -58.778496, -100.000000,
+																95.105652, -30.901695, -100.000000
+	};
+
+	const TArray<int32> CylinderGeometry::RawIndicesArray = {
+																10, 9, 8,
+																109, 110, 111,
+																0, 1, 21,
+																1, 2, 22,
+																2, 3, 23,
+																3, 4, 24,
+																4, 5, 25,
+																5, 6, 26,
+																6, 7, 27,
+																7, 8, 28,
+																8, 9, 29,
+																9, 10, 30,
+																10, 11, 31,
+																11, 12, 32,
+																12, 13, 33,
+																13, 14, 34,
+																14, 15, 35,
+																15, 16, 36,
+																16, 17, 37,
+																17, 18, 38,
+																18, 19, 39,
+																19, 0, 20,
+																20, 21, 41,
+																21, 22, 42,
+																22, 23, 43,
+																23, 24, 44,
+																24, 25, 45,
+																25, 26, 46,
+																26, 27, 47,
+																27, 28, 48,
+																28, 29, 49,
+																29, 30, 50,
+																30, 31, 51,
+																31, 32, 52,
+																32, 33, 53,
+																33, 34, 54,
+																34, 35, 55,
+																35, 36, 56,
+																36, 37, 57,
+																37, 38, 58,
+																38, 39, 59,
+																39, 20, 40,
+																40, 41, 61,
+																41, 42, 62,
+																42, 43, 63,
+																43, 44, 64,
+																44, 45, 65,
+																45, 46, 66,
+																46, 47, 67,
+																47, 48, 68,
+																48, 49, 69,
+																49, 50, 70,
+																50, 51, 71,
+																51, 52, 72,
+																52, 53, 73,
+																53, 54, 74,
+																54, 55, 75,
+																55, 56, 76,
+																56, 57, 77,
+																57, 58, 78,
+																58, 59, 79,
+																59, 40, 60,
+																60, 61, 81,
+																61, 62, 82,
+																62, 63, 83,
+																63, 64, 84,
+																64, 65, 85,
+																65, 66, 86,
+																66, 67, 87,
+																67, 68, 88,
+																68, 69, 89,
+																69, 70, 90,
+																70, 71, 91,
+																71, 72, 92,
+																72, 73, 93,
+																73, 74, 94,
+																74, 75, 95,
+																75, 76, 96,
+																76, 77, 97,
+																77, 78, 98,
+																78, 79, 99,
+																79, 60, 80,
+																80, 81, 101,
+																81, 82, 102,
+																82, 83, 103,
+																83, 84, 104,
+																84, 85, 105,
+																85, 86, 106,
+																86, 87, 107,
+																87, 88, 108,
+																88, 89, 109,
+																89, 90, 110,
+																90, 91, 111,
+																91, 92, 112,
+																92, 93, 113,
+																93, 94, 114,
+																94, 95, 115,
+																95, 96, 116,
+																96, 97, 117,
+																97, 98, 118,
+																98, 99, 119,
+																99, 80, 100,
+																100, 119, 99,
+																119, 118, 98,
+																118, 117, 97,
+																117, 116, 96,
+																116, 115, 95,
+																115, 114, 94,
+																114, 113, 93,
+																113, 112, 92,
+																112, 111, 91,
+																111, 110, 90,
+																110, 109, 89,
+																109, 108, 88,
+																108, 107, 87,
+																107, 106, 86,
+																106, 105, 85,
+																105, 104, 84,
+																104, 103, 83,
+																103, 102, 82,
+																102, 101, 81,
+																101, 100, 80,
+																80, 99, 79,
+																99, 98, 78,
+																98, 97, 77,
+																97, 96, 76,
+																96, 95, 75,
+																95, 94, 74,
+																94, 93, 73,
+																93, 92, 72,
+																92, 91, 71,
+																91, 90, 70,
+																90, 89, 69,
+																89, 88, 68,
+																88, 87, 67,
+																87, 86, 66,
+																86, 85, 65,
+																85, 84, 64,
+																84, 83, 63,
+																83, 82, 62,
+																82, 81, 61,
+																81, 80, 60,
+																60, 79, 59,
+																79, 78, 58,
+																78, 77, 57,
+																77, 76, 56,
+																76, 75, 55,
+																75, 74, 54,
+																74, 73, 53,
+																73, 72, 52,
+																72, 71, 51,
+																71, 70, 50,
+																70, 69, 49,
+																69, 68, 48,
+																68, 67, 47,
+																67, 66, 46,
+																66, 65, 45,
+																65, 64, 44,
+																64, 63, 43,
+																63, 62, 42,
+																62, 61, 41,
+																61, 60, 40,
+																40, 59, 39,
+																59, 58, 38,
+																58, 57, 37,
+																57, 56, 36,
+																56, 55, 35,
+																55, 54, 34,
+																54, 53, 33,
+																53, 52, 32,
+																52, 51, 31,
+																51, 50, 30,
+																50, 49, 29,
+																49, 48, 28,
+																48, 47, 27,
+																47, 46, 26,
+																46, 45, 25,
+																45, 44, 24,
+																44, 43, 23,
+																43, 42, 22,
+																42, 41, 21,
+																41, 40, 20,
+																20, 39, 19,
+																39, 38, 18,
+																38, 37, 17,
+																37, 36, 16,
+																36, 35, 15,
+																35, 34, 14,
+																34, 33, 13,
+																33, 32, 12,
+																32, 31, 11,
+																31, 30, 10,
+																30, 29, 9,
+																29, 28, 8,
+																28, 27, 7,
+																27, 26, 6,
+																26, 25, 5,
+																25, 24, 4,
+																24, 23, 3,
+																23, 22, 2,
+																22, 21, 1,
+																21, 20, 0,
+																100, 101, 119,
+																101, 102, 119,
+																102, 118, 119,
+																102, 103, 118,
+																103, 117, 118,
+																103, 104, 117,
+																104, 116, 117,
+																104, 105, 116,
+																105, 115, 116,
+																105, 106, 115,
+																106, 114, 115,
+																106, 107, 114,
+																107, 113, 114,
+																107, 108, 113,
+																108, 112, 113,
+																108, 109, 112,
+																109, 111, 112,
+																19, 18, 0,
+																18, 17, 0,
+																17, 1, 0,
+																17, 16, 1,
+																16, 2, 1,
+																16, 15, 2,
+																15, 3, 2,
+																15, 14, 3,
+																14, 4, 3,
+																14, 13, 4,
+																13, 5, 4,
+																13, 12, 5,
+																12, 6, 5,
+																12, 11, 6,
+																11, 7, 6,
+																11, 10, 7,
+																10, 8, 7
+	};
+}
\ No newline at end of file
diff --git a/Engine/Source/Programs/HeadlessChaos/Resource/EllipsoidGeometry.h b/Engine/Source/Programs/HeadlessChaos/Resource/EllipsoidGeometry.h
new file mode 100644
--- /dev/null
+++ b/Engine/Source/Programs/HeadlessChaos/Resource/EllipsoidGeometry.h
@@ -0,0 +1,776 @@
+// Copyright Epic Games, Inc. All Rights Reserved.
+
+#pragma once
+namespace GeometryCollectionTest
+{
+	class EllipsoidGeometry
+	{
+
+	public:
+		EllipsoidGeometry() {}
+		~EllipsoidGeometry() {}
+
+		static const TArray<float>	RawVertexArray;
+		static const TArray<int32>	RawIndicesArray;
+	};
+
+
+	const TArray<float> EllipsoidGeometry::RawVertexArray = {
+																0.000000, 0.000000, -100.000000,
+																0.000000, 21.962393, -97.558456,
+																41.774948, 6.786763, -97.558472,
+																0.000000, 42.852455, -90.353020,
+																44.223892, 30.434498, -92.654213,
+																81.510193, 13.242155, -90.353020,
+																0.000000, 61.650002, -78.735481,
+																45.006493, 51.224689, -82.883247,
+																83.527344, 37.231178, -82.883247,
+																117.265259, 19.050928, -78.735489,
+																0.000000, 77.437073, -63.273209,
+																44.223896, 69.261757, -68.657623,
+																83.527359, 57.482807, -70.367065,
+																118.077721, 42.432804, -68.657631,
+																147.294037, 23.929405, -63.273209,
+																0.000000, 89.442734, -44.721313,
+																41.774960, 84.223846, -49.699677,
+																81.510201, 74.892166, -52.251171,
+																117.265259, 61.903389, -52.251171,
+																147.294037, 45.891808, -49.699677,
+																170.130173, 27.639368, -44.721321,
+																25.818336, -17.767946, -97.558472,
+																71.555756, -11.624928, -92.654213,
+																50.376068, -34.668358, -90.353020,
+																111.342911, -5.572543, -82.883270,
+																96.629242, -28.214546, -82.883270,
+																72.473923, -49.875896, -78.735497,
+																145.409622, 0.373371, -68.657646,
+																135.150116, -21.956450, -70.367073,
+																117.199966, -43.036842, -68.657646,
+																91.032738, -62.647907, -63.273220,
+																173.112427, 6.161456, -49.699692,
+																167.641403, -15.617439, -52.251190,
+																153.984177, -36.633751, -52.251190,
+																132.807724, -55.861145, -49.699692,
+																105.146240, -72.360695, -44.721325,
+																-25.818336, -17.767946, -97.558472,
+																0.000000, -37.619091, -92.654213,
+																-50.376068, -34.668358, -90.353020,
+																23.807199, -54.668720, -82.883270,
+																-23.807199, -54.668720, -82.883270,
+																-72.473923, -49.875896, -78.735497,
+																45.644176, -69.031013, -68.657646,
+																0.000000, -71.052620, -70.367065,
+																-45.644176, -69.031013, -68.657646,
+																-91.032738, -62.647907, -63.273220,
+																65.214394, -80.415871, -49.699692,
+																22.097843, -84.544273, -52.251183,
+																-22.097843, -84.544273, -52.251183,
+																-65.214394, -80.415871, -49.699692,
+																-105.146240, -72.360695, -44.721325,
+																-41.774948, 6.786763, -97.558472,
+																-71.555756, -11.624928, -92.654213,
+																-81.510193, 13.242155, -90.353020,
+																-96.629242, -28.214546, -82.883270,
+																-111.342911, -5.572543, -82.883270,
+																-117.265259, 19.050928, -78.735489,
+																-117.199966, -43.036842, -68.657646,
+																-135.150116, -21.956450, -70.367073,
+																-145.409622, 0.373371, -68.657646,
+																-147.294037, 23.929405, -63.273209,
+																-132.807724, -55.861145, -49.699692,
+																-153.984177, -36.633751, -52.251190,
+																-167.641403, -15.617439, -52.251190,
+																-173.112427, 6.161456, -49.699692,
+																-170.130173, 27.639368, -44.721321,
+																-44.223892, 30.434498, -92.654213,
+																-83.527344, 37.231178, -82.883247,
+																-45.006493, 51.224689, -82.883247,
+																-118.077721, 42.432804, -68.657631,
+																-83.527359, 57.482807, -70.367065,
+																-44.223896, 69.261757, -68.657623,
+																-147.294037, 45.891808, -49.699677,
+																-117.265259, 61.903389, -52.251171,
+																-81.510201, 74.892166, -52.251171,
+																-41.774960, 84.223846, -49.699677,
+																-25.818338, 95.205017, -27.737299,
+																-71.555756, 88.071297, -31.038528,
+																-50.376068, 96.318359, -9.398762,
+																-111.342873, 76.625198, -32.082233,
+																-96.629227, 86.750999, -11.830645,
+																-72.473907, 92.728340, 9.398762,
+																-145.409561, 61.242340, -31.038528,
+																-135.150070, 72.757484, -11.830645,
+																-117.199936, 80.655952, 7.788676,
+																-91.032715, 84.610298, 27.737301,
+																-173.112381, 41.697357, -27.737301,
+																-167.641312, 53.719284, -9.398762,
+																-153.984100, 63.118046, 9.398762,
+																-132.807678, 69.434662, 27.737301,
+																-105.146240, 72.360695, 44.721325,
+																-189.069000, 17.142641, -27.737301,
+																-189.633484, -6.811248, -31.038534,
+																-198.775467, 5.808771, -9.398764,
+																-180.156631, -29.268211, -32.082249,
+																-194.870255, -19.142406, -11.830646,
+																-198.775467, -5.808771, 9.398764,
+																-161.423843, -50.221455, -31.038534,
+																-180.156601, -41.784412, -11.830647,
+																-189.633469, -30.807846, 7.788676,
+																-189.069000, -17.142641, 27.737301,
+																-132.807678, -69.434662, -27.737301,
+																-153.984100, -63.118046, -9.398762,
+																-167.641312, -53.719284, 9.398762,
+																-173.112381, -41.697357, 27.737301,
+																-170.130173, -27.639368, 44.721321,
+																-91.032715, -84.610298, -27.737301,
+																-45.644165, -92.280876, -31.038532,
+																-72.473907, -92.728340, -9.398762,
+																0.000000, -94.713943, -32.082241,
+																-23.807196, -98.581642, -11.830645,
+																-50.376068, -96.318359, 9.398762,
+																45.644165, -92.280876, -31.038532,
+																23.807196, -98.581642, -11.830645,
+																0.000000, -99.696220, 7.788675,
+																-25.818338, -95.205017, 27.737299,
+																91.032715, -84.610298, -27.737301,
+																72.473907, -92.728340, -9.398762,
+																50.376068, -96.318359, 9.398762,
+																25.818338, -95.205017, 27.737299,
+																0.000000, -89.442734, 44.721313,
+																132.807678, -69.434662, -27.737301,
+																161.423843, -50.221455, -31.038534,
+																153.984100, -63.118046, -9.398762,
+																180.156631, -29.268211, -32.082249,
+																180.156601, -41.784412, -11.830647,
+																167.641312, -53.719284, 9.398762,
+																189.633484, -6.811248, -31.038534,
+																194.870255, -19.142406, -11.830646,
+																189.633469, -30.807846, 7.788676,
+																173.112381, -41.697357, 27.737301,
+																189.069000, 17.142641, -27.737301,
+																198.775467, 5.808771, -9.398764,
+																198.775467, -5.808771, 9.398764,
+																189.069000, -17.142641, 27.737301,
+																170.130173, -27.639368, 44.721321,
+																173.112381, 41.697357, -27.737301,
+																145.409561, 61.242340, -31.038528,
+																167.641312, 53.719284, -9.398762,
+																111.342873, 76.625198, -32.082233,
+																135.150070, 72.757484, -11.830645,
+																153.984100, 63.118046, 9.398762,
+																71.555756, 88.071297, -31.038528,
+																96.629227, 86.750999, -11.830645,
+																117.199936, 80.655952, 7.788676,
+																132.807678, 69.434662, 27.737301,
+																25.818338, 95.205017, -27.737299,
+																50.376068, 96.318359, -9.398762,
+																72.473907, 92.728340, 9.398762,
+																91.032715, 84.610298, 27.737301,
+																105.146240, 72.360695, 44.721325,
+																0.000000, 99.696220, -7.788675,
+																-23.807196, 98.581635, 11.830645,
+																23.807196, 98.581642, 11.830645,
+																-45.644165, 92.280876, 31.038532,
+																0.000000, 94.713943, 32.082241,
+																45.644165, 92.280876, 31.038532,
+																-65.214394, 80.415871, 49.699692,
+																-22.097843, 84.544273, 52.251183,
+																22.097843, 84.544273, 52.251183,
+																65.214394, 80.415871, 49.699692,
+																-189.633469, 30.807846, -7.788676,
+																-194.870255, 19.142406, 11.830647,
+																-180.156570, 41.784401, 11.830645,
+																-189.633499, 6.811248, 31.038538,
+																-180.156631, 29.268208, 32.082245,
+																-161.423843, 50.221455, 31.038534,
+																-173.112427, -6.161456, 49.699692,
+																-167.641403, 15.617439, 52.251190,
+																-153.984177, 36.633751, 52.251190,
+																-132.807724, 55.861145, 49.699692,
+																-117.199921, -80.655945, -7.788675,
+																-96.629227, -86.750999, 11.830645,
+																-135.150070, -72.757484, 11.830645,
+																-71.555756, -88.071297, 31.038528,
+																-111.342873, -76.625198, 32.082233,
+																-145.409561, -61.242340, 31.038528,
+																-41.774960, -84.223846, 49.699677,
+																-81.510201, -74.892166, 52.251171,
+																-117.265259, -61.903389, 52.251171,
+																-147.294037, -45.891808, 49.699677,
+																117.199936, -80.655952, -7.788676,
+																135.150070, -72.757484, 11.830645,
+																96.629227, -86.750999, 11.830645,
+																145.409561, -61.242340, 31.038528,
+																111.342873, -76.625198, 32.082233,
+																71.555756, -88.071297, 31.038528,
+																147.294037, -45.891808, 49.699677,
+																117.265259, -61.903389, 52.251171,
+																81.510201, -74.892166, 52.251171,
+																41.774960, -84.223846, 49.699677,
+																189.633484, 30.807844, -7.788676,
+																180.156601, 41.784412, 11.830647,
+																194.870255, 19.142406, 11.830646,
+																161.423843, 50.221455, 31.038534,
+																180.156631, 29.268211, 32.082249,
+																189.633499, 6.811248, 31.038538,
+																132.807724, 55.861145, 49.699692,
+																153.984177, 36.633751, 52.251190,
+																167.641403, 15.617439, 52.251190,
+																173.112427, -6.161456, 49.699692,
+																0.000000, 0.000000, 100.000000,
+																25.818336, 17.767946, 97.558472,
+																-25.818336, 17.767946, 97.558472,
+																50.376068, 34.668358, 90.353020,
+																0.000000, 37.619091, 92.654213,
+																-50.376068, 34.668358, 90.353020,
+																72.473923, 49.875896, 78.735497,
+																23.807199, 54.668720, 82.883270,
+																-23.807199, 54.668720, 82.883270,
+																-72.473923, 49.875896, 78.735497,
+																91.032738, 62.647907, 63.273220,
+																45.644176, 69.031013, 68.657646,
+																0.000000, 71.052620, 70.367065,
+																-45.644176, 69.031013, 68.657646,
+																-91.032738, 62.647907, 63.273220,
+																41.774948, -6.786763, 97.558472,
+																81.510193, -13.242155, 90.353020,
+																71.555756, 11.624928, 92.654213,
+																117.265259, -19.050928, 78.735489,
+																111.342911, 5.572543, 82.883270,
+																96.629242, 28.214546, 82.883270,
+																147.294037, -23.929405, 63.273209,
+																145.409622, -0.373371, 68.657646,
+																135.150116, 21.956450, 70.367073,
+																117.199966, 43.036842, 68.657646,
+																0.000000, -21.962393, 97.558456,
+																0.000000, -42.852455, 90.353020,
+																44.223892, -30.434498, 92.654213,
+																0.000000, -61.650002, 78.735481,
+																45.006493, -51.224689, 82.883247,
+																83.527344, -37.231178, 82.883247,
+																0.000000, -77.437073, 63.273209,
+																44.223896, -69.261757, 68.657623,
+																83.527359, -57.482807, 70.367065,
+																118.077721, -42.432804, 68.657631,
+																-41.774948, -6.786763, 97.558472,
+																-81.510193, -13.242155, 90.353020,
+																-44.223892, -30.434498, 92.654213,
+																-117.265259, -19.050928, 78.735489,
+																-83.527344, -37.231178, 82.883247,
+																-45.006493, -51.224689, 82.883247,
+																-147.294037, -23.929405, 63.273209,
+																-118.077721, -42.432804, 68.657631,
+																-83.527359, -57.482807, 70.367065,
+																-44.223896, -69.261757, 68.657623,
+																-71.555756, 11.624928, 92.654213,
+																-96.629242, 28.214546, 82.883270,
+																-111.342911, 5.572543, 82.883270,
+																-117.199966, 43.036842, 68.657646,
+																-135.150116, 21.956450, 70.367073,
+																-145.409622, -0.373371, 68.657646
+	};
+
+	const TArray<int32> EllipsoidGeometry::RawIndicesArray = {
+																2, 1, 0,
+																4, 3, 1,
+																5, 4, 2,
+																7, 6, 3,
+																8, 7, 4,
+																9, 8, 5,
+																11, 10, 6,
+																12, 11, 7,
+																13, 12, 8,
+																14, 13, 9,
+																16, 15, 10,
+																17, 16, 11,
+																18, 17, 12,
+																19, 18, 13,
+																20, 19, 14,
+																2, 4, 1,
+																4, 7, 3,
+																5, 8, 4,
+																7, 11, 6,
+																8, 12, 7,
+																9, 13, 8,
+																11, 16, 10,
+																12, 17, 11,
+																13, 18, 12,
+																14, 19, 13,
+																21, 2, 0,
+																22, 5, 2,
+																23, 22, 21,
+																24, 9, 5,
+																25, 24, 22,
+																26, 25, 23,
+																27, 14, 9,
+																28, 27, 24,
+																29, 28, 25,
+																30, 29, 26,
+																31, 20, 14,
+																32, 31, 27,
+																33, 32, 28,
+																34, 33, 29,
+																35, 34, 30,
+																21, 22, 2,
+																22, 24, 5,
+																23, 25, 22,
+																24, 27, 9,
+																25, 28, 24,
+																26, 29, 25,
+																27, 31, 14,
+																28, 32, 27,
+																29, 33, 28,
+																30, 34, 29,
+																36, 21, 0,
+																37, 23, 21,
+																38, 37, 36,
+																39, 26, 23,
+																40, 39, 37,
+																41, 40, 38,
+																42, 30, 26,
+																43, 42, 39,
+																44, 43, 40,
+																45, 44, 41,
+																46, 35, 30,
+																47, 46, 42,
+																48, 47, 43,
+																49, 48, 44,
+																50, 49, 45,
+																36, 37, 21,
+																37, 39, 23,
+																38, 40, 37,
+																39, 42, 26,
+																40, 43, 39,
+																41, 44, 40,
+																42, 46, 30,
+																43, 47, 42,
+																44, 48, 43,
+																45, 49, 44,
+																51, 36, 0,
+																52, 38, 36,
+																53, 52, 51,
+																54, 41, 38,
+																55, 54, 52,
+																56, 55, 53,
+																57, 45, 41,
+																58, 57, 54,
+																59, 58, 55,
+																60, 59, 56,
+																61, 50, 45,
+																62, 61, 57,
+																63, 62, 58,
+																64, 63, 59,
+																65, 64, 60,
+																51, 52, 36,
+																52, 54, 38,
+																53, 55, 52,
+																54, 57, 41,
+																55, 58, 54,
+																56, 59, 55,
+																57, 61, 45,
+																58, 62, 57,
+																59, 63, 58,
+																60, 64, 59,
+																1, 51, 0,
+																66, 53, 51,
+																3, 66, 1,
+																67, 56, 53,
+																68, 67, 66,
+																6, 68, 3,
+																69, 60, 56,
+																70, 69, 67,
+																71, 70, 68,
+																10, 71, 6,
+																72, 65, 60,
+																73, 72, 69,
+																74, 73, 70,
+																75, 74, 71,
+																15, 75, 10,
+																1, 66, 51,
+																66, 67, 53,
+																3, 68, 66,
+																67, 69, 56,
+																68, 70, 67,
+																6, 71, 68,
+																69, 72, 60,
+																70, 73, 69,
+																71, 74, 70,
+																10, 75, 71,
+																76, 75, 15,
+																77, 74, 75,
+																78, 77, 76,
+																79, 73, 74,
+																80, 79, 77,
+																81, 80, 78,
+																82, 72, 73,
+																83, 82, 79,
+																84, 83, 80,
+																85, 84, 81,
+																86, 65, 72,
+																87, 86, 82,
+																88, 87, 83,
+																89, 88, 84,
+																90, 89, 85,
+																76, 77, 75,
+																77, 79, 74,
+																78, 80, 77,
+																79, 82, 73,
+																80, 83, 79,
+																81, 84, 80,
+																82, 86, 72,
+																83, 87, 82,
+																84, 88, 83,
+																85, 89, 84,
+																91, 64, 65,
+																92, 63, 64,
+																93, 92, 91,
+																94, 62, 63,
+																95, 94, 92,
+																96, 95, 93,
+																97, 61, 62,
+																98, 97, 94,
+																99, 98, 95,
+																100, 99, 96,
+																101, 50, 61,
+																102, 101, 97,
+																103, 102, 98,
+																104, 103, 99,
+																105, 104, 100,
+																91, 92, 64,
+																92, 94, 63,
+																93, 95, 92,
+																94, 97, 62,
+																95, 98, 94,
+																96, 99, 95,
+																97, 101, 61,
+																98, 102, 97,
+																99, 103, 98,
+																100, 104, 99,
+																106, 49, 50,
+																107, 48, 49,
+																108, 107, 106,
+																109, 47, 48,
+																110, 109, 107,
+																111, 110, 108,
+																112, 46, 47,
+																113, 112, 109,
+																114, 113, 110,
+																115, 114, 111,
+																116, 35, 46,
+																117, 116, 112,
+																118, 117, 113,
+																119, 118, 114,
+																120, 119, 115,
+																106, 107, 49,
+																107, 109, 48,
+																108, 110, 107,
+																109, 112, 47,
+																110, 113, 109,
+																111, 114, 110,
+																112, 116, 46,
+																113, 117, 112,
+																114, 118, 113,
+																115, 119, 114,
+																121, 34, 35,
+																122, 33, 34,
+																123, 122, 121,
+																124, 32, 33,
+																125, 124, 122,
+																126, 125, 123,
+																127, 31, 32,
+																128, 127, 124,
+																129, 128, 125,
+																130, 129, 126,
+																131, 20, 31,
+																132, 131, 127,
+																133, 132, 128,
+																134, 133, 129,
+																135, 134, 130,
+																121, 122, 34,
+																122, 124, 33,
+																123, 125, 122,
+																124, 127, 32,
+																125, 128, 124,
+																126, 129, 125,
+																127, 131, 31,
+																128, 132, 127,
+																129, 133, 128,
+																130, 134, 129,
+																136, 19, 20,
+																137, 18, 19,
+																138, 137, 136,
+																139, 17, 18,
+																140, 139, 137,
+																141, 140, 138,
+																142, 16, 17,
+																143, 142, 139,
+																144, 143, 140,
+																145, 144, 141,
+																146, 15, 16,
+																147, 146, 142,
+																148, 147, 143,
+																149, 148, 144,
+																150, 149, 145,
+																136, 137, 19,
+																137, 139, 18,
+																138, 140, 137,
+																139, 142, 17,
+																140, 143, 139,
+																141, 144, 140,
+																142, 146, 16,
+																143, 147, 142,
+																144, 148, 143,
+																145, 149, 144,
+																146, 76, 15,
+																151, 78, 76,
+																147, 151, 146,
+																152, 81, 78,
+																153, 152, 151,
+																148, 153, 147,
+																154, 85, 81,
+																155, 154, 152,
+																156, 155, 153,
+																149, 156, 148,
+																157, 90, 85,
+																158, 157, 154,
+																159, 158, 155,
+																160, 159, 156,
+																150, 160, 149,
+																146, 151, 76,
+																151, 152, 78,
+																147, 153, 151,
+																152, 154, 81,
+																153, 155, 152,
+																148, 156, 153,
+																154, 157, 85,
+																155, 158, 154,
+																156, 159, 155,
+																149, 160, 156,
+																86, 91, 65,
+																161, 93, 91,
+																87, 161, 86,
+																162, 96, 93,
+																163, 162, 161,
+																88, 163, 87,
+																164, 100, 96,
+																165, 164, 162,
+																166, 165, 163,
+																89, 166, 88,
+																167, 105, 100,
+																168, 167, 164,
+																169, 168, 165,
+																170, 169, 166,
+																90, 170, 89,
+																86, 161, 91,
+																161, 162, 93,
+																87, 163, 161,
+																162, 164, 96,
+																163, 165, 162,
+																88, 166, 163,
+																164, 167, 100,
+																165, 168, 164,
+																166, 169, 165,
+																89, 170, 166,
+																101, 106, 50,
+																171, 108, 106,
+																102, 171, 101,
+																172, 111, 108,
+																173, 172, 171,
+																103, 173, 102,
+																174, 115, 111,
+																175, 174, 172,
+																176, 175, 173,
+																104, 176, 103,
+																177, 120, 115,
+																178, 177, 174,
+																179, 178, 175,
+																180, 179, 176,
+																105, 180, 104,
+																101, 171, 106,
+																171, 172, 108,
+																102, 173, 171,
+																172, 174, 111,
+																173, 175, 172,
+																103, 176, 173,
+																174, 177, 115,
+																175, 178, 174,
+																176, 179, 175,
+																104, 180, 176,
+																116, 121, 35,
+																181, 123, 121,
+																117, 181, 116,
+																182, 126, 123,
+																183, 182, 181,
+																118, 183, 117,
+																184, 130, 126,
+																185, 184, 182,
+																186, 185, 183,
+																119, 186, 118,
+																187, 135, 130,
+																188, 187, 184,
+																189, 188, 185,
+																190, 189, 186,
+																120, 190, 119,
+																116, 181, 121,
+																181, 182, 123,
+																117, 183, 181,
+																182, 184, 126,
+																183, 185, 182,
+																118, 186, 183,
+																184, 187, 130,
+																185, 188, 184,
+																186, 189, 185,
+																119, 190, 186,
+																131, 136, 20,
+																191, 138, 136,
+																132, 191, 131,
+																192, 141, 138,
+																193, 192, 191,
+																133, 193, 132,
+																194, 145, 141,
+																195, 194, 192,
+																196, 195, 193,
+																134, 196, 133,
+																197, 150, 145,
+																198, 197, 194,
+																199, 198, 195,
+																200, 199, 196,
+																135, 200, 134,
+																131, 191, 136,
+																191, 192, 138,
+																132, 193, 191,
+																192, 194, 141,
+																193, 195, 192,
+																133, 196, 193,
+																194, 197, 145,
+																195, 198, 194,
+																196, 199, 195,
+																134, 200, 196,
+																203, 202, 201,
+																205, 204, 202,
+																206, 205, 203,
+																208, 207, 204,
+																209, 208, 205,
+																210, 209, 206,
+																212, 211, 207,
+																213, 212, 208,
+																214, 213, 209,
+																215, 214, 210,
+																160, 150, 211,
+																159, 160, 212,
+																158, 159, 213,
+																157, 158, 214,
+																90, 157, 215,
+																203, 205, 202,
+																205, 208, 204,
+																206, 209, 205,
+																208, 212, 207,
+																209, 213, 208,
+																210, 214, 209,
+																212, 160, 211,
+																213, 159, 212,
+																214, 158, 213,
+																215, 157, 214,
+																202, 216, 201,
+																218, 217, 216,
+																204, 218, 202,
+																220, 219, 217,
+																221, 220, 218,
+																207, 221, 204,
+																223, 222, 219,
+																224, 223, 220,
+																225, 224, 221,
+																211, 225, 207,
+																200, 135, 222,
+																199, 200, 223,
+																198, 199, 224,
+																197, 198, 225,
+																150, 197, 211,
+																202, 218, 216,
+																218, 220, 217,
+																204, 221, 218,
+																220, 223, 219,
+																221, 224, 220,
+																207, 225, 221,
+																223, 200, 222,
+																224, 199, 223,
+																225, 198, 224,
+																211, 197, 225,
+																216, 226, 201,
+																228, 227, 226,
+																217, 228, 216,
+																230, 229, 227,
+																231, 230, 228,
+																219, 231, 217,
+																233, 232, 229,
+																234, 233, 230,
+																235, 234, 231,
+																222, 235, 219,
+																190, 120, 232,
+																189, 190, 233,
+																188, 189, 234,
+																187, 188, 235,
+																135, 187, 222,
+																216, 228, 226,
+																228, 230, 227,
+																217, 231, 228,
+																230, 233, 229,
+																231, 234, 230,
+																219, 235, 231,
+																233, 190, 232,
+																234, 189, 233,
+																235, 188, 234,
+																222, 187, 235,
+																226, 236, 201,
+																238, 237, 236,
+																227, 238, 226,
+																240, 239, 237,
+																241, 240, 238,
+																229, 241, 227,
+																243, 242, 239,
+																244, 243, 240,
+																245, 244, 241,
+																232, 245, 229,
+																180, 105, 242,
+																179, 180, 243,
+																178, 179, 244,
+																177, 178, 245,
+																120, 177, 232,
+																226, 238, 236,
+																238, 240, 237,
+																227, 241, 238,
+																240, 243, 239,
+																241, 244, 240,
+																229, 245, 241,
+																243, 180, 242,
+																244, 179, 243,
+																245, 178, 244,
+																232, 177, 245,
+																236, 203, 201,
+																246, 206, 203,
+																237, 246, 236,
+																247, 210, 206,
+																248, 247, 246,
+																239, 248, 237,
+																249, 215, 210,
+																250, 249, 247,
+																251, 250, 248,
+																242, 251, 239,
+																170, 90, 215,
+																169, 170, 249,
+																168, 169, 250,
+																167, 168, 251,
+																105, 167, 242,
+																236, 246, 203,
+																246, 247, 206,
+																237, 248, 246,
+																247, 249, 210,
+																248, 250, 247,
+																239, 251, 248,
+																249, 170, 215,
+																250, 169, 249,
+																251, 168, 250,
+																242, 167, 251
+	};
+
+}
\ No newline at end of file
diff --git a/Engine/Source/Programs/HeadlessChaos/Resource/EllipsoidGeometry2.h b/Engine/Source/Programs/HeadlessChaos/Resource/EllipsoidGeometry2.h
new file mode 100644
--- /dev/null
+++ b/Engine/Source/Programs/HeadlessChaos/Resource/EllipsoidGeometry2.h
@@ -0,0 +1,776 @@
+// Copyright Epic Games, Inc. All Rights Reserved.
+
+#pragma once
+namespace GeometryCollectionTest
+{
+	class EllipsoidGeometry2
+	{
+
+	public:
+		EllipsoidGeometry2() {}
+		~EllipsoidGeometry2() {}
+
+		static const TArray<float>	RawVertexArray;
+		static const TArray<int32>	RawIndicesArray;
+	};
+
+
+	const TArray<float> EllipsoidGeometry2::RawVertexArray = {
+																0.000000, 0.000000, -200.000000,
+																0.000000, 21.962393, -195.116913,
+																20.887474, 6.786763, -195.116943,
+																0.000000, 42.852455, -180.706039,
+																22.111946, 30.434498, -185.308426,
+																40.755096, 13.242155, -180.706039,
+																0.000000, 61.650002, -157.470963,
+																22.503246, 51.224689, -165.766495,
+																41.763672, 37.231178, -165.766495,
+																58.632629, 19.050928, -157.470978,
+																0.000000, 77.437073, -126.546417,
+																22.111948, 69.261757, -137.315247,
+																41.763680, 57.482807, -140.734131,
+																59.038860, 42.432804, -137.315262,
+																73.647018, 23.929405, -126.546417,
+																0.000000, 89.442734, -89.442627,
+																20.887480, 84.223846, -99.399353,
+																40.755100, 74.892166, -104.502342,
+																58.632629, 61.903389, -104.502342,
+																73.647018, 45.891808, -99.399353,
+																85.065086, 27.639368, -89.442642,
+																12.909168, -17.767946, -195.116943,
+																35.777878, -11.624928, -185.308426,
+																25.188034, -34.668358, -180.706039,
+																55.671455, -5.572543, -165.766541,
+																48.314621, -28.214546, -165.766541,
+																36.236961, -49.875896, -157.470993,
+																72.704811, 0.373371, -137.315292,
+																67.575058, -21.956450, -140.734146,
+																58.599983, -43.036842, -137.315292,
+																45.516369, -62.647907, -126.546440,
+																86.556213, 6.161456, -99.399384,
+																83.820702, -15.617439, -104.502380,
+																76.992088, -36.633751, -104.502380,
+																66.403862, -55.861145, -99.399384,
+																52.573120, -72.360695, -89.442650,
+																-12.909168, -17.767946, -195.116943,
+																0.000000, -37.619091, -185.308426,
+																-25.188034, -34.668358, -180.706039,
+																11.903600, -54.668720, -165.766541,
+																-11.903600, -54.668720, -165.766541,
+																-36.236961, -49.875896, -157.470993,
+																22.822088, -69.031013, -137.315292,
+																0.000000, -71.052620, -140.734131,
+																-22.822088, -69.031013, -137.315292,
+																-45.516369, -62.647907, -126.546440,
+																32.607197, -80.415871, -99.399384,
+																11.048922, -84.544273, -104.502365,
+																-11.048922, -84.544273, -104.502365,
+																-32.607197, -80.415871, -99.399384,
+																-52.573120, -72.360695, -89.442650,
+																-20.887474, 6.786763, -195.116943,
+																-35.777878, -11.624928, -185.308426,
+																-40.755096, 13.242155, -180.706039,
+																-48.314621, -28.214546, -165.766541,
+																-55.671455, -5.572543, -165.766541,
+																-58.632629, 19.050928, -157.470978,
+																-58.599983, -43.036842, -137.315292,
+																-67.575058, -21.956450, -140.734146,
+																-72.704811, 0.373371, -137.315292,
+																-73.647018, 23.929405, -126.546417,
+																-66.403862, -55.861145, -99.399384,
+																-76.992088, -36.633751, -104.502380,
+																-83.820702, -15.617439, -104.502380,
+																-86.556213, 6.161456, -99.399384,
+																-85.065086, 27.639368, -89.442642,
+																-22.111946, 30.434498, -185.308426,
+																-41.763672, 37.231178, -165.766495,
+																-22.503246, 51.224689, -165.766495,
+																-59.038860, 42.432804, -137.315262,
+																-41.763680, 57.482807, -140.734131,
+																-22.111948, 69.261757, -137.315247,
+																-73.647018, 45.891808, -99.399353,
+																-58.632629, 61.903389, -104.502342,
+																-40.755100, 74.892166, -104.502342,
+																-20.887480, 84.223846, -99.399353,
+																-12.909169, 95.205017, -55.474598,
+																-35.777878, 88.071297, -62.077057,
+																-25.188034, 96.318359, -18.797523,
+																-55.671436, 76.625198, -64.164467,
+																-48.314613, 86.750999, -23.661289,
+																-36.236954, 92.728340, 18.797523,
+																-72.704781, 61.242340, -62.077057,
+																-67.575035, 72.757484, -23.661289,
+																-58.599968, 80.655952, 15.577353,
+																-45.516357, 84.610298, 55.474602,
+																-86.556190, 41.697357, -55.474602,
+																-83.820656, 53.719284, -18.797523,
+																-76.992050, 63.118046, 18.797523,
+																-66.403839, 69.434662, 55.474602,
+																-52.573120, 72.360695, 89.442650,
+																-94.534500, 17.142641, -55.474602,
+																-94.816742, -6.811248, -62.077068,
+																-99.387733, 5.808771, -18.797527,
+																-90.078316, -29.268211, -64.164497,
+																-97.435127, -19.142406, -23.661291,
+																-99.387733, -5.808771, 18.797527,
+																-80.711922, -50.221455, -62.077068,
+																-90.078300, -41.784412, -23.661295,
+																-94.816734, -30.807846, 15.577353,
+																-94.534500, -17.142641, 55.474602,
+																-66.403839, -69.434662, -55.474602,
+																-76.992050, -63.118046, -18.797523,
+																-83.820656, -53.719284, 18.797523,
+																-86.556190, -41.697357, 55.474602,
+																-85.065086, -27.639368, 89.442642,
+																-45.516357, -84.610298, -55.474602,
+																-22.822083, -92.280876, -62.077065,
+																-36.236954, -92.728340, -18.797523,
+																0.000000, -94.713943, -64.164482,
+																-11.903598, -98.581642, -23.661289,
+																-25.188034, -96.318359, 18.797523,
+																22.822083, -92.280876, -62.077065,
+																11.903598, -98.581642, -23.661289,
+																0.000000, -99.696220, 15.577351,
+																-12.909169, -95.205017, 55.474598,
+																45.516357, -84.610298, -55.474602,
+																36.236954, -92.728340, -18.797523,
+																25.188034, -96.318359, 18.797523,
+																12.909169, -95.205017, 55.474598,
+																0.000000, -89.442734, 89.442627,
+																66.403839, -69.434662, -55.474602,
+																80.711922, -50.221455, -62.077068,
+																76.992050, -63.118046, -18.797523,
+																90.078316, -29.268211, -64.164497,
+																90.078300, -41.784412, -23.661295,
+																83.820656, -53.719284, 18.797523,
+																94.816742, -6.811248, -62.077068,
+																97.435127, -19.142406, -23.661291,
+																94.816734, -30.807846, 15.577353,
+																86.556190, -41.697357, 55.474602,
+																94.534500, 17.142641, -55.474602,
+																99.387733, 5.808771, -18.797527,
+																99.387733, -5.808771, 18.797527,
+																94.534500, -17.142641, 55.474602,
+																85.065086, -27.639368, 89.442642,
+																86.556190, 41.697357, -55.474602,
+																72.704781, 61.242340, -62.077057,
+																83.820656, 53.719284, -18.797523,
+																55.671436, 76.625198, -64.164467,
+																67.575035, 72.757484, -23.661289,
+																76.992050, 63.118046, 18.797523,
+																35.777878, 88.071297, -62.077057,
+																48.314613, 86.750999, -23.661289,
+																58.599968, 80.655952, 15.577353,
+																66.403839, 69.434662, 55.474602,
+																12.909169, 95.205017, -55.474598,
+																25.188034, 96.318359, -18.797523,
+																36.236954, 92.728340, 18.797523,
+																45.516357, 84.610298, 55.474602,
+																52.573120, 72.360695, 89.442650,
+																0.000000, 99.696220, -15.577351,
+																-11.903598, 98.581635, 23.661289,
+																11.903598, 98.581642, 23.661289,
+																-22.822083, 92.280876, 62.077065,
+																0.000000, 94.713943, 64.164482,
+																22.822083, 92.280876, 62.077065,
+																-32.607197, 80.415871, 99.399384,
+																-11.048922, 84.544273, 104.502365,
+																11.048922, 84.544273, 104.502365,
+																32.607197, 80.415871, 99.399384,
+																-94.816734, 30.807846, -15.577353,
+																-97.435127, 19.142406, 23.661295,
+																-90.078285, 41.784401, 23.661289,
+																-94.816750, 6.811248, 62.077076,
+																-90.078316, 29.268208, 64.164490,
+																-80.711922, 50.221455, 62.077068,
+																-86.556213, -6.161456, 99.399384,
+																-83.820702, 15.617439, 104.502380,
+																-76.992088, 36.633751, 104.502380,
+																-66.403862, 55.861145, 99.399384,
+																-58.599960, -80.655945, -15.577351,
+																-48.314613, -86.750999, 23.661289,
+																-67.575035, -72.757484, 23.661289,
+																-35.777878, -88.071297, 62.077057,
+																-55.671436, -76.625198, 64.164467,
+																-72.704781, -61.242340, 62.077057,
+																-20.887480, -84.223846, 99.399353,
+																-40.755100, -74.892166, 104.502342,
+																-58.632629, -61.903389, 104.502342,
+																-73.647018, -45.891808, 99.399353,
+																58.599968, -80.655952, -15.577353,
+																67.575035, -72.757484, 23.661289,
+																48.314613, -86.750999, 23.661289,
+																72.704781, -61.242340, 62.077057,
+																55.671436, -76.625198, 64.164467,
+																35.777878, -88.071297, 62.077057,
+																73.647018, -45.891808, 99.399353,
+																58.632629, -61.903389, 104.502342,
+																40.755100, -74.892166, 104.502342,
+																20.887480, -84.223846, 99.399353,
+																94.816742, 30.807844, -15.577353,
+																90.078300, 41.784412, 23.661295,
+																97.435127, 19.142406, 23.661291,
+																80.711922, 50.221455, 62.077068,
+																90.078316, 29.268211, 64.164497,
+																94.816750, 6.811248, 62.077076,
+																66.403862, 55.861145, 99.399384,
+																76.992088, 36.633751, 104.502380,
+																83.820702, 15.617439, 104.502380,
+																86.556213, -6.161456, 99.399384,
+																0.000000, 0.000000, 200.000000,
+																12.909168, 17.767946, 195.116943,
+																-12.909168, 17.767946, 195.116943,
+																25.188034, 34.668358, 180.706039,
+																0.000000, 37.619091, 185.308426,
+																-25.188034, 34.668358, 180.706039,
+																36.236961, 49.875896, 157.470993,
+																11.903600, 54.668720, 165.766541,
+																-11.903600, 54.668720, 165.766541,
+																-36.236961, 49.875896, 157.470993,
+																45.516369, 62.647907, 126.546440,
+																22.822088, 69.031013, 137.315292,
+																0.000000, 71.052620, 140.734131,
+																-22.822088, 69.031013, 137.315292,
+																-45.516369, 62.647907, 126.546440,
+																20.887474, -6.786763, 195.116943,
+																40.755096, -13.242155, 180.706039,
+																35.777878, 11.624928, 185.308426,
+																58.632629, -19.050928, 157.470978,
+																55.671455, 5.572543, 165.766541,
+																48.314621, 28.214546, 165.766541,
+																73.647018, -23.929405, 126.546417,
+																72.704811, -0.373371, 137.315292,
+																67.575058, 21.956450, 140.734146,
+																58.599983, 43.036842, 137.315292,
+																0.000000, -21.962393, 195.116913,
+																0.000000, -42.852455, 180.706039,
+																22.111946, -30.434498, 185.308426,
+																0.000000, -61.650002, 157.470963,
+																22.503246, -51.224689, 165.766495,
+																41.763672, -37.231178, 165.766495,
+																0.000000, -77.437073, 126.546417,
+																22.111948, -69.261757, 137.315247,
+																41.763680, -57.482807, 140.734131,
+																59.038860, -42.432804, 137.315262,
+																-20.887474, -6.786763, 195.116943,
+																-40.755096, -13.242155, 180.706039,
+																-22.111946, -30.434498, 185.308426,
+																-58.632629, -19.050928, 157.470978,
+																-41.763672, -37.231178, 165.766495,
+																-22.503246, -51.224689, 165.766495,
+																-73.647018, -23.929405, 126.546417,
+																-59.038860, -42.432804, 137.315262,
+																-41.763680, -57.482807, 140.734131,
+																-22.111948, -69.261757, 137.315247,
+																-35.777878, 11.624928, 185.308426,
+																-48.314621, 28.214546, 165.766541,
+																-55.671455, 5.572543, 165.766541,
+																-58.599983, 43.036842, 137.315292,
+																-67.575058, 21.956450, 140.734146,
+																-72.704811, -0.373371, 137.315292
+	};
+
+	const TArray<int32> EllipsoidGeometry2::RawIndicesArray = {
+																2, 1, 0,
+																4, 3, 1,
+																5, 4, 2,
+																7, 6, 3,
+																8, 7, 4,
+																9, 8, 5,
+																11, 10, 6,
+																12, 11, 7,
+																13, 12, 8,
+																14, 13, 9,
+																16, 15, 10,
+																17, 16, 11,
+																18, 17, 12,
+																19, 18, 13,
+																20, 19, 14,
+																2, 4, 1,
+																4, 7, 3,
+																5, 8, 4,
+																7, 11, 6,
+																8, 12, 7,
+																9, 13, 8,
+																11, 16, 10,
+																12, 17, 11,
+																13, 18, 12,
+																14, 19, 13,
+																21, 2, 0,
+																22, 5, 2,
+																23, 22, 21,
+																24, 9, 5,
+																25, 24, 22,
+																26, 25, 23,
+																27, 14, 9,
+																28, 27, 24,
+																29, 28, 25,
+																30, 29, 26,
+																31, 20, 14,
+																32, 31, 27,
+																33, 32, 28,
+																34, 33, 29,
+																35, 34, 30,
+																21, 22, 2,
+																22, 24, 5,
+																23, 25, 22,
+																24, 27, 9,
+																25, 28, 24,
+																26, 29, 25,
+																27, 31, 14,
+																28, 32, 27,
+																29, 33, 28,
+																30, 34, 29,
+																36, 21, 0,
+																37, 23, 21,
+																38, 37, 36,
+																39, 26, 23,
+																40, 39, 37,
+																41, 40, 38,
+																42, 30, 26,
+																43, 42, 39,
+																44, 43, 40,
+																45, 44, 41,
+																46, 35, 30,
+																47, 46, 42,
+																48, 47, 43,
+																49, 48, 44,
+																50, 49, 45,
+																36, 37, 21,
+																37, 39, 23,
+																38, 40, 37,
+																39, 42, 26,
+																40, 43, 39,
+																41, 44, 40,
+																42, 46, 30,
+																43, 47, 42,
+																44, 48, 43,
+																45, 49, 44,
+																51, 36, 0,
+																52, 38, 36,
+																53, 52, 51,
+																54, 41, 38,
+																55, 54, 52,
+																56, 55, 53,
+																57, 45, 41,
+																58, 57, 54,
+																59, 58, 55,
+																60, 59, 56,
+																61, 50, 45,
+																62, 61, 57,
+																63, 62, 58,
+																64, 63, 59,
+																65, 64, 60,
+																51, 52, 36,
+																52, 54, 38,
+																53, 55, 52,
+																54, 57, 41,
+																55, 58, 54,
+																56, 59, 55,
+																57, 61, 45,
+																58, 62, 57,
+																59, 63, 58,
+																60, 64, 59,
+																1, 51, 0,
+																66, 53, 51,
+																3, 66, 1,
+																67, 56, 53,
+																68, 67, 66,
+																6, 68, 3,
+																69, 60, 56,
+																70, 69, 67,
+																71, 70, 68,
+																10, 71, 6,
+																72, 65, 60,
+																73, 72, 69,
+																74, 73, 70,
+																75, 74, 71,
+																15, 75, 10,
+																1, 66, 51,
+																66, 67, 53,
+																3, 68, 66,
+																67, 69, 56,
+																68, 70, 67,
+																6, 71, 68,
+																69, 72, 60,
+																70, 73, 69,
+																71, 74, 70,
+																10, 75, 71,
+																76, 75, 15,
+																77, 74, 75,
+																78, 77, 76,
+																79, 73, 74,
+																80, 79, 77,
+																81, 80, 78,
+																82, 72, 73,
+																83, 82, 79,
+																84, 83, 80,
+																85, 84, 81,
+																86, 65, 72,
+																87, 86, 82,
+																88, 87, 83,
+																89, 88, 84,
+																90, 89, 85,
+																76, 77, 75,
+																77, 79, 74,
+																78, 80, 77,
+																79, 82, 73,
+																80, 83, 79,
+																81, 84, 80,
+																82, 86, 72,
+																83, 87, 82,
+																84, 88, 83,
+																85, 89, 84,
+																91, 64, 65,
+																92, 63, 64,
+																93, 92, 91,
+																94, 62, 63,
+																95, 94, 92,
+																96, 95, 93,
+																97, 61, 62,
+																98, 97, 94,
+																99, 98, 95,
+																100, 99, 96,
+																101, 50, 61,
+																102, 101, 97,
+																103, 102, 98,
+																104, 103, 99,
+																105, 104, 100,
+																91, 92, 64,
+																92, 94, 63,
+																93, 95, 92,
+																94, 97, 62,
+																95, 98, 94,
+																96, 99, 95,
+																97, 101, 61,
+																98, 102, 97,
+																99, 103, 98,
+																100, 104, 99,
+																106, 49, 50,
+																107, 48, 49,
+																108, 107, 106,
+																109, 47, 48,
+																110, 109, 107,
+																111, 110, 108,
+																112, 46, 47,
+																113, 112, 109,
+																114, 113, 110,
+																115, 114, 111,
+																116, 35, 46,
+																117, 116, 112,
+																118, 117, 113,
+																119, 118, 114,
+																120, 119, 115,
+																106, 107, 49,
+																107, 109, 48,
+																108, 110, 107,
+																109, 112, 47,
+																110, 113, 109,
+																111, 114, 110,
+																112, 116, 46,
+																113, 117, 112,
+																114, 118, 113,
+																115, 119, 114,
+																121, 34, 35,
+																122, 33, 34,
+																123, 122, 121,
+																124, 32, 33,
+																125, 124, 122,
+																126, 125, 123,
+																127, 31, 32,
+																128, 127, 124,
+																129, 128, 125,
+																130, 129, 126,
+																131, 20, 31,
+																132, 131, 127,
+																133, 132, 128,
+																134, 133, 129,
+																135, 134, 130,
+																121, 122, 34,
+																122, 124, 33,
+																123, 125, 122,
+																124, 127, 32,
+																125, 128, 124,
+																126, 129, 125,
+																127, 131, 31,
+																128, 132, 127,
+																129, 133, 128,
+																130, 134, 129,
+																136, 19, 20,
+																137, 18, 19,
+																138, 137, 136,
+																139, 17, 18,
+																140, 139, 137,
+																141, 140, 138,
+																142, 16, 17,
+																143, 142, 139,
+																144, 143, 140,
+																145, 144, 141,
+																146, 15, 16,
+																147, 146, 142,
+																148, 147, 143,
+																149, 148, 144,
+																150, 149, 145,
+																136, 137, 19,
+																137, 139, 18,
+																138, 140, 137,
+																139, 142, 17,
+																140, 143, 139,
+																141, 144, 140,
+																142, 146, 16,
+																143, 147, 142,
+																144, 148, 143,
+																145, 149, 144,
+																146, 76, 15,
+																151, 78, 76,
+																147, 151, 146,
+																152, 81, 78,
+																153, 152, 151,
+																148, 153, 147,
+																154, 85, 81,
+																155, 154, 152,
+																156, 155, 153,
+																149, 156, 148,
+																157, 90, 85,
+																158, 157, 154,
+																159, 158, 155,
+																160, 159, 156,
+																150, 160, 149,
+																146, 151, 76,
+																151, 152, 78,
+																147, 153, 151,
+																152, 154, 81,
+																153, 155, 152,
+																148, 156, 153,
+																154, 157, 85,
+																155, 158, 154,
+																156, 159, 155,
+																149, 160, 156,
+																86, 91, 65,
+																161, 93, 91,
+																87, 161, 86,
+																162, 96, 93,
+																163, 162, 161,
+																88, 163, 87,
+																164, 100, 96,
+																165, 164, 162,
+																166, 165, 163,
+																89, 166, 88,
+																167, 105, 100,
+																168, 167, 164,
+																169, 168, 165,
+																170, 169, 166,
+																90, 170, 89,
+																86, 161, 91,
+																161, 162, 93,
+																87, 163, 161,
+																162, 164, 96,
+																163, 165, 162,
+																88, 166, 163,
+																164, 167, 100,
+																165, 168, 164,
+																166, 169, 165,
+																89, 170, 166,
+																101, 106, 50,
+																171, 108, 106,
+																102, 171, 101,
+																172, 111, 108,
+																173, 172, 171,
+																103, 173, 102,
+																174, 115, 111,
+																175, 174, 172,
+																176, 175, 173,
+																104, 176, 103,
+																177, 120, 115,
+																178, 177, 174,
+																179, 178, 175,
+																180, 179, 176,
+																105, 180, 104,
+																101, 171, 106,
+																171, 172, 108,
+																102, 173, 171,
+																172, 174, 111,
+																173, 175, 172,
+																103, 176, 173,
+																174, 177, 115,
+																175, 178, 174,
+																176, 179, 175,
+																104, 180, 176,
+																116, 121, 35,
+																181, 123, 121,
+																117, 181, 116,
+																182, 126, 123,
+																183, 182, 181,
+																118, 183, 117,
+																184, 130, 126,
+																185, 184, 182,
+																186, 185, 183,
+																119, 186, 118,
+																187, 135, 130,
+																188, 187, 184,
+																189, 188, 185,
+																190, 189, 186,
+																120, 190, 119,
+																116, 181, 121,
+																181, 182, 123,
+																117, 183, 181,
+																182, 184, 126,
+																183, 185, 182,
+																118, 186, 183,
+																184, 187, 130,
+																185, 188, 184,
+																186, 189, 185,
+																119, 190, 186,
+																131, 136, 20,
+																191, 138, 136,
+																132, 191, 131,
+																192, 141, 138,
+																193, 192, 191,
+																133, 193, 132,
+																194, 145, 141,
+																195, 194, 192,
+																196, 195, 193,
+																134, 196, 133,
+																197, 150, 145,
+																198, 197, 194,
+																199, 198, 195,
+																200, 199, 196,
+																135, 200, 134,
+																131, 191, 136,
+																191, 192, 138,
+																132, 193, 191,
+																192, 194, 141,
+																193, 195, 192,
+																133, 196, 193,
+																194, 197, 145,
+																195, 198, 194,
+																196, 199, 195,
+																134, 200, 196,
+																203, 202, 201,
+																205, 204, 202,
+																206, 205, 203,
+																208, 207, 204,
+																209, 208, 205,
+																210, 209, 206,
+																212, 211, 207,
+																213, 212, 208,
+																214, 213, 209,
+																215, 214, 210,
+																160, 150, 211,
+																159, 160, 212,
+																158, 159, 213,
+																157, 158, 214,
+																90, 157, 215,
+																203, 205, 202,
+																205, 208, 204,
+																206, 209, 205,
+																208, 212, 207,
+																209, 213, 208,
+																210, 214, 209,
+																212, 160, 211,
+																213, 159, 212,
+																214, 158, 213,
+																215, 157, 214,
+																202, 216, 201,
+																218, 217, 216,
+																204, 218, 202,
+																220, 219, 217,
+																221, 220, 218,
+																207, 221, 204,
+																223, 222, 219,
+																224, 223, 220,
+																225, 224, 221,
+																211, 225, 207,
+																200, 135, 222,
+																199, 200, 223,
+																198, 199, 224,
+																197, 198, 225,
+																150, 197, 211,
+																202, 218, 216,
+																218, 220, 217,
+																204, 221, 218,
+																220, 223, 219,
+																221, 224, 220,
+																207, 225, 221,
+																223, 200, 222,
+																224, 199, 223,
+																225, 198, 224,
+																211, 197, 225,
+																216, 226, 201,
+																228, 227, 226,
+																217, 228, 216,
+																230, 229, 227,
+																231, 230, 228,
+																219, 231, 217,
+																233, 232, 229,
+																234, 233, 230,
+																235, 234, 231,
+																222, 235, 219,
+																190, 120, 232,
+																189, 190, 233,
+																188, 189, 234,
+																187, 188, 235,
+																135, 187, 222,
+																216, 228, 226,
+																228, 230, 227,
+																217, 231, 228,
+																230, 233, 229,
+																231, 234, 230,
+																219, 235, 231,
+																233, 190, 232,
+																234, 189, 233,
+																235, 188, 234,
+																222, 187, 235,
+																226, 236, 201,
+																238, 237, 236,
+																227, 238, 226,
+																240, 239, 237,
+																241, 240, 238,
+																229, 241, 227,
+																243, 242, 239,
+																244, 243, 240,
+																245, 244, 241,
+																232, 245, 229,
+																180, 105, 242,
+																179, 180, 243,
+																178, 179, 244,
+																177, 178, 245,
+																120, 177, 232,
+																226, 238, 236,
+																238, 240, 237,
+																227, 241, 238,
+																240, 243, 239,
+																241, 244, 240,
+																229, 245, 241,
+																243, 180, 242,
+																244, 179, 243,
+																245, 178, 244,
+																232, 177, 245,
+																236, 203, 201,
+																246, 206, 203,
+																237, 246, 236,
+																247, 210, 206,
+																248, 247, 246,
+																239, 248, 237,
+																249, 215, 210,
+																250, 249, 247,
+																251, 250, 248,
+																242, 251, 239,
+																170, 90, 215,
+																169, 170, 249,
+																168, 169, 250,
+																167, 168, 251,
+																105, 167, 242,
+																236, 246, 203,
+																246, 247, 206,
+																237, 248, 246,
+																247, 249, 210,
+																248, 250, 247,
+																239, 251, 248,
+																249, 170, 215,
+																250, 169, 249,
+																251, 168, 250,
+																242, 167, 251
+	};
+
+}
\ No newline at end of file
