Merge geometry collection into headless chaos.

diff --git a/Engine/Source/Programs/HeadlessChaos/Resource/EllipsoidGeometry3.h b/Engine/Source/Programs/HeadlessChaos/Resource/EllipsoidGeometry3.h
new file mode 100644
--- /dev/null
+++ b/Engine/Source/Programs/HeadlessChaos/Resource/EllipsoidGeometry3.h
@@ -0,0 +1,776 @@
+// Copyright Epic Games, Inc. All Rights Reserved.
+
+#pragma once
+namespace GeometryCollectionTest
+{
+	class EllipsoidGeometry3
+	{
+
+	public:
+		EllipsoidGeometry3() {}
+		~EllipsoidGeometry3() {}
+
+		static const TArray<float>	RawVertexArray;
+		static const TArray<int32>	RawIndicesArray;
+	};
+
+
+	const TArray<float> EllipsoidGeometry3::RawVertexArray = {
+																0.000000, 0.000000, -250.000000,
+																0.000000, 21.962393, -243.896149,
+																31.331211, 6.786763, -243.896179,
+																0.000000, 42.852455, -225.882553,
+																33.167919, 30.434498, -231.635529,
+																61.132645, 13.242155, -225.882553,
+																0.000000, 61.650002, -196.838699,
+																33.754868, 51.224689, -207.208115,
+																62.645508, 37.231178, -207.208115,
+																87.948944, 19.050928, -196.838715,
+																0.000000, 77.437073, -158.183014,
+																33.167923, 69.261757, -171.644058,
+																62.645519, 57.482807, -175.917664,
+																88.558289, 42.432804, -171.644073,
+																110.470528, 23.929405, -158.183014,
+																0.000000, 89.442734, -111.803284,
+																31.331219, 84.223846, -124.249191,
+																61.132652, 74.892166, -130.627930,
+																87.948944, 61.903389, -130.627930,
+																110.470528, 45.891808, -124.249191,
+																127.597626, 27.639368, -111.803299,
+																19.363752, -17.767946, -243.896179,
+																53.666817, -11.624928, -231.635529,
+																37.782051, -34.668358, -225.882553,
+																83.507187, -5.572543, -207.208176,
+																72.471931, -28.214546, -207.208176,
+																54.355442, -49.875896, -196.838745,
+																109.057220, 0.373371, -171.644119,
+																101.362587, -21.956450, -175.917679,
+																87.899979, -43.036842, -171.644119,
+																68.274551, -62.647907, -158.183044,
+																129.834320, 6.161456, -124.249229,
+																125.731049, -15.617439, -130.627975,
+																115.488129, -36.633751, -130.627975,
+																99.605789, -55.861145, -124.249229,
+																78.859680, -72.360695, -111.803314,
+																-19.363752, -17.767946, -243.896179,
+																0.000000, -37.619091, -231.635529,
+																-37.782051, -34.668358, -225.882553,
+																17.855400, -54.668720, -207.208176,
+																-17.855400, -54.668720, -207.208176,
+																-54.355442, -49.875896, -196.838745,
+																34.233131, -69.031013, -171.644119,
+																0.000000, -71.052620, -175.917664,
+																-34.233131, -69.031013, -171.644119,
+																-68.274551, -62.647907, -158.183044,
+																48.910797, -80.415871, -124.249229,
+																16.573383, -84.544273, -130.627960,
+																-16.573383, -84.544273, -130.627960,
+																-48.910797, -80.415871, -124.249229,
+																-78.859680, -72.360695, -111.803314,
+																-31.331211, 6.786763, -243.896179,
+																-53.666817, -11.624928, -231.635529,
+																-61.132645, 13.242155, -225.882553,
+																-72.471931, -28.214546, -207.208176,
+																-83.507187, -5.572543, -207.208176,
+																-87.948944, 19.050928, -196.838715,
+																-87.899979, -43.036842, -171.644119,
+																-101.362587, -21.956450, -175.917679,
+																-109.057220, 0.373371, -171.644119,
+																-110.470528, 23.929405, -158.183014,
+																-99.605789, -55.861145, -124.249229,
+																-115.488129, -36.633751, -130.627975,
+																-125.731049, -15.617439, -130.627975,
+																-129.834320, 6.161456, -124.249229,
+																-127.597626, 27.639368, -111.803299,
+																-33.167919, 30.434498, -231.635529,
+																-62.645508, 37.231178, -207.208115,
+																-33.754868, 51.224689, -207.208115,
+																-88.558289, 42.432804, -171.644073,
+																-62.645519, 57.482807, -175.917664,
+																-33.167923, 69.261757, -171.644058,
+																-110.470528, 45.891808, -124.249191,
+																-87.948944, 61.903389, -130.627930,
+																-61.132652, 74.892166, -130.627930,
+																-31.331219, 84.223846, -124.249191,
+																-19.363754, 95.205017, -69.343246,
+																-53.666817, 88.071297, -77.596321,
+																-37.782051, 96.318359, -23.496904,
+																-83.507156, 76.625198, -80.205582,
+																-72.471924, 86.750999, -29.576611,
+																-54.355431, 92.728340, 23.496904,
+																-109.057175, 61.242340, -77.596321,
+																-101.362549, 72.757484, -29.576611,
+																-87.899948, 80.655952, 19.471691,
+																-68.274536, 84.610298, 69.343254,
+																-129.834290, 41.697357, -69.343254,
+																-125.730988, 53.719284, -23.496904,
+																-115.488075, 63.118046, 23.496904,
+																-99.605759, 69.434662, 69.343254,
+																-78.859680, 72.360695, 111.803314,
+																-141.801758, 17.142641, -69.343254,
+																-142.225113, -6.811248, -77.596336,
+																-149.081604, 5.808771, -23.496910,
+																-135.117477, -29.268211, -80.205620,
+																-146.152695, -19.142406, -29.576614,
+																-149.081604, -5.808771, 23.496910,
+																-121.067886, -50.221455, -77.596336,
+																-135.117447, -41.784412, -29.576618,
+																-142.225098, -30.807846, 19.471691,
+																-141.801758, -17.142641, 69.343254,
+																-99.605759, -69.434662, -69.343254,
+																-115.488075, -63.118046, -23.496904,
+																-125.730988, -53.719284, 23.496904,
+																-129.834290, -41.697357, 69.343254,
+																-127.597626, -27.639368, 111.803299,
+																-68.274536, -84.610298, -69.343254,
+																-34.233124, -92.280876, -77.596329,
+																-54.355431, -92.728340, -23.496904,
+																0.000000, -94.713943, -80.205605,
+																-17.855396, -98.581642, -29.576611,
+																-37.782051, -96.318359, 23.496904,
+																34.233124, -92.280876, -77.596329,
+																17.855396, -98.581642, -29.576611,
+																0.000000, -99.696220, 19.471687,
+																-19.363754, -95.205017, 69.343246,
+																68.274536, -84.610298, -69.343254,
+																54.355431, -92.728340, -23.496904,
+																37.782051, -96.318359, 23.496904,
+																19.363754, -95.205017, 69.343246,
+																0.000000, -89.442734, 111.803284,
+																99.605759, -69.434662, -69.343254,
+																121.067886, -50.221455, -77.596336,
+																115.488075, -63.118046, -23.496904,
+																135.117477, -29.268211, -80.205620,
+																135.117447, -41.784412, -29.576618,
+																125.730988, -53.719284, 23.496904,
+																142.225113, -6.811248, -77.596336,
+																146.152695, -19.142406, -29.576614,
+																142.225098, -30.807846, 19.471691,
+																129.834290, -41.697357, 69.343254,
+																141.801758, 17.142641, -69.343254,
+																149.081604, 5.808771, -23.496910,
+																149.081604, -5.808771, 23.496910,
+																141.801758, -17.142641, 69.343254,
+																127.597626, -27.639368, 111.803299,
+																129.834290, 41.697357, -69.343254,
+																109.057175, 61.242340, -77.596321,
+																125.730988, 53.719284, -23.496904,
+																83.507156, 76.625198, -80.205582,
+																101.362549, 72.757484, -29.576611,
+																115.488075, 63.118046, 23.496904,
+																53.666817, 88.071297, -77.596321,
+																72.471924, 86.750999, -29.576611,
+																87.899948, 80.655952, 19.471691,
+																99.605759, 69.434662, 69.343254,
+																19.363754, 95.205017, -69.343246,
+																37.782051, 96.318359, -23.496904,
+																54.355431, 92.728340, 23.496904,
+																68.274536, 84.610298, 69.343254,
+																78.859680, 72.360695, 111.803314,
+																0.000000, 99.696220, -19.471687,
+																-17.855396, 98.581635, 29.576611,
+																17.855396, 98.581642, 29.576611,
+																-34.233124, 92.280876, 77.596329,
+																0.000000, 94.713943, 80.205605,
+																34.233124, 92.280876, 77.596329,
+																-48.910797, 80.415871, 124.249229,
+																-16.573383, 84.544273, 130.627960,
+																16.573383, 84.544273, 130.627960,
+																48.910797, 80.415871, 124.249229,
+																-142.225098, 30.807846, -19.471691,
+																-146.152695, 19.142406, 29.576618,
+																-135.117432, 41.784401, 29.576611,
+																-142.225128, 6.811248, 77.596344,
+																-135.117477, 29.268208, 80.205612,
+																-121.067886, 50.221455, 77.596336,
+																-129.834320, -6.161456, 124.249229,
+																-125.731049, 15.617439, 130.627975,
+																-115.488129, 36.633751, 130.627975,
+																-99.605789, 55.861145, 124.249229,
+																-87.899940, -80.655945, -19.471687,
+																-72.471924, -86.750999, 29.576611,
+																-101.362549, -72.757484, 29.576611,
+																-53.666817, -88.071297, 77.596321,
+																-83.507156, -76.625198, 80.205582,
+																-109.057175, -61.242340, 77.596321,
+																-31.331219, -84.223846, 124.249191,
+																-61.132652, -74.892166, 130.627930,
+																-87.948944, -61.903389, 130.627930,
+																-110.470528, -45.891808, 124.249191,
+																87.899948, -80.655952, -19.471691,
+																101.362549, -72.757484, 29.576611,
+																72.471924, -86.750999, 29.576611,
+																109.057175, -61.242340, 77.596321,
+																83.507156, -76.625198, 80.205582,
+																53.666817, -88.071297, 77.596321,
+																110.470528, -45.891808, 124.249191,
+																87.948944, -61.903389, 130.627930,
+																61.132652, -74.892166, 130.627930,
+																31.331219, -84.223846, 124.249191,
+																142.225113, 30.807844, -19.471691,
+																135.117447, 41.784412, 29.576618,
+																146.152695, 19.142406, 29.576614,
+																121.067886, 50.221455, 77.596336,
+																135.117477, 29.268211, 80.205620,
+																142.225128, 6.811248, 77.596344,
+																99.605789, 55.861145, 124.249229,
+																115.488129, 36.633751, 130.627975,
+																125.731049, 15.617439, 130.627975,
+																129.834320, -6.161456, 124.249229,
+																0.000000, 0.000000, 250.000000,
+																19.363752, 17.767946, 243.896179,
+																-19.363752, 17.767946, 243.896179,
+																37.782051, 34.668358, 225.882553,
+																0.000000, 37.619091, 231.635529,
+																-37.782051, 34.668358, 225.882553,
+																54.355442, 49.875896, 196.838745,
+																17.855400, 54.668720, 207.208176,
+																-17.855400, 54.668720, 207.208176,
+																-54.355442, 49.875896, 196.838745,
+																68.274551, 62.647907, 158.183044,
+																34.233131, 69.031013, 171.644119,
+																0.000000, 71.052620, 175.917664,
+																-34.233131, 69.031013, 171.644119,
+																-68.274551, 62.647907, 158.183044,
+																31.331211, -6.786763, 243.896179,
+																61.132645, -13.242155, 225.882553,
+																53.666817, 11.624928, 231.635529,
+																87.948944, -19.050928, 196.838715,
+																83.507187, 5.572543, 207.208176,
+																72.471931, 28.214546, 207.208176,
+																110.470528, -23.929405, 158.183014,
+																109.057220, -0.373371, 171.644119,
+																101.362587, 21.956450, 175.917679,
+																87.899979, 43.036842, 171.644119,
+																0.000000, -21.962393, 243.896149,
+																0.000000, -42.852455, 225.882553,
+																33.167919, -30.434498, 231.635529,
+																0.000000, -61.650002, 196.838699,
+																33.754868, -51.224689, 207.208115,
+																62.645508, -37.231178, 207.208115,
+																0.000000, -77.437073, 158.183014,
+																33.167923, -69.261757, 171.644058,
+																62.645519, -57.482807, 175.917664,
+																88.558289, -42.432804, 171.644073,
+																-31.331211, -6.786763, 243.896179,
+																-61.132645, -13.242155, 225.882553,
+																-33.167919, -30.434498, 231.635529,
+																-87.948944, -19.050928, 196.838715,
+																-62.645508, -37.231178, 207.208115,
+																-33.754868, -51.224689, 207.208115,
+																-110.470528, -23.929405, 158.183014,
+																-88.558289, -42.432804, 171.644073,
+																-62.645519, -57.482807, 175.917664,
+																-33.167923, -69.261757, 171.644058,
+																-53.666817, 11.624928, 231.635529,
+																-72.471931, 28.214546, 207.208176,
+																-83.507187, 5.572543, 207.208176,
+																-87.899979, 43.036842, 171.644119,
+																-101.362587, 21.956450, 175.917679,
+																-109.057220, -0.373371, 171.644119
+	};
+
+	const TArray<int32> EllipsoidGeometry3::RawIndicesArray = {
+																2, 1, 0,
+																4, 3, 1,
+																5, 4, 2,
+																7, 6, 3,
+																8, 7, 4,
+																9, 8, 5,
+																11, 10, 6,
+																12, 11, 7,
+																13, 12, 8,
+																14, 13, 9,
+																16, 15, 10,
+																17, 16, 11,
+																18, 17, 12,
+																19, 18, 13,
+																20, 19, 14,
+																2, 4, 1,
+																4, 7, 3,
+																5, 8, 4,
+																7, 11, 6,
+																8, 12, 7,
+																9, 13, 8,
+																11, 16, 10,
+																12, 17, 11,
+																13, 18, 12,
+																14, 19, 13,
+																21, 2, 0,
+																22, 5, 2,
+																23, 22, 21,
+																24, 9, 5,
+																25, 24, 22,
+																26, 25, 23,
+																27, 14, 9,
+																28, 27, 24,
+																29, 28, 25,
+																30, 29, 26,
+																31, 20, 14,
+																32, 31, 27,
+																33, 32, 28,
+																34, 33, 29,
+																35, 34, 30,
+																21, 22, 2,
+																22, 24, 5,
+																23, 25, 22,
+																24, 27, 9,
+																25, 28, 24,
+																26, 29, 25,
+																27, 31, 14,
+																28, 32, 27,
+																29, 33, 28,
+																30, 34, 29,
+																36, 21, 0,
+																37, 23, 21,
+																38, 37, 36,
+																39, 26, 23,
+																40, 39, 37,
+																41, 40, 38,
+																42, 30, 26,
+																43, 42, 39,
+																44, 43, 40,
+																45, 44, 41,
+																46, 35, 30,
+																47, 46, 42,
+																48, 47, 43,
+																49, 48, 44,
+																50, 49, 45,
+																36, 37, 21,
+																37, 39, 23,
+																38, 40, 37,
+																39, 42, 26,
+																40, 43, 39,
+																41, 44, 40,
+																42, 46, 30,
+																43, 47, 42,
+																44, 48, 43,
+																45, 49, 44,
+																51, 36, 0,
+																52, 38, 36,
+																53, 52, 51,
+																54, 41, 38,
+																55, 54, 52,
+																56, 55, 53,
+																57, 45, 41,
+																58, 57, 54,
+																59, 58, 55,
+																60, 59, 56,
+																61, 50, 45,
+																62, 61, 57,
+																63, 62, 58,
+																64, 63, 59,
+																65, 64, 60,
+																51, 52, 36,
+																52, 54, 38,
+																53, 55, 52,
+																54, 57, 41,
+																55, 58, 54,
+																56, 59, 55,
+																57, 61, 45,
+																58, 62, 57,
+																59, 63, 58,
+																60, 64, 59,
+																1, 51, 0,
+																66, 53, 51,
+																3, 66, 1,
+																67, 56, 53,
+																68, 67, 66,
+																6, 68, 3,
+																69, 60, 56,
+																70, 69, 67,
+																71, 70, 68,
+																10, 71, 6,
+																72, 65, 60,
+																73, 72, 69,
+																74, 73, 70,
+																75, 74, 71,
+																15, 75, 10,
+																1, 66, 51,
+																66, 67, 53,
+																3, 68, 66,
+																67, 69, 56,
+																68, 70, 67,
+																6, 71, 68,
+																69, 72, 60,
+																70, 73, 69,
+																71, 74, 70,
+																10, 75, 71,
+																76, 75, 15,
+																77, 74, 75,
+																78, 77, 76,
+																79, 73, 74,
+																80, 79, 77,
+																81, 80, 78,
+																82, 72, 73,
+																83, 82, 79,
+																84, 83, 80,
+																85, 84, 81,
+																86, 65, 72,
+																87, 86, 82,
+																88, 87, 83,
+																89, 88, 84,
+																90, 89, 85,
+																76, 77, 75,
+																77, 79, 74,
+																78, 80, 77,
+																79, 82, 73,
+																80, 83, 79,
+																81, 84, 80,
+																82, 86, 72,
+																83, 87, 82,
+																84, 88, 83,
+																85, 89, 84,
+																91, 64, 65,
+																92, 63, 64,
+																93, 92, 91,
+																94, 62, 63,
+																95, 94, 92,
+																96, 95, 93,
+																97, 61, 62,
+																98, 97, 94,
+																99, 98, 95,
+																100, 99, 96,
+																101, 50, 61,
+																102, 101, 97,
+																103, 102, 98,
+																104, 103, 99,
+																105, 104, 100,
+																91, 92, 64,
+																92, 94, 63,
+																93, 95, 92,
+																94, 97, 62,
+																95, 98, 94,
+																96, 99, 95,
+																97, 101, 61,
+																98, 102, 97,
+																99, 103, 98,
+																100, 104, 99,
+																106, 49, 50,
+																107, 48, 49,
+																108, 107, 106,
+																109, 47, 48,
+																110, 109, 107,
+																111, 110, 108,
+																112, 46, 47,
+																113, 112, 109,
+																114, 113, 110,
+																115, 114, 111,
+																116, 35, 46,
+																117, 116, 112,
+																118, 117, 113,
+																119, 118, 114,
+																120, 119, 115,
+																106, 107, 49,
+																107, 109, 48,
+																108, 110, 107,
+																109, 112, 47,
+																110, 113, 109,
+																111, 114, 110,
+																112, 116, 46,
+																113, 117, 112,
+																114, 118, 113,
+																115, 119, 114,
+																121, 34, 35,
+																122, 33, 34,
+																123, 122, 121,
+																124, 32, 33,
+																125, 124, 122,
+																126, 125, 123,
+																127, 31, 32,
+																128, 127, 124,
+																129, 128, 125,
+																130, 129, 126,
+																131, 20, 31,
+																132, 131, 127,
+																133, 132, 128,
+																134, 133, 129,
+																135, 134, 130,
+																121, 122, 34,
+																122, 124, 33,
+																123, 125, 122,
+																124, 127, 32,
+																125, 128, 124,
+																126, 129, 125,
+																127, 131, 31,
+																128, 132, 127,
+																129, 133, 128,
+																130, 134, 129,
+																136, 19, 20,
+																137, 18, 19,
+																138, 137, 136,
+																139, 17, 18,
+																140, 139, 137,
+																141, 140, 138,
+																142, 16, 17,
+																143, 142, 139,
+																144, 143, 140,
+																145, 144, 141,
+																146, 15, 16,
+																147, 146, 142,
+																148, 147, 143,
+																149, 148, 144,
+																150, 149, 145,
+																136, 137, 19,
+																137, 139, 18,
+																138, 140, 137,
+																139, 142, 17,
+																140, 143, 139,
+																141, 144, 140,
+																142, 146, 16,
+																143, 147, 142,
+																144, 148, 143,
+																145, 149, 144,
+																146, 76, 15,
+																151, 78, 76,
+																147, 151, 146,
+																152, 81, 78,
+																153, 152, 151,
+																148, 153, 147,
+																154, 85, 81,
+																155, 154, 152,
+																156, 155, 153,
+																149, 156, 148,
+																157, 90, 85,
+																158, 157, 154,
+																159, 158, 155,
+																160, 159, 156,
+																150, 160, 149,
+																146, 151, 76,
+																151, 152, 78,
+																147, 153, 151,
+																152, 154, 81,
+																153, 155, 152,
+																148, 156, 153,
+																154, 157, 85,
+																155, 158, 154,
+																156, 159, 155,
+																149, 160, 156,
+																86, 91, 65,
+																161, 93, 91,
+																87, 161, 86,
+																162, 96, 93,
+																163, 162, 161,
+																88, 163, 87,
+																164, 100, 96,
+																165, 164, 162,
+																166, 165, 163,
+																89, 166, 88,
+																167, 105, 100,
+																168, 167, 164,
+																169, 168, 165,
+																170, 169, 166,
+																90, 170, 89,
+																86, 161, 91,
+																161, 162, 93,
+																87, 163, 161,
+																162, 164, 96,
+																163, 165, 162,
+																88, 166, 163,
+																164, 167, 100,
+																165, 168, 164,
+																166, 169, 165,
+																89, 170, 166,
+																101, 106, 50,
+																171, 108, 106,
+																102, 171, 101,
+																172, 111, 108,
+																173, 172, 171,
+																103, 173, 102,
+																174, 115, 111,
+																175, 174, 172,
+																176, 175, 173,
+																104, 176, 103,
+																177, 120, 115,
+																178, 177, 174,
+																179, 178, 175,
+																180, 179, 176,
+																105, 180, 104,
+																101, 171, 106,
+																171, 172, 108,
+																102, 173, 171,
+																172, 174, 111,
+																173, 175, 172,
+																103, 176, 173,
+																174, 177, 115,
+																175, 178, 174,
+																176, 179, 175,
+																104, 180, 176,
+																116, 121, 35,
+																181, 123, 121,
+																117, 181, 116,
+																182, 126, 123,
+																183, 182, 181,
+																118, 183, 117,
+																184, 130, 126,
+																185, 184, 182,
+																186, 185, 183,
+																119, 186, 118,
+																187, 135, 130,
+																188, 187, 184,
+																189, 188, 185,
+																190, 189, 186,
+																120, 190, 119,
+																116, 181, 121,
+																181, 182, 123,
+																117, 183, 181,
+																182, 184, 126,
+																183, 185, 182,
+																118, 186, 183,
+																184, 187, 130,
+																185, 188, 184,
+																186, 189, 185,
+																119, 190, 186,
+																131, 136, 20,
+																191, 138, 136,
+																132, 191, 131,
+																192, 141, 138,
+																193, 192, 191,
+																133, 193, 132,
+																194, 145, 141,
+																195, 194, 192,
+																196, 195, 193,
+																134, 196, 133,
+																197, 150, 145,
+																198, 197, 194,
+																199, 198, 195,
+																200, 199, 196,
+																135, 200, 134,
+																131, 191, 136,
+																191, 192, 138,
+																132, 193, 191,
+																192, 194, 141,
+																193, 195, 192,
+																133, 196, 193,
+																194, 197, 145,
+																195, 198, 194,
+																196, 199, 195,
+																134, 200, 196,
+																203, 202, 201,
+																205, 204, 202,
+																206, 205, 203,
+																208, 207, 204,
+																209, 208, 205,
+																210, 209, 206,
+																212, 211, 207,
+																213, 212, 208,
+																214, 213, 209,
+																215, 214, 210,
+																160, 150, 211,
+																159, 160, 212,
+																158, 159, 213,
+																157, 158, 214,
+																90, 157, 215,
+																203, 205, 202,
+																205, 208, 204,
+																206, 209, 205,
+																208, 212, 207,
+																209, 213, 208,
+																210, 214, 209,
+																212, 160, 211,
+																213, 159, 212,
+																214, 158, 213,
+																215, 157, 214,
+																202, 216, 201,
+																218, 217, 216,
+																204, 218, 202,
+																220, 219, 217,
+																221, 220, 218,
+																207, 221, 204,
+																223, 222, 219,
+																224, 223, 220,
+																225, 224, 221,
+																211, 225, 207,
+																200, 135, 222,
+																199, 200, 223,
+																198, 199, 224,
+																197, 198, 225,
+																150, 197, 211,
+																202, 218, 216,
+																218, 220, 217,
+																204, 221, 218,
+																220, 223, 219,
+																221, 224, 220,
+																207, 225, 221,
+																223, 200, 222,
+																224, 199, 223,
+																225, 198, 224,
+																211, 197, 225,
+																216, 226, 201,
+																228, 227, 226,
+																217, 228, 216,
+																230, 229, 227,
+																231, 230, 228,
+																219, 231, 217,
+																233, 232, 229,
+																234, 233, 230,
+																235, 234, 231,
+																222, 235, 219,
+																190, 120, 232,
+																189, 190, 233,
+																188, 189, 234,
+																187, 188, 235,
+																135, 187, 222,
+																216, 228, 226,
+																228, 230, 227,
+																217, 231, 228,
+																230, 233, 229,
+																231, 234, 230,
+																219, 235, 231,
+																233, 190, 232,
+																234, 189, 233,
+																235, 188, 234,
+																222, 187, 235,
+																226, 236, 201,
+																238, 237, 236,
+																227, 238, 226,
+																240, 239, 237,
+																241, 240, 238,
+																229, 241, 227,
+																243, 242, 239,
+																244, 243, 240,
+																245, 244, 241,
+																232, 245, 229,
+																180, 105, 242,
+																179, 180, 243,
+																178, 179, 244,
+																177, 178, 245,
+																120, 177, 232,
+																226, 238, 236,
+																238, 240, 237,
+																227, 241, 238,
+																240, 243, 239,
+																241, 244, 240,
+																229, 245, 241,
+																243, 180, 242,
+																244, 179, 243,
+																245, 178, 244,
+																232, 177, 245,
+																236, 203, 201,
+																246, 206, 203,
+																237, 246, 236,
+																247, 210, 206,
+																248, 247, 246,
+																239, 248, 237,
+																249, 215, 210,
+																250, 249, 247,
+																251, 250, 248,
+																242, 251, 239,
+																170, 90, 215,
+																169, 170, 249,
+																168, 169, 250,
+																167, 168, 251,
+																105, 167, 242,
+																236, 246, 203,
+																246, 247, 206,
+																237, 248, 246,
+																247, 249, 210,
+																248, 250, 247,
+																239, 251, 248,
+																249, 170, 215,
+																250, 169, 249,
+																251, 168, 250,
+																242, 167, 251
+	};
+
+}
\ No newline at end of file
diff --git a/Engine/Source/Programs/HeadlessChaos/Resource/FracturedGeometry.cpp b/Engine/Source/Programs/HeadlessChaos/Resource/FracturedGeometry.cpp
new file mode 100644
--- /dev/null
+++ b/Engine/Source/Programs/HeadlessChaos/Resource/FracturedGeometry.cpp
@@ -0,0 +1,2003 @@
+// Copyright Epic Games, Inc. All Rights Reserved.
+
+#include "../Resource/FracturedGeometry.h"
+#include "Chaos/Box.h"
+#include "Math/RandomStream.h"
+
+namespace GeometryCollectionTest 
+{
+
+	const TArray<float> FracturedGeometry::RawVertexArray = {
+													50, 0, -50,
+													50, 50, -50,
+													0, 50, -50,
+													0, 0, -50,
+													-50, 50, -50,
+													50, -50, -50,
+													-50, 0, -50,
+													0, -50, -50,
+													-50, -50, -50,
+													50, 50, 50,
+													50, 1.22465e-14, 50,
+													0, 50, 50,
+													0, 1.22465e-14, 50,
+													50, -50, 50,
+													-50, 50, 50,
+													0, -50, 50,
+													-50, 1.22465e-14, 50,
+													-50, -50, 50,
+													50, 50, -50,
+													50, 50, -3.06162e-15,
+													0, 50, -50,
+													0, 50, -3.06162e-15,
+													50, 50, 50,
+													-50, 50, -50,
+													0, 50, 50,
+													-50, 50, -3.06162e-15,
+													-50, 50, 50,
+													50, -50, -50,
+													50, -50, 3.06162e-15,
+													50, 0, -50,
+													50, 6.12323e-15, -1.8747e-31,
+													50, -50, 50,
+													50, 50, -50,
+													50, 1.22465e-14, 50,
+													50, 50, -3.06162e-15,
+													50, 50, 50,
+													-50, -50, -50,
+													-50, -50, 3.06162e-15,
+													0, -50, -50,
+													0, -50, 3.06162e-15,
+													-50, -50, 50,
+													50, -50, -50,
+													0, -50, 50,
+													50, -50, 3.06162e-15,
+													50, -50, 50,
+													-50, 50, -50,
+													-50, 50, -3.06162e-15,
+													-50, 0, -50,
+													-50, 6.12323e-15, -1.8747e-31,
+													-50, 50, 50,
+													-50, -50, -50,
+													-50, 1.22465e-14, 50,
+													-50, -50, 3.06162e-15,
+													-50, -50, 50,
+													16.2172, 8.6841, 15.7825,
+													-10.3678, 35.2691, 15.7825,
+													24.0997, 35.2691, 15.7825,
+													9.27464, -14.7309, 15.7825,
+													-10.3678, -14.7309, 15.7825,
+													3.62653, -28.7252, 15.7825,
+													5.3612, -27.9296, 15.7825,
+													-10.3678, -35.1439, 15.7825,
+													15.9054, 35.2691, -7.94424,
+													2.08283, 35.2691, -21.7668,
+													11.8691, 35.2691, -19.6313,
+													-10.3678, 35.2691, 15.7825,
+													-10.3678, 35.2691, -24.4836,
+													24.0997, 35.2691, 15.7825,
+													-10.3678, 35.2691, 15.7825,
+													-10.3678, 2.16495, -17.3216,
+													-10.3678, 35.2691, -24.4836,
+													-10.3678, -14.7309, 15.7825,
+													-10.3678, -14.7309, -13.6663,
+													-10.3678, -33.1022, -2.58878,
+													-10.3678, -32.2934, -9.86669,
+													-10.3678, -35.1439, 15.7825,
+													-3.61541, -29.2878, -9.04352,
+													11.8691, 35.2691, -19.6313,
+													-10.3678, -32.2934, -9.86669,
+													-10.3678, -14.7309, -13.6663,
+													-10.3678, 2.16495, -17.3216,
+													-10.3678, 35.2691, -24.4836,
+													2.08283, 35.2691, -21.7668,
+													5.3612, -27.9296, 15.7825,
+													9.27464, -14.7309, 15.7825,
+													-3.61541, -29.2878, -9.04352,
+													11.8691, 35.2691, -19.6313,
+													15.9054, 35.2691, -7.94424,
+													24.0997, 35.2691, 15.7825,
+													16.2172, 8.6841, 15.7825,
+													-10.3678, -35.1439, 15.7825,
+													3.62653, -28.7252, 15.7825,
+													-10.3678, -33.1022, -2.58878,
+													-10.3678, -32.2934, -9.86669,
+													-3.61541, -29.2878, -9.04352,
+													5.3612, -27.9296, 15.7825,
+													-0.0144844, -15.9719, -42.5147,
+													-4.03455, 9.82732, -42.5147,
+													-27.0941, -15.9719, -42.5147,
+													-27.0941, 9.35167, -42.5147,
+													-14.633, 13.3885, -42.5147,
+													-27.0941, -15.9719, -42.5147,
+													-27.0941, -15.9719, 7.4853,
+													-0.0144844, -15.9719, -42.5147,
+													8.77472, -15.9719, -28.3835,
+													22.9059, -15.9719, -5.6636,
+													22.9059, -15.9719, 7.4853,
+													6.25886, -15.9719, 24.1323,
+													-27.0941, -15.9719, 29.5568,
+													27.9481, -15.9719, 2.44314,
+													22.9059, -15.9719, 21.4249,
+													31.0841, -15.9719, 7.4853,
+													38.2065, -15.9719, 18.9365,
+													-27.0941, 9.35167, -42.5147,
+													-27.0941, 11.456, -15.0869,
+													-27.0941, -15.9719, -42.5147,
+													-27.0941, -15.9719, 7.4853,
+													-27.0941, 13.1877, 7.4853,
+													-27.0941, 5.07368, 28.5308,
+													-27.0941, 14.7661, 28.0583,
+													-27.0941, -15.9719, 29.5568,
+													-27.0941, 14.7661, 28.0583,
+													-27.0941, 13.1877, 7.4853,
+													-11.9324, 19.4711, 25.3631,
+													-27.0941, 11.456, -15.0869,
+													-2.6789, 21.8609, 17.4398,
+													-14.633, 13.3885, -42.5147,
+													-27.0941, 9.35167, -42.5147,
+													-0.0144844, -15.9719, -42.5147,
+													8.77472, -15.9719, -28.3835,
+													-4.03455, 9.82732, -42.5147,
+													22.9059, -15.9719, -5.6636,
+													27.9481, -15.9719, 2.44314,
+													31.0841, -15.9719, 7.4853,
+													20.3276, 9.79925, -3.35256,
+													36.7772, -7.517, 18.7568,
+													38.2065, -15.9719, 18.9365,
+													-2.6789, 21.8609, 17.4398,
+													-14.633, 13.3885, -42.5147,
+													5.96455, 19.9955, 22.4268,
+													25.0912, 12.9919, 19.6575,
+													20.3276, 9.79925, -3.35256,
+													-4.03455, 9.82732, -42.5147,
+													-2.6789, 21.8609, 17.4398,
+													5.96455, 19.9955, 22.4268,
+													-11.9324, 19.4711, 25.3631,
+													20.3276, 9.79925, -3.35256,
+													36.7772, -7.517, 18.7568,
+													25.0912, 12.9919, 19.6575,
+													-27.0941, -15.9719, 29.5568,
+													-27.0941, 5.07368, 28.5308,
+													6.25886, -15.9719, 24.1323,
+													22.9059, -15.9719, 21.4249,
+													38.2065, -15.9719, 18.9365,
+													36.7772, -7.517, 18.7568,
+													25.0912, 12.9919, 19.6575,
+													5.96455, 19.9955, 22.4268,
+													-11.9324, 19.4711, 25.3631,
+													-27.0941, 14.7661, 28.0583,
+													7.54969, 9.40126, -54.4615,
+													16.0109, 4.613, -54.4615,
+													29.9135, 9.40126, -54.4615,
+													29.9135, 7.20573, -54.4615,
+													29.9135, -25.0935, 45.5385,
+													16.8454, -27.5306, 45.5385,
+													29.9135, 9.40126, 45.5385,
+													-20.0865, 9.40126, 45.5385,
+													-20.0865, -0.546741, 45.5385,
+													2.03889, -30.2919, 45.5385,
+													-27.4861, 9.40126, 45.5385,
+													7.54969, 9.40126, -54.4615,
+													29.9135, 9.40126, -54.4615,
+													-6.91031, 9.40126, -41.2853,
+													29.9135, 9.40126, -4.4615,
+													-20.0865, 9.40126, -4.4615,
+													-20.0865, 9.40126, -29.279,
+													29.9135, 9.40126, 45.5385,
+													-30.7467, 9.40126, -15.1217,
+													-28.7678, 9.40126, -21.3684,
+													-20.0865, 9.40126, 45.5385,
+													-34.1239, 9.40126, -4.4615,
+													-35.2935, 9.40126, -0.769467,
+													-28.9867, 9.40126, 36.6383,
+													-27.4861, 9.40126, 45.5385,
+													29.9135, 9.40126, -54.4615,
+													29.9135, 7.20573, -54.4615,
+													29.9135, 6.15827, -51.2185,
+													29.9135, -8.94388, -4.4615,
+													29.9135, 9.40126, -4.4615,
+													29.9135, -17.6961, 22.6359,
+													29.9135, 9.40126, 45.5385,
+													29.9135, -25.0935, 45.5385,
+													-28.7678, 9.40126, -21.3684,
+													-30.7467, 9.40126, -15.1217,
+													-29.5054, 0.391495, -6.18885,
+													-34.1239, 9.40126, -4.4615,
+													-35.2935, 9.40126, -0.769467,
+													29.9135, -25.0935, 45.5385,
+													29.9135, -17.6961, 22.6359,
+													16.8454, -27.5306, 45.5385,
+													2.03889, -30.2919, 45.5385,
+													-10.1522, -21.5984, 11.5841,
+													16.0109, 4.613, -54.4615,
+													29.9135, 7.20573, -54.4615,
+													29.9135, 6.15827, -51.2185,
+													29.9135, -8.94388, -4.4615,
+													16.0109, 4.613, -54.4615,
+													7.54969, 9.40126, -54.4615,
+													-10.1522, -21.5984, 11.5841,
+													-29.5054, 0.391495, -6.18885,
+													-28.7678, 9.40126, -21.3684,
+													-35.2935, 9.40126, -0.769467,
+													-28.9867, 9.40126, 36.6383,
+													-29.5054, 0.391495, -6.18885,
+													-10.1522, -21.5984, 11.5841,
+													2.03889, -30.2919, 45.5385,
+													-20.0865, -0.546741, 45.5385,
+													-27.4861, 9.40126, 45.5385,
+													12.2185, 7.69071, -40.3769,
+													21.3769, 34.8307, -40.3769,
+													-14.9215, 34.8307, -40.3769,
+													4.50432, -15.1693, -40.3769,
+													-14.9215, -15.1693, -40.3769,
+													-0.397007, -29.6938, -40.3769,
+													-14.9215, -39.8458, -40.3769,
+													-2.46044, -35.8089, -40.3769,
+													24.4182, 34.8307, -1.03714,
+													-14.9215, 34.8307, -40.3769,
+													26.3972, 34.8307, -7.28381,
+													21.3769, 34.8307, -40.3769,
+													-14.9215, 34.8307, 9.62309,
+													21.0411, 34.8307, 9.62309,
+													19.8714, 34.8307, 13.3151,
+													-2.47086, 34.8307, 22.0738,
+													7.31542, 34.8307, 24.2092,
+													-14.9215, 34.8307, 19.3569,
+													-14.9215, -15.1693, -40.3769,
+													-14.9215, 34.8307, -40.3769,
+													-14.9215, 34.8307, 9.62309,
+													-14.9215, -15.1693, 9.62309,
+													-14.9215, 34.8307, 19.3569,
+													-14.9215, 1.72654, 26.5189,
+													-14.9215, -37.7415, -12.9491,
+													-14.9215, -39.8458, -40.3769,
+													-14.9215, -15.1693, 30.1743,
+													-14.9215, -36.0097, 9.62309,
+													-14.9215, -32.7318, 33.9739,
+													-14.9215, -34.4313, 30.1961,
+													-14.9215, 34.8307, 19.3569,
+													-2.4709, 34.8307, 22.0737,
+													-14.9215, 1.72652, 26.5189,
+													-14.9215, -15.1693, 30.1742,
+													-14.9215, -32.7319, 33.9739,
+													-8.16909, -29.7263, 34.797,
+													7.31557, 34.8307, 24.2092,
+													-2.46057, -35.8089, -40.3769,
+													-14.9215, -39.8457, -40.3769,
+													9.49352, -27.3366, 19.5775,
+													0.240276, -29.7263, 27.5007,
+													-14.9215, -34.4313, 30.1959,
+													-14.9215, -36.0097, 9.62309,
+													-14.9215, -37.7414, -12.949,
+													19.8715, 34.8307, 13.3151,
+													21.0411, 34.8307, 9.62309,
+													25.6595, 25.821, 7.8958,
+													26.3972, 34.8307, -7.28368,
+													24.4183, 34.8307, -1.03711,
+													25.6595, 25.821, 7.8958,
+													26.3972, 34.8307, -7.28368,
+													9.49352, -27.3366, 19.5775,
+													-2.46057, -35.8089, -40.3769,
+													-0.397007, -29.6938, -40.3769,
+													4.50432, -15.1693, -40.3769,
+													12.2185, 7.69071, -40.3769,
+													21.3769, 34.8307, -40.3769,
+													-8.16909, -29.7263, 34.797,
+													7.31557, 34.8307, 24.2092,
+													0.240276, -29.7263, 27.5007,
+													9.49352, -27.3366, 19.5775,
+													25.6595, 25.821, 7.8958,
+													19.8715, 34.8307, 13.3151,
+													0.240276, -29.7263, 27.5007,
+													-14.9215, -34.4313, 30.1959,
+													-8.16909, -29.7263, 34.797,
+													-14.9215, -32.7319, 33.9739,
+													34.8575, -12.9869, -29.136,
+													34.8575, 24.5256, -29.136,
+													-8.70369, 30.5743, -29.136,
+													-3.4102, 33.3911, -29.136,
+													-15.1425, 27.148, -29.136,
+													-15.1425, -12.9869, -29.136,
+													-41.338, 13.2086, -29.136,
+													-38.0628, -12.9869, -29.136,
+													-42.0829, 12.8122, -29.136,
+													34.8575, -12.9869, 20.864,
+													34.8575, 4.92442, 2.9527,
+													34.8575, -12.9869, -29.136,
+													34.8575, 24.5256, -29.136,
+													34.8575, -6.01665, 20.864,
+													34.8575, -12.9869, 32.2749,
+													-38.0628, -12.9869, -29.136,
+													-29.2736, -12.9869, -15.0048,
+													-15.1425, -12.9869, -29.136,
+													-15.1425, -12.9869, 7.71515,
+													-10.1003, -12.9869, 15.8219,
+													34.8575, -12.9869, -29.136,
+													34.8575, -12.9869, 20.864,
+													-6.96422, -12.9869, 20.864,
+													0.158117, -12.9869, 32.3153,
+													16.4748, -12.9869, 39.2467,
+													11.6566, -12.9869, 41.0741,
+													34.8575, -12.9869, 32.2749,
+													0.158214, -12.9869, 32.3154,
+													-6.96419, -12.9869, 20.864,
+													-1.27103, -4.53197, 32.1356,
+													-10.1003, -12.9869, 15.8219,
+													-15.1425, -12.9869, 7.71512,
+													-17.7206, 12.7842, 10.0264,
+													-42.0829, 12.8123, -29.136,
+													-38.0628, -12.9869, -29.136,
+													-29.2736, -12.9869, -15.0048,
+													-3.4102, 33.3911, -29.136,
+													-17.7206, 12.7842, 10.0264,
+													-8.70369, 30.5743, -29.136,
+													-15.1425, 27.148, -29.136,
+													-41.338, 13.2086, -29.136,
+													-42.0829, 12.8123, -29.136,
+													34.8575, -12.9869, 32.2749,
+													16.4748, -12.9869, 39.2467,
+													34.8575, -6.01665, 20.864,
+													34.8575, 4.92442, 2.9527,
+													34.8575, 24.5256, -29.136,
+													-3.4102, 33.3911, -29.136,
+													-17.7206, 12.7842, 10.0264,
+													-1.27103, -4.53197, 32.1356,
+													11.6566, -12.9869, 41.0741,
+													-1.27103, -4.53197, 32.1356,
+													11.6566, -12.9869, 41.0741,
+													0.158214, -12.9869, 32.3154,
+													0.0289307, 39.9012, -61.1455,
+													-2.54625, 16.3791, -61.1455,
+													13.9315, 42.494, -61.1455,
+													13.9315, 14.5344, -61.1455,
+													13.9315, -5.31049, 38.8545,
+													0.86338, 7.75765, 38.8545,
+													13.9315, 10.1948, 38.8545,
+													-17.2768, -5.31049, 38.8545,
+													-13.9431, 4.99635, 38.8545,
+													-20.9687, -16.7255, 38.8545,
+													12.5846, -48.5735, 38.8545,
+													13.9315, -48.7243, 38.8545,
+													13.9315, -17.0949, -11.1455,
+													13.9315, -5.31049, -11.1455,
+													13.9315, -5.31049, -29.7745,
+													13.9315, -37.3846, 20.9286,
+													13.9315, 14.5344, -61.1455,
+													13.9315, 42.494, -61.1455,
+													13.9315, 41.4465, -57.9025,
+													13.9315, -5.31049, 38.8545,
+													13.9315, -48.7243, 38.8545,
+													13.9315, 26.3444, -11.1455,
+													13.9315, 17.5921, 15.9518,
+													13.9315, 10.1948, 38.8545,
+													13.9315, 42.4941, -61.1455,
+													13.9315, 41.4467, -57.9027,
+													0.0285797, 39.9013, -61.1455,
+													13.9315, 26.3445, -11.1455,
+													13.9315, 17.5923, 15.9517,
+													13.9315, 10.1949, 38.8545,
+													-26.1345, 13.69, 4.89963,
+													-13.9432, 4.99643, 38.8545,
+													0.863289, 7.75774, 38.8545,
+													-30.807, -23.3456, 6.6532,
+													-2.54625, 16.3791, -61.1455,
+													-32.655, -18.2041, 13.4884,
+													-26.1345, 13.69, 4.89963,
+													0.0285797, 39.9013, -61.1455,
+													-26.1345, 13.69, 4.89963,
+													-13.9432, 4.99643, 38.8545,
+													-32.655, -18.2041, 13.4884,
+													-20.9687, -16.7255, 38.8545,
+													-17.2768, -5.31049, 38.8545,
+													13.9315, -48.7243, 38.8545,
+													13.9315, -37.3846, 20.9286,
+													12.5846, -48.5735, 38.8545,
+													13.9315, -17.0949, -11.1455,
+													13.9315, -5.31049, -29.7745,
+													-30.807, -23.3456, 6.6532,
+													-2.54625, 16.3791, -61.1455,
+													13.9315, 14.5344, -61.1455,
+													-32.655, -18.2041, 13.4884,
+													-20.9687, -16.7255, 38.8545,
+													-30.807, -23.3456, 6.6532,
+													12.5846, -48.5735, 38.8545,
+													32.9967, 9.70081, -22.8865,
+													35.4077, 12.5017, -22.8865,
+													1.88543, 40.8121, -22.8865,
+													29.5216, 40.8121, -22.8865,
+													37.9828, 36.0239, -22.8865,
+													1.88543, -9.18788, -22.8865,
+													16.7359, -9.18788, -22.8865,
+													-11.8161, 40.8121, -22.8865,
+													-20.9746, 13.6721, -22.8865,
+													8.32422, -15.6267, -22.8865,
+													13.6177, -12.8099, -22.8865,
+													-28.6887, -9.18788, -22.8865,
+													1.88543, -19.0529, -22.8865,
+													-24.3101, -32.9923, -22.8865,
+													-33.59, -23.7124, -22.8865,
+													-25.055, -33.3887, -22.8865,
+													-35.6535, -29.8275, -22.8865,
+													1.88543, 40.8121, -22.8865,
+													29.5216, 40.8121, -22.8865,
+													15.0616, 40.8121, -9.71027,
+													1.88543, 40.8121, 2.29609,
+													-6.79587, 40.8121, 10.2066,
+													-11.8161, 40.8121, -22.8865,
+													-0.692513, -33.4167, 16.2762,
+													-25.055, -33.3886, -22.8865,
+													4.07097, -30.2241, 39.286,
+													-15.0555, -23.2205, 42.0553,
+													-23.6995, -21.355, 37.0679,
+													-35.6536, -29.8274, -22.8865,
+													-6.79593, 40.8121, 10.2064,
+													1.88543, 40.8121, 2.29581,
+													-7.53357, 31.8023, 25.3861,
+													15.0615, 40.8121, -9.71041,
+													11.8194, 9.81274, 43.1586,
+													37.9824, 36.0239, -22.8865,
+													29.5214, 40.8121, -22.8865,
+													-35.6536, -29.8274, -22.8865,
+													-33.5901, -23.7124, -22.8865,
+													-23.6995, -21.355, 37.0679,
+													-28.6888, -9.18788, -22.8865,
+													-20.9746, 13.6721, -22.8865,
+													-7.53357, 31.8023, 25.3861,
+													-6.79593, 40.8121, 10.2064,
+													-11.8161, 40.8121, -22.8865,
+													-25.055, -33.3886, -22.8865,
+													-0.692513, -33.4167, 16.2762,
+													-24.3102, -32.9923, -22.8865,
+													1.88543, -19.0529, -22.8865,
+													8.32418, -15.6266, -22.8865,
+													13.618, -12.8097, -22.8865,
+													11.8194, 9.81274, 43.1586,
+													37.9824, 36.0239, -22.8865,
+													5.2989, -22.0815, 51.7475,
+													7.14699, -27.2233, 44.9118,
+													35.4076, 12.5012, -22.8865,
+													11.8194, 9.81274, 43.1586,
+													5.2989, -22.0815, 51.7475,
+													-7.53357, 31.8023, 25.3861,
+													-23.6995, -21.355, 37.0679,
+													-15.0555, -23.2205, 42.0553,
+													35.4076, 12.5012, -22.8865,
+													32.9967, 9.70081, -22.8865,
+													7.14699, -27.2233, 44.9118,
+													4.07097, -30.2241, 39.286,
+													-0.692513, -33.4167, 16.2762,
+													13.618, -12.8097, -22.8865,
+													16.7359, -9.18788, -22.8865,
+													7.14699, -27.2233, 44.9118,
+													4.07097, -30.2241, 39.286,
+													5.2989, -22.0815, 51.7475,
+													-15.0555, -23.2205, 42.0553,
+													33.1597, -3.4619, 17.6225,
+													29.826, -13.7687, 17.6225,
+													11.0343, 26.2833, 17.6225,
+													11.0343, -13.7687, 17.6225,
+													26.1341, -25.1838, 17.6225,
+													22.5778, -25.3122, 17.6225,
+													-12.3807, 9.64623, 17.6225,
+													3.63467, 36.2313, 17.6225,
+													-4.4982, 36.2313, 17.6225,
+													11.0343, -25.7293, 17.6225,
+													-19.3232, -13.7687, 17.6225,
+													-23.2366, -26.9675, 17.6225,
+													2.13412, 36.2313, 8.72228,
+													-12.6925, 36.2313, -6.10432,
+													-4.17271, 36.2313, -28.6855,
+													-16.7287, 36.2313, -17.7914,
+													3.63467, 36.2313, 17.6225,
+													-4.4982, 36.2313, 17.6225,
+													-16.7287, 36.2313, -17.7913,
+													-12.6925, 36.2313, -6.10434,
+													-32.2133, -28.3257, -7.20351,
+													-4.4982, 36.2313, 17.6225,
+													-23.2367, -26.9675, 17.6225,
+													-19.3232, -13.7687, 17.6225,
+													-12.3807, 9.64623, 17.6225,
+													-14.5506, -25.9361, -22.4231,
+													-23.8033, -28.3257, -14.5004,
+													-5.90771, -27.8014, -17.4364,
+													20.9682, 5.23197, -16.3325,
+													33.1595, -3.46165, 17.6225,
+													1.6154, 27.2214, -34.1049,
+													-4.17274, 36.2313, -28.6855,
+													2.13407, 36.2313, 8.72223,
+													3.63462, 36.2313, 17.6225,
+													11.0343, 26.2832, 17.6225,
+													1.6154, 27.2214, -34.1049,
+													-4.17274, 36.2313, -28.6855,
+													-14.5506, -25.9361, -22.4231,
+													-23.8033, -28.3257, -14.5004,
+													-32.2133, -28.3257, -7.20351,
+													-16.7287, 36.2313, -17.7913,
+													14.4478, -26.6623, -7.74374,
+													26.1341, -25.1837, 17.6225,
+													20.9682, 5.23197, -16.3325,
+													33.1595, -3.46165, 17.6225,
+													29.8259, -13.7687, 17.6225,
+													-14.5506, -25.9361, -22.4231,
+													-5.90771, -27.8014, -17.4364,
+													1.6154, 27.2214, -34.1049,
+													20.9682, 5.23197, -16.3325,
+													14.4478, -26.6623, -7.74374,
+													-23.2367, -26.9675, 17.6225,
+													11.0343, -25.7293, 17.6225,
+													-32.2133, -28.3257, -7.20351,
+													22.5778, -25.3122, 17.6225,
+													-23.8033, -28.3257, -14.5004,
+													-5.90771, -27.8014, -17.4364,
+													14.4478, -26.6623, -7.74374,
+													26.1341, -25.1837, 17.6225,
+													-3.49003, 43.2227, -49.9396,
+													15.3987, 24.334, -49.9396,
+													-1.0791, 46.0236, -49.9396,
+													15.3987, 44.1789, -49.9396,
+													-19.7509, 24.334, -49.9396,
+													15.3987, 11.8466, -49.9396,
+													-22.8691, 20.712, -49.9396,
+													15.3987, -19.0798, 50.0604,
+													15.3987, -21.1432, 50.0604,
+													14.0517, -18.929, 50.0604,
+													15.3987, 11.8466, -49.9396,
+													15.3987, -7.75465, -17.8509,
+													15.3987, 24.334, -49.9396,
+													15.3987, 24.334, -18.5686,
+													15.3987, 12.5495, 0.060406,
+													15.3987, -18.6957, 0.060406,
+													15.3987, -7.74012, 32.1345,
+													15.3987, -23.8068, 48.2012,
+													15.3987, -25.666, 46.9035,
+													15.3987, -25.666, 11.4713,
+													15.3987, 44.1789, -49.9396,
+													15.3987, -19.0798, 50.0604,
+													15.3987, -21.1432, 50.0604,
+													-6.40977, -25.666, 21.8688,
+													-2.98403, -25.666, 18.4431,
+													-7.80222, -25.666, 20.2704,
+													15.3987, -25.666, 46.9035,
+													15.3987, -25.666, 11.4713,
+													-32.416, 3.29756, 12.2323,
+													-20.7301, -17.211, 11.3318,
+													-37.1796, 0.105114, -10.7773,
+													-7.8023, -25.666, 20.2703,
+													-2.98378, -25.666, 18.4428,
+													-20.7301, -17.211, 11.3318,
+													15.3987, -25.666, 11.4711,
+													-37.1796, 0.105114, -10.7773,
+													-22.8692, 20.712, -49.9396,
+													15.3987, 11.8465, -49.9396,
+													15.3987, -7.75481, -17.8508,
+													15.3987, -18.6958, 0.060406,
+													15.3987, 44.1788, -49.9396,
+													15.3987, 24.334, -18.5687,
+													-1.07873, 46.0235, -49.9396,
+													15.3987, 12.5495, 0.060406,
+													15.3987, -7.74031, 32.1347,
+													15.3987, -19.0799, 50.0604,
+													-29.3394, 6.29887, 17.8589,
+													14.0524, -18.9292, 50.0604,
+													-22.8692, 20.712, -49.9396,
+													-19.751, 24.334, -49.9396,
+													-37.1796, 0.105114, -10.7773,
+													-32.416, 3.29756, 12.2323,
+													-29.3394, 6.29887, 17.8589,
+													-1.07873, 46.0235, -49.9396,
+													-3.48997, 43.2226, -49.9396,
+													14.0524, -18.9292, 50.0604,
+													15.3987, -21.1432, 50.0604,
+													-29.3394, 6.29887, 17.8589,
+													-32.416, 3.29756, 12.2323,
+													-20.7301, -17.211, 11.3318,
+													-7.8023, -25.666, 20.2703,
+													-6.40977, -25.666, 21.8688,
+													15.3987, -25.666, 46.9035,
+													55.2747, -18.9041, 17.9434,
+													16.8182, 19.5525, 17.9434,
+													55.2747, -14.3813, 17.9434,
+													53.9278, -12.1671, 17.9434,
+													20.3745, 19.6809, 17.9434,
+													5.27474, 19.1354, 17.9434,
+													5.27474, -18.9041, 17.9434,
+													-30.7309, 17.1016, 17.9434,
+													-28.9962, 17.8972, 17.9434,
+													-44.7253, -18.9041, 17.9434,
+													-44.7253, 10.6829, 17.9434,
+													55.2747, -18.9041, 17.9434,
+													55.2747, -17.0449, 16.0842,
+													55.2747, -18.9041, 14.7865,
+													55.2747, -14.3813, 17.9434,
+													-44.7253, -18.9041, -9.98509,
+													-44.7253, -18.9041, 17.9434,
+													-11.3723, -18.9041, -15.4095,
+													5.27474, -18.9041, 17.9434,
+													5.27474, -18.9041, -18.117,
+													33.4663, -18.9041, -10.2481,
+													32.0739, -18.9041, -11.8465,
+													20.5753, -18.9041, -20.6054,
+													55.2747, -18.9041, 17.9434,
+													55.2747, -18.9041, 14.7865,
+													-44.7253, 11.8339, -11.4835,
+													-44.7253, 13.5335, -7.70581,
+													-44.7253, 2.14147, -11.011,
+													-44.7253, 12.7246, -0.427895,
+													-44.7253, -18.9041, -9.98509,
+													-44.7253, -18.9041, 17.9434,
+													-44.7253, 10.6829, 17.9434,
+													-37.973, 16.539, -6.8829,
+													-28.9962, 17.8972, 17.9434,
+													-44.7253, 13.5335, -7.70605,
+													-44.7253, 12.7246, -0.427923,
+													-44.7253, 10.6829, 17.9434,
+													-30.7309, 17.1016, 17.9434,
+													-29.5643, 16.539, -14.1787,
+													-44.7253, 11.8343, -11.4836,
+													-11.6668, 17.0634, -17.115,
+													7.46001, 10.0597, -19.8843,
+													19.146, -10.4492, -20.7851,
+													20.5752, -18.9041, -20.6054,
+													5.27474, -18.9041, -18.117,
+													-11.3723, -18.9041, -15.4095,
+													-44.7253, -18.9041, -9.98513,
+													-44.7253, 2.14142, -11.0111,
+													-37.973, 16.539, -6.8829,
+													-44.7253, 13.5335, -7.70605,
+													-29.5643, 16.539, -14.1787,
+													-44.7253, 11.8343, -11.4836,
+													19.146, -10.4492, -20.7851,
+													20.5752, -18.9041, -20.6054,
+													32.0735, -18.9041, -11.8468,
+													10.5365, 13.0605, -14.2583,
+													53.9286, -12.1676, 17.9434,
+													8.68832, 18.2024, -7.42243,
+													20.3744, 19.681, 17.9434,
+													-11.6668, 17.0634, -17.115,
+													7.46001, 10.0597, -19.8843,
+													8.68832, 18.2024, -7.42243,
+													10.5365, 13.0605, -14.2583,
+													8.68832, 18.2024, -7.42243,
+													20.3744, 19.681, 17.9434,
+													-11.6668, 17.0634, -17.115,
+													-29.5643, 16.539, -14.1787,
+													-37.973, 16.539, -6.8829,
+													-28.9962, 17.8972, 17.9434,
+													5.27474, 19.1354, 17.9434,
+													16.8182, 19.5525, 17.9434,
+													19.146, -10.4492, -20.7851,
+													32.0735, -18.9041, -11.8468,
+													7.46001, 10.0597, -19.8843,
+													10.5365, 13.0605, -14.2583,
+													53.9286, -12.1676, 17.9434,
+													55.2747, -14.3814, 17.9434,
+													55.2747, -17.0449, 16.0843,
+													55.2747, -18.9041, 14.7866,
+													33.4662, -18.9041, -10.248,
+	};
+
+	const TArray<int32> FracturedGeometry::RawIndicesArray = {
+													0, 1, 2,
+													0, 2, 3,
+													3, 2, 4,
+													5, 0, 3,
+													3, 4, 6,
+													5, 3, 7,
+													7, 3, 6,
+													7, 6, 8,
+													9, 10, 11,
+													11, 10, 12,
+													10, 13, 12,
+													11, 12, 14,
+													12, 13, 15,
+													14, 12, 16,
+													12, 15, 16,
+													16, 15, 17,
+													18, 19, 20,
+													20, 19, 21,
+													19, 22, 21,
+													20, 21, 23,
+													21, 22, 24,
+													23, 21, 25,
+													21, 24, 25,
+													25, 24, 26,
+													27, 28, 29,
+													29, 28, 30,
+													28, 31, 30,
+													29, 30, 32,
+													30, 31, 33,
+													32, 30, 34,
+													30, 33, 34,
+													34, 33, 35,
+													36, 37, 38,
+													38, 37, 39,
+													37, 40, 39,
+													38, 39, 41,
+													39, 40, 42,
+													41, 39, 43,
+													39, 42, 43,
+													43, 42, 44,
+													45, 46, 47,
+													47, 46, 48,
+													46, 49, 48,
+													47, 48, 50,
+													48, 49, 51,
+													50, 48, 52,
+													48, 51, 52,
+													52, 51, 53,
+													54, 55, 56,
+													57, 58, 54,
+													54, 58, 55,
+													58, 57, 59,
+													59, 57, 60,
+													61, 58, 59,
+													62, 63, 64,
+													65, 66, 67,
+													67, 66, 62,
+													62, 66, 63,
+													68, 69, 70,
+													71, 72, 68,
+													68, 72, 69,
+													73, 74, 71,
+													71, 74, 72,
+													71, 75, 73,
+													76, 77, 78,
+													78, 77, 79,
+													79, 77, 80,
+													81, 80, 82,
+													80, 77, 82,
+													83, 84, 85,
+													85, 84, 86,
+													86, 84, 87,
+													88, 87, 89,
+													87, 84, 89,
+													90, 91, 92,
+													92, 91, 93,
+													93, 91, 94,
+													94, 91, 95,
+													96, 97, 98,
+													98, 97, 99,
+													99, 97, 100,
+													101, 102, 103,
+													103, 102, 104,
+													104, 102, 105,
+													105, 102, 106,
+													107, 106, 108,
+													108, 106, 102,
+													106, 109, 105,
+													106, 107, 110,
+													106, 111, 109,
+													110, 112, 106,
+													106, 112, 111,
+													113, 114, 115,
+													115, 114, 116,
+													116, 114, 117,
+													118, 116, 119,
+													119, 116, 117,
+													116, 118, 120,
+													121, 122, 123,
+													122, 124, 123,
+													123, 124, 125,
+													125, 124, 126,
+													126, 124, 127,
+													128, 129, 130,
+													129, 131, 130,
+													131, 132, 130,
+													132, 133, 130,
+													130, 133, 134,
+													134, 133, 135,
+													135, 133, 136,
+													137, 138, 139,
+													139, 138, 140,
+													140, 138, 141,
+													141, 138, 142,
+													143, 144, 145,
+													146, 147, 148,
+													149, 150, 151,
+													151, 150, 152,
+													152, 150, 153,
+													153, 150, 154,
+													154, 150, 155,
+													155, 150, 156,
+													156, 150, 157,
+													157, 150, 158,
+													159, 160, 161,
+													161, 160, 162,
+													163, 164, 165,
+													165, 164, 166,
+													167, 166, 168,
+													168, 166, 164,
+													169, 166, 167,
+													170, 171, 172,
+													172, 171, 173,
+													174, 175, 173,
+													173, 175, 172,
+													174, 173, 176,
+													174, 177, 175,
+													175, 177, 178,
+													179, 174, 176,
+													174, 180, 177,
+													180, 174, 181,
+													181, 174, 182,
+													182, 174, 179,
+													182, 179, 183,
+													184, 185, 186,
+													186, 187, 184,
+													184, 187, 188,
+													187, 189, 188,
+													188, 189, 190,
+													190, 189, 191,
+													192, 193, 194,
+													193, 195, 194,
+													194, 195, 196,
+													197, 198, 199,
+													199, 198, 200,
+													200, 198, 201,
+													201, 198, 202,
+													203, 202, 204,
+													204, 202, 205,
+													202, 198, 205,
+													206, 207, 208,
+													208, 207, 209,
+													209, 207, 210,
+													211, 212, 213,
+													213, 212, 214,
+													214, 212, 215,
+													215, 212, 216,
+													216, 212, 217,
+													218, 219, 220,
+													218, 220, 221,
+													221, 220, 222,
+													222, 223, 221,
+													224, 225, 222,
+													222, 225, 223,
+													226, 227, 228,
+													228, 227, 229,
+													227, 226, 230,
+													230, 226, 231,
+													231, 232, 230,
+													230, 232, 233,
+													233, 232, 234,
+													230, 233, 235,
+													236, 237, 238,
+													236, 238, 239,
+													240, 241, 238,
+													238, 241, 239,
+													239, 242, 236,
+													236, 242, 243,
+													239, 241, 244,
+													239, 245, 242,
+													244, 246, 239,
+													239, 246, 245,
+													245, 246, 247,
+													248, 249, 250,
+													250, 249, 251,
+													251, 249, 252,
+													252, 249, 253,
+													253, 249, 254,
+													255, 256, 257,
+													257, 256, 258,
+													259, 258, 260,
+													260, 258, 261,
+													258, 256, 261,
+													262, 263, 264,
+													265, 264, 266,
+													264, 263, 266,
+													267, 268, 269,
+													269, 268, 270,
+													270, 268, 271,
+													271, 268, 272,
+													272, 268, 273,
+													273, 268, 274,
+													275, 276, 277,
+													277, 276, 278,
+													278, 276, 279,
+													279, 276, 280,
+													281, 282, 283,
+													283, 282, 284,
+													285, 286, 287,
+													287, 286, 288,
+													287, 289, 285,
+													285, 289, 290,
+													289, 291, 290,
+													290, 291, 292,
+													292, 291, 293,
+													294, 295, 296,
+													296, 295, 297,
+													295, 294, 298,
+													298, 294, 299,
+													300, 301, 302,
+													302, 301, 303,
+													303, 304, 302,
+													302, 304, 305,
+													305, 304, 306,
+													306, 304, 307,
+													307, 308, 306,
+													306, 308, 309,
+													309, 308, 310,
+													306, 309, 311,
+													312, 313, 314,
+													313, 315, 314,
+													315, 316, 314,
+													314, 316, 317,
+													317, 316, 318,
+													319, 318, 320,
+													318, 316, 320,
+													321, 322, 323,
+													323, 322, 324,
+													324, 322, 325,
+													322, 326, 325,
+													327, 328, 329,
+													329, 328, 330,
+													330, 328, 331,
+													331, 328, 332,
+													332, 328, 333,
+													333, 328, 334,
+													334, 328, 335,
+													336, 337, 338,
+													339, 340, 341,
+													341, 340, 342,
+													343, 344, 345,
+													346, 347, 343,
+													343, 347, 344,
+													346, 343, 348,
+													348, 343, 349,
+													349, 343, 350,
+													351, 352, 353,
+													352, 351, 354,
+													353, 352, 355,
+													355, 352, 356,
+													356, 352, 357,
+													352, 354, 358,
+													358, 354, 359,
+													352, 360, 357,
+													360, 352, 361,
+													361, 352, 358,
+													358, 362, 361,
+													363, 364, 365,
+													364, 366, 365,
+													366, 367, 365,
+													367, 368, 365,
+													365, 368, 369,
+													370, 369, 371,
+													369, 368, 371,
+													372, 373, 374,
+													374, 373, 375,
+													375, 373, 376,
+													377, 378, 379,
+													380, 379, 381,
+													379, 378, 381,
+													382, 383, 384,
+													383, 385, 384,
+													385, 386, 384,
+													384, 386, 387,
+													387, 386, 388,
+													388, 386, 389,
+													390, 391, 392,
+													392, 391, 393,
+													394, 395, 396,
+													396, 395, 397,
+													397, 395, 398,
+													399, 400, 396,
+													396, 400, 394,
+													396, 401, 399,
+													399, 401, 402,
+													403, 404, 399,
+													399, 404, 400,
+													405, 399, 402,
+													399, 406, 403,
+													399, 405, 406,
+													406, 405, 407,
+													407, 405, 408,
+													407, 408, 409,
+													409, 408, 410,
+													411, 412, 413,
+													411, 413, 414,
+													415, 416, 414,
+													414, 416, 411,
+													417, 418, 419,
+													419, 418, 420,
+													420, 418, 421,
+													421, 418, 422,
+													423, 424, 425,
+													424, 426, 425,
+													425, 426, 427,
+													427, 426, 428,
+													428, 426, 429,
+													430, 431, 432,
+													431, 433, 432,
+													433, 434, 432,
+													432, 434, 435,
+													435, 434, 436,
+													436, 434, 437,
+													438, 439, 440,
+													440, 439, 441,
+													441, 439, 442,
+													439, 443, 442,
+													444, 445, 446,
+													446, 445, 447,
+													447, 445, 448,
+													449, 450, 451,
+													451, 450, 452,
+													452, 450, 453,
+													454, 455, 456,
+													456, 455, 457,
+													457, 455, 458,
+													459, 458, 460,
+													458, 455, 460,
+													461, 462, 463,
+													463, 462, 464,
+													465, 466, 467,
+													467, 466, 468,
+													469, 470, 466,
+													466, 470, 468,
+													468, 471, 467,
+													467, 471, 472,
+													472, 471, 473,
+													470, 474, 468,
+													468, 475, 471,
+													468, 474, 475,
+													475, 474, 476,
+													477, 478, 479,
+													479, 478, 480,
+													477, 481, 478,
+													478, 481, 482,
+													483, 484, 485,
+													484, 486, 485,
+													487, 485, 488,
+													488, 485, 489,
+													485, 486, 489,
+													490, 491, 492,
+													493, 494, 495,
+													495, 494, 496,
+													496, 494, 497,
+													498, 497, 499,
+													497, 494, 499,
+													500, 501, 502,
+													502, 501, 503,
+													503, 501, 504,
+													504, 501, 505,
+													506, 507, 508,
+													509, 508, 510,
+													508, 507, 510,
+													511, 512, 513,
+													513, 512, 514,
+													514, 512, 515,
+													516, 517, 518,
+													517, 519, 518,
+													518, 519, 520,
+													520, 519, 521,
+													521, 519, 522,
+													522, 519, 523,
+													524, 525, 526,
+													526, 525, 527,
+													528, 525, 524,
+													525, 528, 529,
+													529, 528, 530,
+													531, 532, 533,
+													534, 535, 536,
+													536, 535, 537,
+													537, 535, 538,
+													538, 535, 539,
+													538, 539, 540,
+													540, 539, 541,
+													541, 539, 542,
+													542, 539, 543,
+													537, 544, 536,
+													545, 540, 546,
+													546, 540, 541,
+													547, 548, 549,
+													550, 551, 547,
+													547, 551, 548,
+													552, 553, 554,
+													555, 556, 557,
+													556, 558, 557,
+													557, 558, 559,
+													559, 558, 560,
+													561, 560, 562,
+													562, 560, 563,
+													560, 558, 563,
+													564, 565, 566,
+													565, 567, 566,
+													567, 568, 566,
+													568, 569, 566,
+													566, 569, 570,
+													570, 569, 571,
+													572, 573, 574,
+													574, 573, 575,
+													575, 573, 576,
+													577, 576, 578,
+													576, 573, 578,
+													579, 580, 581,
+													581, 580, 582,
+													582, 580, 583,
+													583, 580, 584,
+													584, 580, 585,
+													585, 580, 586,
+													587, 588, 589,
+													589, 588, 590,
+													590, 588, 591,
+													588, 587, 592,
+													592, 587, 593,
+													593, 594, 592,
+													592, 594, 595,
+													596, 597, 593,
+													593, 597, 594,
+													598, 599, 600,
+													601, 599, 598,
+													602, 603, 604,
+													605, 606, 603,
+													603, 606, 604,
+													607, 608, 605,
+													605, 608, 606,
+													606, 608, 609,
+													610, 611, 605,
+													605, 611, 607,
+													612, 613, 614,
+													614, 613, 615,
+													614, 615, 616,
+													616, 615, 617,
+													617, 615, 618,
+													619, 620, 621,
+													621, 620, 622,
+													623, 622, 624,
+													622, 620, 624,
+													625, 626, 627,
+													627, 626, 628,
+													628, 626, 629,
+													629, 626, 630,
+													630, 626, 631,
+													631, 626, 632,
+													633, 632, 634,
+													632, 626, 634,
+													635, 636, 637,
+													637, 636, 638,
+													639, 640, 641,
+													642, 643, 644,
+													644, 643, 645,
+													646, 647, 648,
+													648, 647, 649,
+													650, 651, 652,
+													652, 651, 653,
+													653, 651, 654,
+													655, 654, 656,
+													656, 654, 657,
+													654, 651, 657,
+													658, 659, 660,
+													660, 659, 661,
+													661, 659, 662,
+													662, 659, 663,
+													663, 659, 664,
+													665, 664, 666,
+													664, 659, 666,
+};
+
+const TArray<int32> FracturedGeometry::RawBoneMapArray = {
+													0,
+													0,
+													0,
+													0,
+													0,
+													0,
+													0,
+													0,
+													0,
+													0,
+													0,
+													0,
+													0,
+													0,
+													0,
+													0,
+													0,
+													0,
+													0,
+													0,
+													0,
+													0,
+													0,
+													0,
+													0,
+													0,
+													0,
+													0,
+													0,
+													0,
+													0,
+													0,
+													0,
+													0,
+													0,
+													0,
+													0,
+													0,
+													0,
+													0,
+													0,
+													0,
+													0,
+													0,
+													0,
+													0,
+													0,
+													0,
+													0,
+													0,
+													0,
+													0,
+													0,
+													0,
+													1,
+													1,
+													1,
+													1,
+													1,
+													1,
+													1,
+													1,
+													1,
+													1,
+													1,
+													1,
+													1,
+													1,
+													1,
+													1,
+													1,
+													1,
+													1,
+													1,
+													1,
+													1,
+													1,
+													1,
+													1,
+													1,
+													1,
+													1,
+													1,
+													1,
+													1,
+													1,
+													1,
+													1,
+													1,
+													1,
+													1,
+													1,
+													1,
+													1,
+													1,
+													1,
+													2,
+													2,
+													2,
+													2,
+													2,
+													2,
+													2,
+													2,
+													2,
+													2,
+													2,
+													2,
+													2,
+													2,
+													2,
+													2,
+													2,
+													2,
+													2,
+													2,
+													2,
+													2,
+													2,
+													2,
+													2,
+													2,
+													2,
+													2,
+													2,
+													2,
+													2,
+													2,
+													2,
+													2,
+													2,
+													2,
+													2,
+													2,
+													2,
+													2,
+													2,
+													2,
+													2,
+													2,
+													2,
+													2,
+													2,
+													2,
+													2,
+													2,
+													2,
+													2,
+													2,
+													2,
+													2,
+													2,
+													2,
+													2,
+													2,
+													2,
+													2,
+													2,
+													2,
+													3,
+													3,
+													3,
+													3,
+													3,
+													3,
+													3,
+													3,
+													3,
+													3,
+													3,
+													3,
+													3,
+													3,
+													3,
+													3,
+													3,
+													3,
+													3,
+													3,
+													3,
+													3,
+													3,
+													3,
+													3,
+													3,
+													3,
+													3,
+													3,
+													3,
+													3,
+													3,
+													3,
+													3,
+													3,
+													3,
+													3,
+													3,
+													3,
+													3,
+													3,
+													3,
+													3,
+													3,
+													3,
+													3,
+													3,
+													3,
+													3,
+													3,
+													3,
+													3,
+													3,
+													3,
+													3,
+													3,
+													3,
+													3,
+													3,
+													4,
+													4,
+													4,
+													4,
+													4,
+													4,
+													4,
+													4,
+													4,
+													4,
+													4,
+													4,
+													4,
+													4,
+													4,
+													4,
+													4,
+													4,
+													4,
+													4,
+													4,
+													4,
+													4,
+													4,
+													4,
+													4,
+													4,
+													4,
+													4,
+													4,
+													4,
+													4,
+													4,
+													4,
+													4,
+													4,
+													4,
+													4,
+													4,
+													4,
+													4,
+													4,
+													4,
+													4,
+													4,
+													4,
+													4,
+													4,
+													4,
+													4,
+													4,
+													4,
+													4,
+													4,
+													4,
+													4,
+													4,
+													4,
+													4,
+													4,
+													4,
+													4,
+													4,
+													4,
+													4,
+													4,
+													4,
+													5,
+													5,
+													5,
+													5,
+													5,
+													5,
+													5,
+													5,
+													5,
+													5,
+													5,
+													5,
+													5,
+													5,
+													5,
+													5,
+													5,
+													5,
+													5,
+													5,
+													5,
+													5,
+													5,
+													5,
+													5,
+													5,
+													5,
+													5,
+													5,
+													5,
+													5,
+													5,
+													5,
+													5,
+													5,
+													5,
+													5,
+													5,
+													5,
+													5,
+													5,
+													5,
+													5,
+													5,
+													5,
+													5,
+													5,
+													5,
+													5,
+													5,
+													5,
+													5,
+													5,
+													5,
+													6,
+													6,
+													6,
+													6,
+													6,
+													6,
+													6,
+													6,
+													6,
+													6,
+													6,
+													6,
+													6,
+													6,
+													6,
+													6,
+													6,
+													6,
+													6,
+													6,
+													6,
+													6,
+													6,
+													6,
+													6,
+													6,
+													6,
+													6,
+													6,
+													6,
+													6,
+													6,
+													6,
+													6,
+													6,
+													6,
+													6,
+													6,
+													6,
+													6,
+													6,
+													6,
+													6,
+													6,
+													6,
+													6,
+													6,
+													6,
+													6,
+													6,
+													6,
+													6,
+													6,
+													6,
+													6,
+													7,
+													7,
+													7,
+													7,
+													7,
+													7,
+													7,
+													7,
+													7,
+													7,
+													7,
+													7,
+													7,
+													7,
+													7,
+													7,
+													7,
+													7,
+													7,
+													7,
+													7,
+													7,
+													7,
+													7,
+													7,
+													7,
+													7,
+													7,
+													7,
+													7,
+													7,
+													7,
+													7,
+													7,
+													7,
+													7,
+													7,
+													7,
+													7,
+													7,
+													7,
+													7,
+													7,
+													7,
+													7,
+													7,
+													7,
+													7,
+													7,
+													7,
+													7,
+													7,
+													7,
+													7,
+													7,
+													7,
+													7,
+													7,
+													7,
+													7,
+													7,
+													7,
+													7,
+													7,
+													7,
+													7,
+													7,
+													7,
+													7,
+													7,
+													7,
+													8,
+													8,
+													8,
+													8,
+													8,
+													8,
+													8,
+													8,
+													8,
+													8,
+													8,
+													8,
+													8,
+													8,
+													8,
+													8,
+													8,
+													8,
+													8,
+													8,
+													8,
+													8,
+													8,
+													8,
+													8,
+													8,
+													8,
+													8,
+													8,
+													8,
+													8,
+													8,
+													8,
+													8,
+													8,
+													8,
+													8,
+													8,
+													8,
+													8,
+													8,
+													8,
+													8,
+													8,
+													8,
+													8,
+													8,
+													8,
+													8,
+													8,
+													8,
+													8,
+													8,
+													8,
+													8,
+													8,
+													8,
+													8,
+													8,
+													9,
+													9,
+													9,
+													9,
+													9,
+													9,
+													9,
+													9,
+													9,
+													9,
+													9,
+													9,
+													9,
+													9,
+													9,
+													9,
+													9,
+													9,
+													9,
+													9,
+													9,
+													9,
+													9,
+													9,
+													9,
+													9,
+													9,
+													9,
+													9,
+													9,
+													9,
+													9,
+													9,
+													9,
+													9,
+													9,
+													9,
+													9,
+													9,
+													9,
+													9,
+													9,
+													9,
+													9,
+													9,
+													9,
+													9,
+													9,
+													9,
+													9,
+													9,
+													9,
+													9,
+													9,
+													9,
+													9,
+													9,
+													9,
+													9,
+													9,
+													9,
+													9,
+													9,
+													10,
+													10,
+													10,
+													10,
+													10,
+													10,
+													10,
+													10,
+													10,
+													10,
+													10,
+													10,
+													10,
+													10,
+													10,
+													10,
+													10,
+													10,
+													10,
+													10,
+													10,
+													10,
+													10,
+													10,
+													10,
+													10,
+													10,
+													10,
+													10,
+													10,
+													10,
+													10,
+													10,
+													10,
+													10,
+													10,
+													10,
+													10,
+													10,
+													10,
+													10,
+													10,
+													10,
+													10,
+													10,
+													10,
+													10,
+													10,
+													10,
+													10,
+													10,
+													10,
+													10,
+													10,
+													10,
+													10,
+													10,
+													10,
+													10,
+													10,
+													10,
+													10,
+													10,
+													10,
+													10,
+													10,
+													10,
+													10,
+													10,
+													10,
+													10,
+													10,
+													10,
+													10,
+													10,
+													10,
+													10,
+													10,
+													10,
+													10,
+};
+
+const TArray<FTransform> FracturedGeometry::RawTransformArray = {
+   FTransform(FQuat(0, 0, 0, 1), FVector(0, 0, 0), FVector(1, 1, 1)),
+   FTransform(FQuat(0, 0, 0, 1), FVector(-39.6322, 14.7309, 34.2175), FVector(1, 1, 1)),
+   FTransform(FQuat(0, 0, 0, 1), FVector(-22.9059, -34.0281, -7.4853), FVector(1, 1, 1)),
+   FTransform(FQuat(0, 0, 0, 1), FVector(20.0865, 40.5987, 4.4615), FVector(1, 1, 1)),
+   FTransform(FQuat(0, 0, 0, 1), FVector(-35.0785, 15.1693, -9.62309), FVector(1, 1, 1)),
+   FTransform(FQuat(0, 0, 0, 1), FVector(15.1425, -37.0131, -20.864), FVector(1, 1, 1)),
+   FTransform(FQuat(0, 0, 0, 1), FVector(36.0685, 5.31049, 11.1455), FVector(1, 1, 1)),
+   FTransform(FQuat(0, 0, 0, 1), FVector(-1.88543, 9.18788, -27.1135), FVector(1, 1, 1)),
+   FTransform(FQuat(0, 0, 0, 1), FVector(-11.0343, 13.7687, 32.3775), FVector(1, 1, 1)),
+   FTransform(FQuat(0, 0, 0, 1), FVector(34.6013, -24.334, -0.060406), FVector(1, 1, 1)),
+   FTransform(FQuat(0, 0, 0, 1), FVector(-5.27474, -31.0959, 32.0566), FVector(1, 1, 1)),
+};
+
+const TArray<int32> FracturedGeometry::RawLevelArray = {
+   0,
+   1,
+   1,
+   1,
+   1,
+   1,
+   1,
+   1,
+   1,
+   1,
+   1,
+};
+
+const TArray<int32> FracturedGeometry::RawParentArray = {
+  -1,
+   0, 
+   0, 
+   0, 
+   0, 
+   0, 
+   0, 
+   0, 
+   0, 
+   0, 
+   0, 
+};
+
+const TArray<TSet<int32>> FracturedGeometry::RawChildrenArray = {
+
+};
+
+const TArray<int32> FracturedGeometry::RawSimulationTypeArray = {
+   0,
+   0,
+   0,
+   0,
+   0,
+   0,
+   0,
+   0,
+   0,
+   0,
+   0
+};
+
+const TArray<int32> FracturedGeometry::RawStatusFlagsArray = {
+   0,
+   0,
+   0,
+   0,
+   0,
+   0,
+   0,
+   0,
+   0,
+   0,
+   0
+};
+
+FracturedGeometry::FracturedGeometry()
+{}
+
+FracturedGeometry::~FracturedGeometry()
+{}
+
+GlobalFracturedGeometry::GlobalFracturedGeometry()
+	: RawIndicesArray0(FracturedGeometry::RawIndicesArray)
+	, RawIndicesArray1(FracturedGeometry::RawIndicesArray)
+	, RawIndicesArray2(FracturedGeometry::RawIndicesArray)
+	, RawIndicesArray(RawIndicesArray1)
+	, RawBoneMapArray(FracturedGeometry::RawBoneMapArray)
+	, RawTransformArray(FracturedGeometry::RawTransformArray)
+	, RawLevelArray(FracturedGeometry::RawLevelArray)
+	, RawParentArray(FracturedGeometry::RawParentArray)
+	, RawChildrenArray(FracturedGeometry::RawChildrenArray)
+	, RawSimulationTypeArray(FracturedGeometry::RawSimulationTypeArray)
+	, RawStatusFlagsArray(FracturedGeometry::RawStatusFlagsArray)
+{
+	// The FracturedGeometry data set has lots of free and coincident vertices.
+	// We're going to concatenate 2 copies of it.
+	const int NumOrigFractured = FracturedGeometry::RawVertexArray.Num();
+	RawVertexArray.Reserve(NumOrigFractured * 3);
+	RawVertexArray.Append(FracturedGeometry::RawVertexArray);
+	RawVertexArray.Append(FracturedGeometry::RawVertexArray);
+	RawVertexArray.Append(FracturedGeometry::RawVertexArray);
+	check(RawVertexArray.Num() == NumOrigFractured * 3);
+
+	// Randomly perturb some of the vertices in the first point set.
+	FRandomStream RandIdx(NumOrigFractured);
+	FRandomStream RandPos(NumOrigFractured);
+	for(int32 i=0; i < NumOrigFractured; i += RandIdx.RandRange(0,10))
+	{
+		RawVertexArray[i] += RandPos.FRandRange(-1.0f, 1.0f);
+	}
+
+	// Measure the extent of the second point set.
+	Chaos::TAABB<float, 3> BBox(
+		Chaos::TVector<float, 3>(
+			RawVertexArray[NumOrigFractured + 0],
+			RawVertexArray[NumOrigFractured + 1],
+			RawVertexArray[NumOrigFractured + 2]),
+		Chaos::TVector<float, 3>(
+			RawVertexArray[NumOrigFractured + 0],
+			RawVertexArray[NumOrigFractured + 1],
+			RawVertexArray[NumOrigFractured + 2]));
+	for(int32 i=NumOrigFractured+3; i < NumOrigFractured*2; i+=3)
+	{
+		BBox.GrowToInclude(
+			Chaos::TVector<float, 3>(
+				RawVertexArray[i + 0],
+				RawVertexArray[i + 1],
+				RawVertexArray[i + 2]));
+	}
+
+	// Translate the third point set out of the way.
+	const Chaos::TVector<float, 3> Extents = BBox.Extents();
+	for(int32 i=NumOrigFractured*2 + 1; i < NumOrigFractured*3; i+=3)
+	{
+		RawVertexArray[i] += Extents[1] + 1.0e-5;
+	}
+
+	// Offset the index arrays (orig array is flat)
+	for(int32 &Idx : RawIndicesArray1)
+	{
+		Idx += NumOrigFractured / 3;
+	}
+	for (int32 &Idx : RawIndicesArray2)
+	{
+		Idx += NumOrigFractured / 3 * 2;
+	}
+
+	// Merge the index arrays
+	RawIndicesArrayMerged.Reserve(RawIndicesArray0.Num() + RawIndicesArray1.Num() + RawIndicesArray2.Num());
+	RawIndicesArrayMerged.Append(RawIndicesArray0);
+	RawIndicesArrayMerged.Append(RawIndicesArray1);
+	RawIndicesArrayMerged.Append(RawIndicesArray2);
+}
+
+GlobalFracturedGeometry::~GlobalFracturedGeometry()
+{}
+
+
+}
\ No newline at end of file
diff --git a/Engine/Source/Programs/HeadlessChaos/Resource/FracturedGeometry.h b/Engine/Source/Programs/HeadlessChaos/Resource/FracturedGeometry.h
new file mode 100644
--- /dev/null
+++ b/Engine/Source/Programs/HeadlessChaos/Resource/FracturedGeometry.h
@@ -0,0 +1,50 @@
+// Copyright Epic Games, Inc. All Rights Reserved.
+
+#pragma once
+#include "Containers/Array.h"
+#include "Math/Transform.h"
+
+namespace GeometryCollectionTest
+{
+
+	class FracturedGeometry
+	{
+	public:
+		FracturedGeometry();
+		~FracturedGeometry();
+
+		static const TArray<float>	RawVertexArray;
+		static const TArray<int32>	RawIndicesArray;
+		static const TArray<int32>	RawBoneMapArray;
+		static const TArray<FTransform> RawTransformArray;
+		static const TArray<int32> RawLevelArray;
+		static const TArray<int32> RawParentArray;
+		static const TArray<TSet<int32>> RawChildrenArray;
+		static const TArray<int32> RawSimulationTypeArray;
+		static const TArray<int32> RawStatusFlagsArray;
+	};
+
+	/**
+	 * A version of \c FracturedGeometry in a global point pool.
+	 */
+	class GlobalFracturedGeometry
+	{
+	public:
+		GlobalFracturedGeometry();
+		~GlobalFracturedGeometry();
+
+		TArray<float> RawVertexArray;
+		TArray<int32>	RawIndicesArray0; // Randomly coincident to RawIndicesArray1
+		TArray<int32>	RawIndicesArray1; // Unchanged from FracturedGeometry
+		TArray<int32>	RawIndicesArray2; // Offset in Y
+		const TArray<int32> &RawIndicesArray; // Points to RawIndicesArray1
+		TArray<int32> RawIndicesArrayMerged; // RawIndicesArray 0, 1, and 2 concatenated
+		const TArray<int32>	RawBoneMapArray;
+		const TArray<FTransform> RawTransformArray;
+		const TArray<int32> RawLevelArray;
+		const TArray<int32> RawParentArray;
+		const TArray<TSet<int32>> RawChildrenArray;
+		const TArray<int32> RawSimulationTypeArray;
+		const TArray<int32> RawStatusFlagsArray;
+	};
+}
\ No newline at end of file
