Enable Async by default.

diff --git a/Engine/Source/Programs/HeadlessChaos/Private/HeadlessChaosTestPhysicsScene.cpp b/Engine/Source/Programs/HeadlessChaos/Private/HeadlessChaosTestPhysicsScene.cpp
--- a/Engine/Source/Programs/HeadlessChaos/Private/HeadlessChaosTestPhysicsScene.cpp
+++ b/Engine/Source/Programs/HeadlessChaos/Private/HeadlessChaosTestPhysicsScene.cpp
@@ -126,7 +126,7 @@ namespace ChaosTest {
 		for(int Step = 1; Step < 10; ++Step)
 		{
 			FVec3 Grav(0,0,-1);
-			Scene.SetUpForFrame(&Grav, 1,99999,99999,10,false);
+			Scene.SetUpForFrame(&Grav, 1,99999,99999,10,false,0.033f);
 			Scene.StartFrame();
 			Scene.GetSolver()->GetEvolution()->FlushSpatialAcceleration();	//make sure we get a new tree every step
 			Scene.EndFrame();
@@ -146,7 +146,7 @@ namespace ChaosTest {
 		EXPECT_EQ(Scene.GetSpacialAcceleration()->GetSyncTimestamp(), 0);	//timestamp of 0 because we flush when scene is created
 
 		FVec3 Grav(0, 0, -1);
-		Scene.SetUpForFrame(&Grav, 1, 99999, 99999, 10, false);
+		Scene.SetUpForFrame(&Grav, 1, 99999, 99999, 10, false,0.033f);
 
 		// Game thread enqueues second solver task before first completes (we did not execute advance task)
 		Scene.StartFrame();
@@ -186,7 +186,7 @@ namespace ChaosTest {
 		EXPECT_EQ(Scene.GetSpacialAcceleration()->GetSyncTimestamp(),0);	//timestamp of 0 because we flush when scene is created
 
 		FVec3 Grav(0,0,-1);
-		Scene.SetUpForFrame(&Grav, 1,99999,99999,10,false);
+		Scene.SetUpForFrame(&Grav, 1,99999,99999,10,false,0.033f);
 
 		// PT not finished yet (we didn't execute solver task), should still be 0.
 		Scene.StartFrame();
@@ -226,7 +226,7 @@ namespace ChaosTest {
 		Scene.GetSolver()->SetStealAdvanceTasks_ForTesting(true); // prevents execution on StartFrame so we can execute task manually.
 
 		FVec3 Grav(0,0,-1);
-		Scene.SetUpForFrame(&Grav, 1,99999,99999,10,false);
+		Scene.SetUpForFrame(&Grav, 1,99999,99999,10,false,0.033f);
 
 		FActorCreationParams Params;
 		Params.Scene = &Scene;
@@ -345,7 +345,7 @@ namespace ChaosTest {
 		//tick solver but don't call EndFrame (want to flush and swap manually)
 		{
 			FVec3 Grav(0,0,-1);
-			Scene.SetUpForFrame(&Grav,1,99999,99999,10,false);
+			Scene.SetUpForFrame(&Grav,1,99999,99999,10,false,0.033f);
 			Scene.StartFrame();
 		}
 
@@ -388,7 +388,7 @@ namespace ChaosTest {
 		//tick solver so that particle is created, but don't call EndFrame (want to flush and swap manually)
 		{
 			FVec3 Grav(0,0,-1);
-			Scene.SetUpForFrame(&Grav,1,99999,99999,10,false);
+			Scene.SetUpForFrame(&Grav,1,99999,99999,10,false,0.033f);
 			Scene.StartFrame();
 		}
 
@@ -431,7 +431,7 @@ namespace ChaosTest {
 		//tick solver so that particle is created, but don't call EndFrame (want to flush and swap manually)
 		{
 			FVec3 Grav(0,0,-1);
-			Scene.SetUpForFrame(&Grav,1,99999,99999,10,false);
+			Scene.SetUpForFrame(&Grav,1,99999,99999,10,false,0.033f);
 			Scene.StartFrame();
 		}
 
@@ -474,7 +474,7 @@ namespace ChaosTest {
 		{
 			//use 0 dt to make sure pending operations are not sensitive to 0 dt
 			FVec3 Grav(0,0,-1);
-			Scene.SetUpForFrame(&Grav,0,99999,99999,10,false);
+			Scene.SetUpForFrame(&Grav,0,99999,99999,10,false,0.033f);
 			Scene.StartFrame();
 		}
 
@@ -504,7 +504,7 @@ namespace ChaosTest {
 		//tick solver, but don't call EndFrame (want to flush and swap manually)
 		{
 			FVec3 Grav(0,0,-1);
-			Scene.SetUpForFrame(&Grav,1,99999,99999,10,false);
+			Scene.SetUpForFrame(&Grav,1,99999,99999,10,false,0.033f);
 			Scene.StartFrame();
 		}
 
@@ -563,7 +563,7 @@ namespace ChaosTest {
 				//tick solver
 				{
 					FVec3 Grav(0,0,-1);
-					Scene.SetUpForFrame(&Grav,1,99999,99999,10,false);
+					Scene.SetUpForFrame(&Grav,1,99999,99999,10,false,0.033f);
 					Scene.StartFrame();
 					Scene.EndFrame();
 				}
@@ -585,7 +585,7 @@ namespace ChaosTest {
 			//tick solver one last time
 			{
 				FVec3 Grav(0,0,-1);
-				Scene.SetUpForFrame(&Grav,1,99999,99999,10,false);
+				Scene.SetUpForFrame(&Grav,1,99999,99999,10,false,0.033f);
 				Scene.StartFrame();
 				Scene.EndFrame();
 			}
@@ -604,7 +604,7 @@ namespace ChaosTest {
 				//tick solver
 				{
 					FVec3 Grav(0,0,-1);
-					Scene.SetUpForFrame(&Grav,1,99999,99999,10,false);
+					Scene.SetUpForFrame(&Grav,1,99999,99999,10,false,0.033f);
 					Scene.StartFrame();
 					Scene.EndFrame();
 				}
@@ -621,7 +621,7 @@ namespace ChaosTest {
 			//tick solver one last time
 			{
 				FVec3 Grav(0,0,-1);
-				Scene.SetUpForFrame(&Grav,1,99999,99999,10,false);
+				Scene.SetUpForFrame(&Grav,1,99999,99999,10,false,0.033f);
 				Scene.StartFrame();
 				Scene.EndFrame();
 			}
@@ -651,7 +651,7 @@ namespace ChaosTest {
 			{
 				//tick solver
 				FVec3 Grav(0,0,-1);
-				Scene.SetUpForFrame(&Grav,1,99999,99999,10,false);
+				Scene.SetUpForFrame(&Grav,1,99999,99999,10,false,0.033f);
 				Scene.StartFrame();
 				Scene.EndFrame();
 
@@ -709,7 +709,7 @@ namespace ChaosTest {
 			{
 				{
 					FVec3 Grav(0,0,0);
-					Scene.SetUpForFrame(&Grav,1,99999,99999,10,false);
+					Scene.SetUpForFrame(&Grav,1,99999,99999,10,false,0.033f);
 					Scene.StartFrame();
 					Scene.EndFrame();
 				}
@@ -722,7 +722,7 @@ namespace ChaosTest {
 			//tick solver and see new position synced from sim
 			{
 				FVec3 Grav(0,0,0);
-				Scene.SetUpForFrame(&Grav,1,99999,99999,10,false);
+				Scene.SetUpForFrame(&Grav,1,99999,99999,10,false,0.033f);
 				Scene.StartFrame();
 				Scene.EndFrame();
 				EXPECT_NEAR(Particle->X()[2],-1, 1e-4);
@@ -732,7 +732,7 @@ namespace ChaosTest {
 			//tick solver and delete in between solver finishing and sync
 			{
 				FVec3 Grav(0,0,0);
-				Scene.SetUpForFrame(&Grav,1,99999,99999,10,false);
+				Scene.SetUpForFrame(&Grav,1,99999,99999,10,false,0.033f);
 				Scene.StartFrame();
 
 				//delete particle
@@ -748,7 +748,7 @@ namespace ChaosTest {
 			{
 				{
 					FVec3 Grav(0,0,0);
-					Scene.SetUpForFrame(&Grav,1,99999,99999,10,false);
+					Scene.SetUpForFrame(&Grav,1,99999,99999,10,false,0.033f);
 					Scene.StartFrame();
 					Scene.EndFrame();
 					EXPECT_NEAR(Particle2->X()[1],-3 - Repeat, 1e-4);	//other particle keeps moving
@@ -791,7 +791,7 @@ namespace ChaosTest {
 			{
 				{
 					FVec3 Grav(0,0,0);
-					Scene.SetUpForFrame(&Grav,1,99999,99999,10,false);
+					Scene.SetUpForFrame(&Grav,1,99999,99999,10,false,0.033f);
 					Scene.StartFrame();
 					Scene.EndFrame();
 				}
@@ -803,7 +803,7 @@ namespace ChaosTest {
 			//tick solver and see new position synced from sim
 			{
 				FVec3 Grav(0,0,0);
-				Scene.SetUpForFrame(&Grav,1,99999,99999,10,false);
+				Scene.SetUpForFrame(&Grav,1,99999,99999,10,false,0.033f);
 				Scene.StartFrame();
 				Scene.EndFrame();
 				EXPECT_NEAR(Particle->X()[2],-1,1e-4);
@@ -816,7 +816,7 @@ namespace ChaosTest {
 			{
 				{
 					FVec3 Grav(0,0,0);
-					Scene.SetUpForFrame(&Grav,1,99999,99999,10,false);
+					Scene.SetUpForFrame(&Grav,1,99999,99999,10,false,0.033f);
 					Scene.StartFrame();
 					Scene.EndFrame();
 
@@ -827,7 +827,7 @@ namespace ChaosTest {
 			//tick solver one last time, should see sim results from the place we teleported to
 			{
 				FVec3 Grav(0,0,0);
-				Scene.SetUpForFrame(&Grav,1,99999,99999,10,false);
+				Scene.SetUpForFrame(&Grav,1,99999,99999,10,false,0.033f);
 				Scene.StartFrame();
 				Scene.EndFrame();
 				EXPECT_NEAR(Particle->X()[2],9,1e-4);
@@ -836,7 +836,7 @@ namespace ChaosTest {
 			//set x after sim but before EndFrame, make sure to see gt position since it was written after
 			{
 				FVec3 Grav(0,0,0);
-				Scene.SetUpForFrame(&Grav,1,99999,99999,10,false);
+				Scene.SetUpForFrame(&Grav,1,99999,99999,10,false,0.033f);
 				Scene.StartFrame();
 				FChaosEngineInterface::SetGlobalPose_AssumesLocked(Particle,FTransform(FQuat::Identity,FVec3(0,0,100)));
 				Scene.EndFrame();
@@ -847,7 +847,7 @@ namespace ChaosTest {
 			{
 				{
 					FVec3 Grav(0,0,0);
-					Scene.SetUpForFrame(&Grav,1,99999,99999,10,false);
+					Scene.SetUpForFrame(&Grav,1,99999,99999,10,false,0.033f);
 					Scene.StartFrame();
 					Scene.EndFrame();
 
@@ -858,7 +858,7 @@ namespace ChaosTest {
 			//tick solver one last time, should see sim results from the place we teleported to
 			{
 				FVec3 Grav(0,0,0);
-				Scene.SetUpForFrame(&Grav,1,99999,99999,10,false);
+				Scene.SetUpForFrame(&Grav,1,99999,99999,10,false,0.033f);
 				Scene.StartFrame();
 				Scene.EndFrame();
 				EXPECT_NEAR(Particle->X()[2],99,1e-4);
@@ -890,7 +890,7 @@ namespace ChaosTest {
 		Simulated->SetF(FVec3(0,0,10) * Simulated->M());
 
 		FVec3 Grav(0,0,0);
-		Scene.SetUpForFrame(&Grav,1,99999,99999,10,false);
+		Scene.SetUpForFrame(&Grav,1,99999,99999,10,false,0.033f);
 		Scene.StartFrame();
 		Scene.EndFrame();
 
@@ -947,7 +947,7 @@ namespace ChaosTest {
 		for(int32 i=0; i<12;i++)
 		{
 			FVec3 Grav(0, 0, 0);
-			Scene.SetUpForFrame(&Grav, FixedDT * 0.25, 99999, 99999, 10, false);
+			Scene.SetUpForFrame(&Grav, FixedDT * 0.25, 99999, 99999, 10, false,0.033f);
 			Scene.StartFrame();
 			Scene.EndFrame();
 		}
diff --git a/Engine/Source/Runtime/Engine/Classes/PhysicsEngine/PhysicsSettings.h b/Engine/Source/Runtime/Engine/Classes/PhysicsEngine/PhysicsSettings.h
--- a/Engine/Source/Runtime/Engine/Classes/PhysicsEngine/PhysicsSettings.h
+++ b/Engine/Source/Runtime/Engine/Classes/PhysicsEngine/PhysicsSettings.h
@@ -172,6 +172,14 @@ class ENGINE_API UPhysicsSettings : public UPhysicsSettingsCore
 	UPROPERTY(config, EditAnywhere, Category = Framerate)
 	bool bSubsteppingAsync;
 
+	/** Whether to tick physics simulation on an async thread. This feature is still experimental. Certain functionality might not work correctly*/
+	UPROPERTY(config, EditAnywhere, Category = Framerate)
+	bool bTickPhysicsAsync;
+
+	/** If using async, the time step size to tick at. This feature is still experimental. Certain functionality might not work correctly*/
+	UPROPERTY(config, EditAnywhere, Category = Framerate, meta=(editcondition = "bTickPhysicsAsync"))
+	float AsyncFixedTimeStepSize;
+
 	/** Max delta time (in seconds) for an individual simulation substep. */
 	UPROPERTY(config, EditAnywhere, meta = (ClampMin = "0.0013", UIMin = "0.0013", ClampMax = "1.0", UIMax = "1.0", editcondition = "bSubStepping"), Category=Framerate)
 	float MaxSubstepDeltaTime;
diff --git a/Engine/Source/Runtime/Engine/Private/PhysicsEngine/PhysLevel.cpp b/Engine/Source/Runtime/Engine/Private/PhysicsEngine/PhysLevel.cpp
--- a/Engine/Source/Runtime/Engine/Private/PhysicsEngine/PhysLevel.cpp
+++ b/Engine/Source/Runtime/Engine/Private/PhysicsEngine/PhysLevel.cpp
@@ -144,7 +144,7 @@ void UWorld::SetupPhysicsTickFunctions(float DeltaSeconds)
 
 	static const auto CVar_MaxPhysicsDeltaTime = IConsoleManager::Get().FindTConsoleVariableDataFloat(TEXT("p.MaxPhysicsDeltaTime"));
 	PhysScene->SetUpForFrame(&DefaultGravity, DeltaSeconds, UPhysicsSettings::Get()->MaxPhysicsDeltaTime,
-		UPhysicsSettings::Get()->MaxSubstepDeltaTime, UPhysicsSettings::Get()->MaxSubsteps, UPhysicsSettings::Get()->bSubstepping);
+		UPhysicsSettings::Get()->MaxSubstepDeltaTime, UPhysicsSettings::Get()->MaxSubsteps, UPhysicsSettings::Get()->bSubstepping, UPhysicsSettings::Get()->bTickPhysicsAsync ? UPhysicsSettings::Get()->AsyncFixedTimeStepSize : -1);
 }
 
 void UWorld::StartPhysicsSim()
diff --git a/Engine/Source/Runtime/Engine/Private/PhysicsEngine/PhysScene_PhysX.cpp b/Engine/Source/Runtime/Engine/Private/PhysicsEngine/PhysScene_PhysX.cpp
--- a/Engine/Source/Runtime/Engine/Private/PhysicsEngine/PhysScene_PhysX.cpp
+++ b/Engine/Source/Runtime/Engine/Private/PhysicsEngine/PhysScene_PhysX.cpp
@@ -1461,7 +1461,7 @@ void FPhysScene_PhysX::DispatchPhysNotifications_AssumesLocked()
 
 }
 
-void FPhysScene_PhysX::SetUpForFrame(const FVector* NewGrav, float InDeltaSeconds, float InMaxPhysicsDeltaTime, float InMaxSubstepDeltaTime, int32 InMaxSubsteps, bool bUnused)
+void FPhysScene_PhysX::SetUpForFrame(const FVector* NewGrav, float InDeltaSeconds, float InMaxPhysicsDeltaTime, float InMaxSubstepDeltaTime, int32 InMaxSubsteps, bool bUnused, float InUnused)
 {
 	DeltaSeconds = InDeltaSeconds;
 	MaxPhysicsDeltaTime = InMaxPhysicsDeltaTime;
diff --git a/Engine/Source/Runtime/Engine/Private/PhysicsEngine/PhysicsSettings.cpp b/Engine/Source/Runtime/Engine/Private/PhysicsEngine/PhysicsSettings.cpp
--- a/Engine/Source/Runtime/Engine/Private/PhysicsEngine/PhysicsSettings.cpp
+++ b/Engine/Source/Runtime/Engine/Private/PhysicsEngine/PhysicsSettings.cpp
@@ -19,6 +19,8 @@ UPhysicsSettings::UPhysicsSettings(const FObjectInitializer& ObjectInitializer)
 	, bSimulateAnimPhysicsAfterReset(false)
 	, MaxPhysicsDeltaTime(1.f / 30.f)
 	, bSubstepping(false)
+	, bTickPhysicsAsync(true)
+	, AsyncFixedTimeStepSize(1.f / 30.f)
 	, MaxSubstepDeltaTime(1.f / 60.f)
 	, MaxSubsteps(6)
 	, SyncSceneSmoothingFactor(0.0f)
diff --git a/Engine/Source/Runtime/Engine/Private/World.cpp b/Engine/Source/Runtime/Engine/Private/World.cpp
--- a/Engine/Source/Runtime/Engine/Private/World.cpp
+++ b/Engine/Source/Runtime/Engine/Private/World.cpp
@@ -1609,7 +1609,7 @@ void UWorld::InitWorld(const InitializationValues IVS)
 	if (GetPhysicsScene())
 	{
 		FVector Gravity = FVector( 0.f, 0.f, GetGravityZ() );
-		GetPhysicsScene()->SetUpForFrame( &Gravity, 0, 0, 0, 0, false );
+		GetPhysicsScene()->SetUpForFrame( &Gravity, 0, 0, 0, 0, false, -1 );
 	}
 
 	// Create physics collision handler, if we have a physics scene
diff --git a/Engine/Source/Runtime/Engine/Public/Physics/PhysScene_PhysX.h b/Engine/Source/Runtime/Engine/Public/Physics/PhysScene_PhysX.h
--- a/Engine/Source/Runtime/Engine/Public/Physics/PhysScene_PhysX.h
+++ b/Engine/Source/Runtime/Engine/Public/Physics/PhysScene_PhysX.h
@@ -168,7 +168,7 @@ public:
 	void ApplyWorldOffset(FVector InOffset);
 
 	/** Set the gravity and timing of all physics scenes */
-	ENGINE_API void SetUpForFrame(const FVector* NewGrav, float InDeltaSeconds = 0.0f, float InMaxPhysicsDeltaTime = 0.0f, float InMaxSubstepDeltaTime = 0.f, int32 InMaxSubsteps = 1, bool bUnused = false);
+	ENGINE_API void SetUpForFrame(const FVector* NewGrav, float InDeltaSeconds = 0.0f, float InMaxPhysicsDeltaTime = 0.0f, float InMaxSubstepDeltaTime = 0.f, int32 InMaxSubsteps = 1, bool bUnused = false, float InUnused = 0);
 
 	/** Starts a frame */
 	ENGINE_API void StartFrame();
diff --git a/Engine/Source/Runtime/Experimental/Chaos/Private/Chaos/Framework/PhysicsSolverBase.cpp b/Engine/Source/Runtime/Experimental/Chaos/Private/Chaos/Framework/PhysicsSolverBase.cpp
--- a/Engine/Source/Runtime/Experimental/Chaos/Private/Chaos/Framework/PhysicsSolverBase.cpp
+++ b/Engine/Source/Runtime/Experimental/Chaos/Private/Chaos/Framework/PhysicsSolverBase.cpp
@@ -105,9 +105,6 @@ namespace Chaos
 		Solver.AdvanceSolverBy(Dt);
 	}
 
-	CHAOS_API float DefaultAsyncDt = -1;
-	FAutoConsoleVariableRef CVarDefaultAsyncDt(TEXT("p.DefaultAsyncDt"), DefaultAsyncDt,TEXT("Whether to use async results -1 means not async"));
-
 	CHAOS_API int32 UseAsyncInterpolation = 1;
 	FAutoConsoleVariableRef CVarUseAsyncInterpolation(TEXT("p.UseAsyncInterpolation"), UseAsyncInterpolation, TEXT("Whether to interpolate when async mode is enabled"));
 
@@ -122,7 +119,7 @@ namespace Chaos
 		, Owner(InOwner)
 		, ExternalDataLock_External(new FPhysicsSceneGuard())
 		, TraitIdx(InTraitIdx)
-		, AsyncDt(DefaultAsyncDt)
+		, AsyncDt(-1)
 #if !UE_BUILD_SHIPPING
 		, bStealAdvanceTasksForTesting(false)
 #endif
diff --git a/Engine/Source/Runtime/PhysicsCore/Private/ChaosScene.cpp b/Engine/Source/Runtime/PhysicsCore/Private/ChaosScene.cpp
--- a/Engine/Source/Runtime/PhysicsCore/Private/ChaosScene.cpp
+++ b/Engine/Source/Runtime/PhysicsCore/Private/ChaosScene.cpp
@@ -285,7 +285,7 @@ void FChaosScene::AddActorsToScene_AssumesLocked(TArray<FPhysicsActorHandle>& In
 #endif
 }
 
-void FChaosScene::SetUpForFrame(const FVector* NewGrav,float InDeltaSeconds /*= 0.0f*/,float InMaxPhysicsDeltaTime /*= 0.0f*/,float InMaxSubstepDeltaTime /*= 0.0f*/,int32 InMaxSubsteps,bool bSubstepping)
+void FChaosScene::SetUpForFrame(const FVector* NewGrav,float InDeltaSeconds /*= 0.0f*/,float InMaxPhysicsDeltaTime /*= 0.0f*/,float InMaxSubstepDeltaTime /*= 0.0f*/,int32 InMaxSubsteps,bool bSubstepping,float InAsyncDt)
 {
 #if WITH_CHAOS
 	using namespace Chaos;
@@ -294,6 +294,14 @@ void FChaosScene::SetUpForFrame(const FVector* NewGrav,float InDeltaSeconds /*=
 
 	if(FPhysicsSolver* Solver = GetSolver())
 	{
+		if (InAsyncDt > 0)
+		{
+			Solver->EnableAsyncMode(InAsyncDt);
+		}
+		else
+		{
+			Solver->DisableAsyncMode();
+		}
 		if(bSubstepping)
 		{
 			Solver->SetMaxDeltaTime(InMaxSubstepDeltaTime);
diff --git a/Engine/Source/Runtime/PhysicsCore/Public/Chaos/ChaosScene.h b/Engine/Source/Runtime/PhysicsCore/Public/Chaos/ChaosScene.h
--- a/Engine/Source/Runtime/PhysicsCore/Public/Chaos/ChaosScene.h
+++ b/Engine/Source/Runtime/PhysicsCore/Public/Chaos/ChaosScene.h
@@ -116,7 +116,7 @@ public:
 #endif
 
 	void StartFrame();
-	void SetUpForFrame(const FVector* NewGrav,float InDeltaSeconds,float InMaxPhysicsDeltaTime,float InMaxSubstepDeltaTime,int32 InMaxSubsteps,bool bSubstepping);
+	void SetUpForFrame(const FVector* NewGrav,float InDeltaSeconds,float InMaxPhysicsDeltaTime,float InMaxSubstepDeltaTime,int32 InMaxSubsteps,bool bSubstepping,float InAsyncDt);
 	void EndFrame();
 
 	DECLARE_MULTICAST_DELEGATE_OneParam(FOnPhysScenePostTick,FChaosScene*);
