Fix async accumulated time resetting every frame. Also fix unit tests using fixed dt when they shouldn't

diff --git a/Engine/Source/Programs/HeadlessChaos/Private/HeadlessChaosTestPhysicsScene.cpp b/Engine/Source/Programs/HeadlessChaos/Private/HeadlessChaosTestPhysicsScene.cpp
--- a/Engine/Source/Programs/HeadlessChaos/Private/HeadlessChaosTestPhysicsScene.cpp
+++ b/Engine/Source/Programs/HeadlessChaos/Private/HeadlessChaosTestPhysicsScene.cpp
@@ -130,7 +130,7 @@ namespace ChaosTest {
 		for(int Step = 1; Step < 10; ++Step)
 		{
 			FVec3 Grav(0,0,-1);
-			Scene.SetUpForFrame(&Grav, 1,99999,99999,10,false,0.033f);
+			Scene.SetUpForFrame(&Grav, 1,99999,99999,10,false, -1);
 			Scene.StartFrame();
 			Scene.GetSolver()->GetEvolution()->FlushSpatialAcceleration();	//make sure we get a new tree every step
 			Scene.EndFrame();
@@ -150,7 +150,7 @@ namespace ChaosTest {
 		EXPECT_EQ(Scene.GetSpacialAcceleration()->GetSyncTimestamp(), 0);	//timestamp of 0 because we flush when scene is created
 
 		FVec3 Grav(0, 0, -1);
-		Scene.SetUpForFrame(&Grav, 1, 99999, 99999, 10, false,0.033f);
+		Scene.SetUpForFrame(&Grav, 1, 99999, 99999, 10, false, -1);
 
 		// Game thread enqueues second solver task before first completes (we did not execute advance task)
 		Scene.StartFrame();
@@ -192,7 +192,7 @@ namespace ChaosTest {
 		EXPECT_EQ(Scene.GetSpacialAcceleration()->GetSyncTimestamp(),0);	//timestamp of 0 because we flush when scene is created
 
 		FVec3 Grav(0,0,-1);
-		Scene.SetUpForFrame(&Grav, 1,99999,99999,10,false,0.033f);
+		Scene.SetUpForFrame(&Grav, 1,99999,99999,10,false, -1);
 
 		// PT not finished yet (we didn't execute solver task), should still be 0.
 		Scene.StartFrame();
@@ -235,7 +235,7 @@ namespace ChaosTest {
 		Scene.GetSolver()->SetStealAdvanceTasks_ForTesting(true); // prevents execution on StartFrame so we can execute task manually.
 
 		FVec3 Grav(0,0,-1);
-		Scene.SetUpForFrame(&Grav, 1,99999,99999,10,false,0.033f);
+		Scene.SetUpForFrame(&Grav, 1,99999,99999,10,false, -1);
 
 		FActorCreationParams Params;
 		Params.Scene = &Scene;
@@ -348,7 +348,7 @@ namespace ChaosTest {
 		//tick solver but don't call EndFrame (want to flush and swap manually)
 		{
 			FVec3 Grav(0,0,-1);
-			Scene.SetUpForFrame(&Grav,1,99999,99999,10,false,0.033f);
+			Scene.SetUpForFrame(&Grav,1,99999,99999,10,false, -1);
 			Scene.StartFrame();
 		}
 
@@ -391,7 +391,7 @@ namespace ChaosTest {
 		//tick solver so that particle is created, but don't call EndFrame (want to flush and swap manually)
 		{
 			FVec3 Grav(0,0,-1);
-			Scene.SetUpForFrame(&Grav,1,99999,99999,10,false,0.033f);
+			Scene.SetUpForFrame(&Grav,1,99999,99999,10,false, -1);
 			Scene.StartFrame();
 		}
 
@@ -434,7 +434,7 @@ namespace ChaosTest {
 		//tick solver so that particle is created, but don't call EndFrame (want to flush and swap manually)
 		{
 			FVec3 Grav(0,0,-1);
-			Scene.SetUpForFrame(&Grav,1,99999,99999,10,false,0.033f);
+			Scene.SetUpForFrame(&Grav,1,99999,99999,10,false, -1);
 			Scene.StartFrame();
 		}
 
@@ -477,7 +477,7 @@ namespace ChaosTest {
 		{
 			//use 0 dt to make sure pending operations are not sensitive to 0 dt
 			FVec3 Grav(0,0,-1);
-			Scene.SetUpForFrame(&Grav,0,99999,99999,10,false,0.033f);
+			Scene.SetUpForFrame(&Grav,0,99999,99999,10,false, -1);
 			Scene.StartFrame();
 		}
 
@@ -507,7 +507,7 @@ namespace ChaosTest {
 		//tick solver, but don't call EndFrame (want to flush and swap manually)
 		{
 			FVec3 Grav(0,0,-1);
-			Scene.SetUpForFrame(&Grav,1,99999,99999,10,false,0.033f);
+			Scene.SetUpForFrame(&Grav,1,99999,99999,10,false, -1);
 			Scene.StartFrame();
 		}
 
@@ -566,7 +566,7 @@ namespace ChaosTest {
 				//tick solver
 				{
 					FVec3 Grav(0,0,-1);
-					Scene.SetUpForFrame(&Grav,1,99999,99999,10,false,0.033f);
+					Scene.SetUpForFrame(&Grav,1,99999,99999,10,false, -1);
 					Scene.StartFrame();
 					Scene.EndFrame();
 				}
@@ -588,7 +588,7 @@ namespace ChaosTest {
 			//tick solver one last time
 			{
 				FVec3 Grav(0,0,-1);
-				Scene.SetUpForFrame(&Grav,1,99999,99999,10,false,0.033f);
+				Scene.SetUpForFrame(&Grav,1,99999,99999,10,false, -1);
 				Scene.StartFrame();
 				Scene.EndFrame();
 			}
@@ -607,7 +607,7 @@ namespace ChaosTest {
 				//tick solver
 				{
 					FVec3 Grav(0,0,-1);
-					Scene.SetUpForFrame(&Grav,1,99999,99999,10,false,0.033f);
+					Scene.SetUpForFrame(&Grav,1,99999,99999,10,false, -1);
 					Scene.StartFrame();
 					Scene.EndFrame();
 				}
@@ -624,7 +624,7 @@ namespace ChaosTest {
 			//tick solver one last time
 			{
 				FVec3 Grav(0,0,-1);
-				Scene.SetUpForFrame(&Grav,1,99999,99999,10,false,0.033f);
+				Scene.SetUpForFrame(&Grav,1,99999,99999,10,false, -1);
 				Scene.StartFrame();
 				Scene.EndFrame();
 			}
@@ -654,7 +654,7 @@ namespace ChaosTest {
 			{
 				//tick solver
 				FVec3 Grav(0,0,-1);
-				Scene.SetUpForFrame(&Grav,1,99999,99999,10,false,0.033f);
+				Scene.SetUpForFrame(&Grav,1,99999,99999,10,false, -1);
 				Scene.StartFrame();
 				Scene.EndFrame();
 
@@ -712,7 +712,7 @@ namespace ChaosTest {
 			{
 				{
 					FVec3 Grav(0,0,0);
-					Scene.SetUpForFrame(&Grav,1,99999,99999,10,false,0.033f);
+					Scene.SetUpForFrame(&Grav,1,99999,99999,10,false, -1);
 					Scene.StartFrame();
 					Scene.EndFrame();
 				}
@@ -725,7 +725,7 @@ namespace ChaosTest {
 			//tick solver and see new position synced from sim
 			{
 				FVec3 Grav(0,0,0);
-				Scene.SetUpForFrame(&Grav,1,99999,99999,10,false,0.033f);
+				Scene.SetUpForFrame(&Grav,1,99999,99999,10,false, -1);
 				Scene.StartFrame();
 				Scene.EndFrame();
 				EXPECT_NEAR(Particle->X()[2],-1, 1e-4);
@@ -735,7 +735,7 @@ namespace ChaosTest {
 			//tick solver and delete in between solver finishing and sync
 			{
 				FVec3 Grav(0,0,0);
-				Scene.SetUpForFrame(&Grav,1,99999,99999,10,false,0.033f);
+				Scene.SetUpForFrame(&Grav,1,99999,99999,10,false, -1);
 				Scene.StartFrame();
 
 				//delete particle
@@ -751,7 +751,7 @@ namespace ChaosTest {
 			{
 				{
 					FVec3 Grav(0,0,0);
-					Scene.SetUpForFrame(&Grav,1,99999,99999,10,false,0.033f);
+					Scene.SetUpForFrame(&Grav,1,99999,99999,10,false, -1);
 					Scene.StartFrame();
 					Scene.EndFrame();
 					EXPECT_NEAR(Particle2->X()[1],-3 - Repeat, 1e-4);	//other particle keeps moving
@@ -794,7 +794,7 @@ namespace ChaosTest {
 			{
 				{
 					FVec3 Grav(0,0,0);
-					Scene.SetUpForFrame(&Grav,1,99999,99999,10,false,0.033f);
+					Scene.SetUpForFrame(&Grav,1,99999,99999,10,false, -1);
 					Scene.StartFrame();
 					Scene.EndFrame();
 				}
@@ -806,7 +806,7 @@ namespace ChaosTest {
 			//tick solver and see new position synced from sim
 			{
 				FVec3 Grav(0,0,0);
-				Scene.SetUpForFrame(&Grav,1,99999,99999,10,false,0.033f);
+				Scene.SetUpForFrame(&Grav,1,99999,99999,10,false, -1);
 				Scene.StartFrame();
 				Scene.EndFrame();
 				EXPECT_NEAR(Particle->X()[2],-1,1e-4);
@@ -819,7 +819,7 @@ namespace ChaosTest {
 			{
 				{
 					FVec3 Grav(0,0,0);
-					Scene.SetUpForFrame(&Grav,1,99999,99999,10,false,0.033f);
+					Scene.SetUpForFrame(&Grav,1,99999,99999,10,false, -1);
 					Scene.StartFrame();
 					Scene.EndFrame();
 
@@ -830,7 +830,7 @@ namespace ChaosTest {
 			//tick solver one last time, should see sim results from the place we teleported to
 			{
 				FVec3 Grav(0,0,0);
-				Scene.SetUpForFrame(&Grav,1,99999,99999,10,false,0.033f);
+				Scene.SetUpForFrame(&Grav,1,99999,99999,10,false, -1);
 				Scene.StartFrame();
 				Scene.EndFrame();
 				EXPECT_NEAR(Particle->X()[2],9,1e-4);
@@ -839,7 +839,7 @@ namespace ChaosTest {
 			//set x after sim but before EndFrame, make sure to see gt position since it was written after
 			{
 				FVec3 Grav(0,0,0);
-				Scene.SetUpForFrame(&Grav,1,99999,99999,10,false,0.033f);
+				Scene.SetUpForFrame(&Grav,1,99999,99999,10,false,-1);
 				Scene.StartFrame();
 				FChaosEngineInterface::SetGlobalPose_AssumesLocked(Particle,FTransform(FQuat::Identity,FVec3(0,0,100)));
 				Scene.EndFrame();
@@ -850,7 +850,7 @@ namespace ChaosTest {
 			{
 				{
 					FVec3 Grav(0,0,0);
-					Scene.SetUpForFrame(&Grav,1,99999,99999,10,false,0.033f);
+					Scene.SetUpForFrame(&Grav,1,99999,99999,10,false, -1);
 					Scene.StartFrame();
 					Scene.EndFrame();
 
@@ -861,7 +861,7 @@ namespace ChaosTest {
 			//tick solver one last time, should see sim results from the place we teleported to
 			{
 				FVec3 Grav(0,0,0);
-				Scene.SetUpForFrame(&Grav,1,99999,99999,10,false,0.033f);
+				Scene.SetUpForFrame(&Grav,1,99999,99999,10,false, -1);
 				Scene.StartFrame();
 				Scene.EndFrame();
 				EXPECT_NEAR(Particle->X()[2],99,1e-4);
@@ -893,7 +893,7 @@ namespace ChaosTest {
 		Simulated->AddForce(FVec3(0,0,10) * Simulated->M());
 
 		FVec3 Grav(0,0,0);
-		Scene.SetUpForFrame(&Grav,1,99999,99999,10,false,0.033f);
+		Scene.SetUpForFrame(&Grav,1,99999,99999,10,false,-1);
 		Scene.StartFrame();
 		Scene.EndFrame();
 
@@ -917,7 +917,6 @@ namespace ChaosTest {
 		FChaosScene Scene(nullptr);
 		Scene.GetSolver()->SetThreadingMode_External(EThreadingModeTemp::SingleThread);
 		const float FixedDT = 1;
-		Scene.GetSolver()->EnableAsyncMode(1);	//tick 1 dt at a time
 
 		FActorCreationParams Params;
 		Params.Scene = &Scene;
@@ -977,7 +976,7 @@ namespace ChaosTest {
 			//set force every external frame
 			Simulated2->AddForce(ConstantForce);
 			FVec3 Grav(0, 0, 0);
-			Scene.SetUpForFrame(&Grav, GTDt, 99999, 99999, 10, false,0.033f);
+			Scene.SetUpForFrame(&Grav, GTDt, 99999, 99999, 10, false,FixedDT);
 			Scene.StartFrame();
 			Scene.EndFrame();
 			
@@ -1014,7 +1013,6 @@ namespace ChaosTest {
 		FChaosScene Scene(nullptr);
 		Scene.GetSolver()->SetThreadingMode_External(EThreadingModeTemp::SingleThread);
 		const float FixedDT = 1;
-		Scene.GetSolver()->EnableAsyncMode(FixedDT);	//tick 1 dt at a time
 
 		FActorCreationParams Params;
 		Params.Scene = &Scene;
@@ -1066,7 +1064,7 @@ namespace ChaosTest {
 			//set force every external frame
 			Simulated->AddForce(FVec3(0, 0, 1 * Simulated->M()));	//should counteract gravity
 			FVec3 Grav(0, 0, -1);
-			Scene.SetUpForFrame(&Grav, GTDt, 99999, 99999, 10, false);
+			Scene.SetUpForFrame(&Grav, GTDt, 99999, 99999, 10, false, FixedDT);
 			Scene.StartFrame();
 			Scene.EndFrame();
 
diff --git a/Engine/Source/Runtime/Experimental/Chaos/Public/Chaos/Framework/PhysicsSolverBase.h b/Engine/Source/Runtime/Experimental/Chaos/Public/Chaos/Framework/PhysicsSolverBase.h
--- a/Engine/Source/Runtime/Experimental/Chaos/Public/Chaos/Framework/PhysicsSolverBase.h
+++ b/Engine/Source/Runtime/Experimental/Chaos/Public/Chaos/Framework/PhysicsSolverBase.h
@@ -238,8 +238,12 @@ namespace Chaos
 
 		void EnableAsyncMode(FReal FixedDt)
 		{
+			if (AsyncDt != FixedDt)
+			{
+				AccumulatedTime = 0;
+			}
+
 			AsyncDt = FixedDt;
-			AccumulatedTime = 0;
 		}
 
 		void DisableAsyncMode()
