HeadlessChaos
-Fix compile issue

diff --git a/Engine/Source/Programs/HeadlessChaos/Private/HeadlessChaosTestPhysicsScene.cpp b/Engine/Source/Programs/HeadlessChaos/Private/HeadlessChaosTestPhysicsScene.cpp
--- a/Engine/Source/Programs/HeadlessChaos/Private/HeadlessChaosTestPhysicsScene.cpp
+++ b/Engine/Source/Programs/HeadlessChaos/Private/HeadlessChaosTestPhysicsScene.cpp
@@ -1049,7 +1049,7 @@ namespace ChaosTest {
 			auto Callback = Scene.GetSolver()->CreateAndRegisterSimCallbackObject_External<FCallback>();
 
 			FVec3 Grav(0, 0, 0);
-			Scene.SetUpForFrame(&Grav, 0, 99999, 99999, 10, false);	//flush with dt 0
+			Scene.SetUpForFrame(&Grav, 0, 99999, 99999, 10, false,false);	//flush with dt 0
 			Scene.StartFrame();
 			Scene.EndFrame();
 
@@ -1182,7 +1182,7 @@ namespace ChaosTest {
 		Scene.GetSolver()->UnregisterObject(Particle);
 
 		FVec3 Grav(0, 0, -1);
-		Scene.SetUpForFrame(&Grav, FixedDT * 3, 99999, 99999, 10, false);
+		Scene.SetUpForFrame(&Grav, FixedDT * 3, 99999, 99999, 10, false, false);
 		Scene.StartFrame();
 		Scene.EndFrame();
 	}
