HeadlessChaos : Plasticity

- Added linear plasticity tests

diff --git a/Engine/Source/Programs/HeadlessChaos/Private/HeadlessChaosTestJoint.h b/Engine/Source/Programs/HeadlessChaos/Private/HeadlessChaosTestJoint.h
--- a/Engine/Source/Programs/HeadlessChaos/Private/HeadlessChaosTestJoint.h
+++ b/Engine/Source/Programs/HeadlessChaos/Private/HeadlessChaosTestJoint.h
@@ -61,11 +61,8 @@ namespace ChaosTest
 		}
 
 		// Create a pendulum chain along the specified direction with the first particle kinematic
-		void InitChain(int32 NumParticles, const FVec3& Dir)
+		void InitChain(int32 NumParticles, const FVec3& Dir, FReal Size = 10.f, FReal Separation = 30.f)
 		{
-			FReal Size = 10.0f;
-			FReal Separation = 30.0f;
-
 			for (int32 ParticleIndex = 0; ParticleIndex < NumParticles; ++ParticleIndex)
 			{
 				FReal D = ParticleIndex * Separation;
diff --git a/Engine/Source/Programs/HeadlessChaos/Private/HeadlessChaosTestJointPlasticity.cpp b/Engine/Source/Programs/HeadlessChaos/Private/HeadlessChaosTestJointPlasticity.cpp
new file mode 100644
--- /dev/null
+++ b/Engine/Source/Programs/HeadlessChaos/Private/HeadlessChaosTestJointPlasticity.cpp
@@ -0,0 +1,236 @@
+// Copyright Epic Games, Inc. All Rights Reserved.
+
+#include "HeadlessChaos.h"
+#include "HeadlessChaosTestJoint.h"
+#include "Modules/ModuleManager.h"
+
+namespace ChaosTest {
+
+	using namespace Chaos;
+
+
+	// 1 Kinematic Body with 1 Dynamic body held horizontally by a plastic angular constraint.
+	// Constraint plasticity limit is larger than resulting rotational settling so constraint will not bend.
+	template <typename TEvolution>
+	void JointPlasticity_UnderAngularPlasticityThreshold()
+	{
+		const FReal PlasticityAngle = 10;
+		const int32 NumIterations = 1;
+		const FReal Gravity = 980;
+		const FReal Dt = 0.01f;
+		int32 NumIts = 100;
+
+		FJointChainTest<TEvolution> Test(NumIterations, Gravity);
+		Test.InitChain(2, FVec3(0, 1, 0));
+
+		// Joint should break only if Threshold < MGL
+		// So not in this test
+		Test.JointSettings[0].bCollisionEnabled = false;
+		Test.JointSettings[0].LinearMotionTypes = { EJointMotionType::Locked, EJointMotionType::Locked , EJointMotionType::Locked };
+
+		Test.JointSettings[0].bAngularSLerpPositionDriveEnabled = true;
+		Test.JointSettings[0].bAngularSLerpVelocityDriveEnabled = true;
+		Test.JointSettings[0].AngularDriveDamping = 500;
+		Test.JointSettings[0].AngularDriveStiffness = 500000.f;
+		
+		Test.JointSettings[0].AngularPlasticityLimit = PlasticityAngle * (PI/180.);
+
+		Test.Create();
+		Test.AddParticleBox(FVec3(0, 30, 50), FRotation3::Identity, FVec3(10.f), 100.f);
+
+		FReal Angle = Test.Joints.GetConstraintSettings(0).AngularDrivePositionTarget.GetAngle() * (180. / PI);
+		EXPECT_TRUE(FMath::IsNearlyEqual(Angle, 0.f));
+
+		// Run the sim
+		for (int32 i = 0; i < NumIts; ++i)
+		{
+
+			Test.Evolution.AdvanceOneTimeStep(Dt);
+			Test.Evolution.EndFrame(Dt);
+
+			//FReal Angle = Test.Joints.GetConstraintSettings(0).AngularDrivePositionTarget.GetAngle() * (180. / PI);
+			//FVec3 Pos = Test.SOAs.GetDynamicParticles().X(0);
+			//std::cout << "["<< Angle <<"]" << Pos.X << "," << Pos.Y << "," << Pos.Z << std::endl;
+		}
+		Angle = Test.Joints.GetConstraintSettings(0).AngularDrivePositionTarget.GetAngle() * (180. / PI);
+
+		// Nothing should have been reset
+		EXPECT_TRUE(FMath::IsNearlyEqual(Angle, 0.f));
+		
+	}
+
+	TYPED_TEST(AllEvolutions, JointPlasticity_UnderAngularPlasticityThreshold)
+	{
+		JointPlasticity_UnderAngularPlasticityThreshold<TypeParam>();
+	}
+
+
+	// 1 Kinematic Body with 1 Dynamic body held horizontally by a plastic angular constraint.
+	// Constraint plasticity limit is larger than resulting rotational settling so constraint will not bend.
+	template <typename TEvolution>
+	void JointPlasticity_OverAngularPlasticityThreshold()
+	{
+		const FReal PlasticityAngle = 10;
+		const int32 NumIterations = 1;
+		const FReal Gravity = 980;
+		const FReal Dt = 0.01f;
+		int32 NumIts = 100;
+
+		FJointChainTest<TEvolution> Test(NumIterations, Gravity);
+		Test.InitChain(2, FVec3(0, 1, 0));
+
+		// Joint should break only if Threshold < MGL
+		// So not in this test
+		Test.JointSettings[0].bCollisionEnabled = false;
+		Test.JointSettings[0].LinearMotionTypes = { EJointMotionType::Locked, EJointMotionType::Locked , EJointMotionType::Locked };
+
+		Test.JointSettings[0].bAngularSLerpPositionDriveEnabled = true;
+		Test.JointSettings[0].bAngularSLerpVelocityDriveEnabled = true;
+		Test.JointSettings[0].AngularDriveDamping = 500;
+		Test.JointSettings[0].AngularDriveStiffness = 50000.f;
+
+		Test.JointSettings[0].AngularPlasticityLimit = PlasticityAngle * (PI / 180.);
+
+		Test.Create();
+		Test.AddParticleBox(FVec3(0, 30, 50), FRotation3::Identity, FVec3(10.f), 100.f);
+
+		FReal Angle = Test.Joints.GetConstraintSettings(0).AngularDrivePositionTarget.GetAngle() * (180. / PI);
+		EXPECT_TRUE(FMath::IsNearlyEqual(Angle, 0.f));
+
+		// Run the sim
+		for (int32 i = 0; i < NumIts; ++i)
+		{
+
+			Test.Evolution.AdvanceOneTimeStep(Dt);
+			Test.Evolution.EndFrame(Dt);
+
+			//FReal Angle = Test.Joints.GetConstraintSettings(0).AngularDrivePositionTarget.GetAngle() * (180. / PI);
+			//FVec3 Pos = Test.SOAs.GetDynamicParticles().X(0);
+			//std::cout << "["<< Angle <<"]" << Pos.X << "," << Pos.Y << "," << Pos.Z << std::endl;
+		}
+		Angle = Test.Joints.GetConstraintSettings(0).AngularDrivePositionTarget.GetAngle() * (180. / PI);
+
+		// The angle should have reset. 
+		EXPECT_FALSE(FMath::IsNearlyEqual(Angle, 0.f));
+		EXPECT_TRUE(FMath::IsNearlyEqual(Angle, PlasticityAngle, PlasticityAngle * 0.1f));
+
+	}
+
+	TYPED_TEST(AllEvolutions, JointPlasticity_OverAngularPlasticityThreshold)
+	{
+		JointPlasticity_OverAngularPlasticityThreshold<TypeParam>();
+	}
+
+
+	// 1 Kinematic Body with 1 Dynamic body held horizontally by a plastic angular constraint.
+	// Constraint plasticity limit is larger than resulting linear setting so constraint will not reset.
+	template <typename TEvolution>
+	void JointPlasticity_UnderLinearPlasticityThreshold()
+	{
+		const FReal PlasticityRatio = 0.3;
+		const int32 NumIterations = 1;
+		const FReal Gravity = 980;
+		const FReal Dt = 0.01f;
+		int32 NumIts = 200;
+
+		FJointChainTest<TEvolution> Test(NumIterations, Gravity);
+		Test.InitChain(2, FVec3(0, 0, 1), 10.f, 50.f);
+
+		Test.JointSettings[0].bCollisionEnabled = false;
+		Test.JointSettings[0].AngularMotionTypes = { EJointMotionType::Locked, EJointMotionType::Locked , EJointMotionType::Locked };
+		Test.JointSettings[0].LinearMotionTypes = { EJointMotionType::Limited, EJointMotionType::Limited , EJointMotionType::Limited };
+		Test.JointSettings[0].bSoftLinearLimitsEnabled = true;
+		Test.JointSettings[0].LinearLimit = 0;
+		Test.JointSettings[0].LinearSoftForceMode = EJointForceMode::Force;
+		Test.JointSettings[0].SoftLinearStiffness = 100000;
+		Test.JointSettings[0].SoftLinearDamping = 100;
+
+		Test.JointSettings[0].LinearPlasticityLimit = PlasticityRatio;
+
+		Test.Create();
+		Test.AddParticleBox(FVec3(0, 0, 100), FRotation3::Identity, FVec3(10.f), 100.f);
+
+		FReal DeltaPos = (Test.SOAs.GetDynamicParticles().X(0)).Size();
+		EXPECT_TRUE(FMath::IsNearlyEqual(DeltaPos, 50.f));
+
+		// Run the sim
+		for (int32 i = 0; i < NumIts; ++i)
+		{
+
+			Test.Evolution.AdvanceOneTimeStep(Dt);
+			Test.Evolution.EndFrame(Dt);
+
+			//FReal DeltaPosI = (Test.SOAs.GetDynamicParticles().X(0)).Size();
+			//FVector Pos = Test.SOAs.GetDynamicParticles().X(0);
+			//std::cout << "["<< DeltaPosI <<"]" << Pos.X << "," << Pos.Y << "," << Pos.Z << std::endl;
+		}
+		FReal DeltaPosPost = (Test.SOAs.GetDynamicParticles().X(0)).Size();
+
+		// Nothing should have reset
+		EXPECT_TRUE(FMath::IsNearlyEqual(DeltaPosPost, DeltaPos, 5.f));
+
+	}
+
+	TYPED_TEST(AllEvolutions, JointPlasticity_UnderLinearPlasticityThreshold)
+	{
+		JointPlasticity_UnderLinearPlasticityThreshold<TypeParam>();
+	}
+
+
+	// 1 Kinematic Body with 1 Dynamic body held horizontally by a plastic angular constraint.
+// Constraint plasticity limit is larger than resulting linear setting so constraint will not reset.
+	template <typename TEvolution>
+	void JointPlasticity_OverLinearPlasticityThreshold()
+	{
+		const FReal PlasticityRatio = 0.15;
+		const int32 NumIterations = 1;
+		const FReal Gravity = 980;
+		const FReal Dt = 0.01f;
+		int32 NumIts = 1000;
+
+		FJointChainTest<TEvolution> Test(NumIterations, Gravity);
+		Test.InitChain(2, FVec3(0, 0, 1), 10.f, 50.f);
+
+		Test.JointSettings[0].bCollisionEnabled = false;
+		Test.JointSettings[0].AngularMotionTypes = { EJointMotionType::Locked, EJointMotionType::Locked , EJointMotionType::Locked };
+		Test.JointSettings[0].LinearMotionTypes = { EJointMotionType::Limited, EJointMotionType::Limited , EJointMotionType::Limited };
+		Test.JointSettings[0].bSoftLinearLimitsEnabled = true;
+		Test.JointSettings[0].LinearLimit = 0;
+		Test.JointSettings[0].LinearSoftForceMode = EJointForceMode::Force;
+		Test.JointSettings[0].SoftLinearStiffness = 100000;
+		Test.JointSettings[0].SoftLinearDamping = 100;
+
+		Test.JointSettings[0].LinearPlasticityLimit = PlasticityRatio;
+
+		Test.Create();
+		Test.AddParticleBox(FVec3(0, 0, 100), FRotation3::Identity, FVec3(10.f), 100.f);
+
+		FReal DeltaPos = (Test.SOAs.GetDynamicParticles().X(0)).Size();
+		EXPECT_TRUE(FMath::IsNearlyEqual(DeltaPos, 50.f));
+
+		// Run the sim
+		for (int32 i = 0; i < NumIts; ++i)
+		{
+
+			Test.Evolution.AdvanceOneTimeStep(Dt);
+			Test.Evolution.EndFrame(Dt);
+
+			//FReal DeltaPosI = (Test.SOAs.GetDynamicParticles().X(0)).Size();
+			//FVector Pos = Test.SOAs.GetDynamicParticles().X(0);
+			//std::cout << "["<< DeltaPosI <<"]" << Pos.X << "," << Pos.Y << "," << Pos.Z << std::endl;
+		}
+		FReal DeltaPosPost = (Test.SOAs.GetDynamicParticles().X(0)).Size();
+
+		// The linear spring should have reset. 
+		EXPECT_TRUE(DeltaPosPost < DeltaPos * (1.f-PlasticityRatio) );
+		EXPECT_TRUE(Test.SOAs.GetDynamicParticles().X(0).Z > 0.f);
+	}
+
+	TYPED_TEST(AllEvolutions, JointPlasticity_OverLinearPlasticityThreshold)
+	{
+		JointPlasticity_OverLinearPlasticityThreshold<TypeParam>();
+	}
+
+
+
+}
\ No newline at end of file
