Chaos - Add extremely small normal check unit test for FConvex generation

diff --git a/Engine/Source/Programs/HeadlessChaos/Private/HeadlessChaosTestConvex.cpp b/Engine/Source/Programs/HeadlessChaos/Private/HeadlessChaosTestConvex.cpp
--- a/Engine/Source/Programs/HeadlessChaos/Private/HeadlessChaosTestConvex.cpp
+++ b/Engine/Source/Programs/HeadlessChaos/Private/HeadlessChaosTestConvex.cpp
@@ -193,4 +193,23 @@ namespace ChaosTest
 		}
 	}
 
+	// Check that extremely small generated triangle don't trigger the normal check
+	GTEST_TEST(ConvexStructureTests, TestConvexFaceNormalCheck)
+	{
+		// Create a long mesh with a extremely small end (YZ plane) 
+		// so that it generate extremely sized triangle that will produce extremely small (unormalized) normals
+		const float SmallNumber = 0.00001f;
+		const FVec3 Range{ 100.0f, SmallNumber, SmallNumber };
+
+		const FVec3 Vertices[] =
+		{
+			FVec3(0, 0, 0),
+			FVec3(Range.X, 0, 0),
+			FVec3(Range.X, Range.Y, 0),
+			FVec3(Range.X, Range.Y, Range.Z),
+			FVec3(Range.X+ SmallNumber, Range.Y*0.5f, Range.Z * 0.5f),
+		};
+
+		TestConvexStructureData(Vertices, UE_ARRAY_COUNT(Vertices));
+	}
 }
\ No newline at end of file
