Fix Missing Advance in the chaos field unit test

diff --git a/Engine/Source/Programs/HeadlessChaos/Private/GeometryCollection/GeometryCollectionTestSimulationField.cpp b/Engine/Source/Programs/HeadlessChaos/Private/GeometryCollection/GeometryCollectionTestSimulationField.cpp
--- a/Engine/Source/Programs/HeadlessChaos/Private/GeometryCollection/GeometryCollectionTestSimulationField.cpp
+++ b/Engine/Source/Programs/HeadlessChaos/Private/GeometryCollection/GeometryCollectionTestSimulationField.cpp
@@ -59,6 +59,7 @@ namespace GeometryCollectionTest
 		Collection->PhysObject->BufferCommand(UnitTest.Solver, { TargetNameTmp, RadialMaskTmp });
 
 		UnitTest.Initialize();
+		UnitTest.Advance();
 
 		UnitTest.Solver->RegisterSimOneShotCallback([&]()
 		{
