Headless Chaos -- fixed failing AdvancedDisabledClusteredObject test. The test was using what appears to be outdated SimulationType definitions.

diff --git a/Engine/Source/Programs/HeadlessChaos/Private/GeometryCollection/GeometryCollectionTestSimulationSolver.cpp b/Engine/Source/Programs/HeadlessChaos/Private/GeometryCollection/GeometryCollectionTestSimulationSolver.cpp
--- a/Engine/Source/Programs/HeadlessChaos/Private/GeometryCollection/GeometryCollectionTestSimulationSolver.cpp
+++ b/Engine/Source/Programs/HeadlessChaos/Private/GeometryCollection/GeometryCollectionTestSimulationSolver.cpp
@@ -108,9 +108,9 @@ using namespace ChaosTest;
 		UnitTest.Solver->RegisterSimOneShotCallback([&]()
 		{
 			TManagedArray<int32>& SimulationType = Collection->DynamicCollection->SimulationType;
-			EXPECT_EQ(SimulationType[0],FGeometryCollection::ESimulationTypes::FST_Clustered);
-			EXPECT_EQ(SimulationType[1],FGeometryCollection::ESimulationTypes::FST_Clustered);
-			EXPECT_EQ(SimulationType[2],FGeometryCollection::ESimulationTypes::FST_Rigid);
+			EXPECT_EQ(SimulationType[0],FGeometryCollection::ESimulationTypes::FST_Rigid);
+			EXPECT_EQ(SimulationType[1],FGeometryCollection::ESimulationTypes::FST_Rigid);
+			EXPECT_EQ(SimulationType[2],FGeometryCollection::ESimulationTypes::FST_Clustered);
 
 			TManagedArray<int32>& Parent = Collection->DynamicCollection->Parent;
 			EXPECT_EQ(Parent[0],2);
