Fix ambiguous function resolution issues in HeadlessChaos

diff --git a/Engine/Source/Programs/HeadlessChaos/Private/HeadlessChaosTestRaycast.cpp b/Engine/Source/Programs/HeadlessChaos/Private/HeadlessChaosTestRaycast.cpp
--- a/Engine/Source/Programs/HeadlessChaos/Private/HeadlessChaosTestRaycast.cpp
+++ b/Engine/Source/Programs/HeadlessChaos/Private/HeadlessChaosTestRaycast.cpp
@@ -438,7 +438,7 @@ namespace ChaosTest
 		EXPECT_FLOAT_EQ(Position.Z, 1);
 
 		//near hit by corner edge
-		const TVector<T, 3> StartEmptyRegion(1 - FMath::Sqrt(2) / 2, 1 - FMath::Sqrt(2) / 2, -1);
+		const TVector<T, 3> StartEmptyRegion(1 - FMath::Sqrt(static_cast<T>(2)) / 2, 1 - FMath::Sqrt(static_cast<T>(2)) / 2, -1);
 		bHit = Box.Raycast(StartEmptyRegion, TVector<T, 3>(0, 0, 1), 2, 1, Time, Position, Normal, FaceIndex);
 		EXPECT_TRUE(bHit);
 		EXPECT_EQ(FaceIndex, INDEX_NONE);
