Fix for shutdown hang in HeadlessChaos due to incorrect shutdown.
#tests build + run headlesschaos in the same configurations the builders use checking for correct shutdown

diff --git a/Engine/Source/Programs/HeadlessChaos/Private/HeadlessChaosApp.cpp b/Engine/Source/Programs/HeadlessChaos/Private/HeadlessChaosApp.cpp
--- a/Engine/Source/Programs/HeadlessChaos/Private/HeadlessChaosApp.cpp
+++ b/Engine/Source/Programs/HeadlessChaos/Private/HeadlessChaosApp.cpp
@@ -439,9 +439,9 @@ INT32_MAIN_INT32_ARGC_TCHAR_ARGV()
 
 	ensure(RUN_ALL_TESTS() == 0);
 
-	FCoreDelegates::OnExit.Broadcast();
+	FEngineLoop::AppPreExit();
 	FModuleManager::Get().UnloadModulesAtShutdown();
-
+	FEngineLoop::AppExit();
 	FPlatformMisc::RequestExit(false);
 
 	return 0;
