Chaos Headless -- fixed issues arising from convex particles refactor.

diff --git a/Engine/Source/Programs/HeadlessChaos/Private/HeadlessChaosTestGJK.cpp b/Engine/Source/Programs/HeadlessChaos/Private/HeadlessChaosTestGJK.cpp
--- a/Engine/Source/Programs/HeadlessChaos/Private/HeadlessChaosTestGJK.cpp
+++ b/Engine/Source/Programs/HeadlessChaos/Private/HeadlessChaosTestGJK.cpp
@@ -1411,7 +1411,7 @@ namespace ChaosTest
 		}
 
 		{
-			TArray<TVec3<T>> ConvexParticles;
+			TArray<FVec3> ConvexParticles;
 			ConvexParticles.SetNum(16);
 
 			ConvexParticles[0] ={-127.216454,203.240234,124.726524};
@@ -1480,7 +1480,7 @@ namespace ChaosTest
 
 		{
 			//capsule vs triangle as we make the sweep longer the world space point of impact should stay the same
-			TArray<TVec3<T>> ConvexParticles;
+			TArray<FVec3> ConvexParticles;
 			ConvexParticles.SetNum(3);
 
 			ConvexParticles[0] ={7400.00000, 12600.0000, 206.248123};
@@ -1491,7 +1491,7 @@ namespace ChaosTest
 			TSerializablePtr<FConvex> AConv(UniqueConvex);
 			const TImplicitObjectScaled<FConvex> AConvScaled(AConv,FVec3(1.0,1.0,1.0));
 
-			TTriangle<T> A(ConvexParticles[0],ConvexParticles[1],ConvexParticles[2]);
+			TTriangle<FReal> A(ConvexParticles[0],ConvexParticles[1],ConvexParticles[2]);
 
 			const FVec3 Pt0(0.0,0.0,-29.6999969);
 			FVec3 Pt1 = Pt0;
diff --git a/Engine/Source/Programs/HeadlessChaos/Private/HeadlessChaosTestImplicits.cpp b/Engine/Source/Programs/HeadlessChaos/Private/HeadlessChaosTestImplicits.cpp
--- a/Engine/Source/Programs/HeadlessChaos/Private/HeadlessChaosTestImplicits.cpp
+++ b/Engine/Source/Programs/HeadlessChaos/Private/HeadlessChaosTestImplicits.cpp
@@ -1235,7 +1235,7 @@ namespace ChaosTest {
 				Vertices[VertexIndex] = Particles.X(VertexIndex);
 			}
 			FConvex Convex(Vertices, 0.0f);
-			const TArray<TVec3<T>>& CulledParticles = Convex.GetVertices();
+			const TArray<FVec3>& CulledParticles = Convex.GetVertices();
 			EXPECT_EQ(CulledParticles.Num(), 8);
 
 			for (int32 Idx = 0; Idx < CulledParticles.Num(); ++Idx)
@@ -1274,30 +1274,30 @@ namespace ChaosTest {
 	{
 		{
 			//degenerates
-			TArray<Chaos::TVec3<T>> Particles;
+			TArray<FVec3> Particles;
 			Particles.SetNum(3);
-			Particles[0] = Chaos::TVec3<float>(-1, -1, -1);
-			Particles[1] = Chaos::TVec3<float>(1, -1, -1);
-			Particles[2] = Chaos::TVec3<float>(1, 1, -1);
+			Particles[0] = FVec3(-1, -1, -1);
+			Particles[1] = FVec3(1, -1, -1);
+			Particles[2] = FVec3(1, 1, -1);
 			TArray<TVector<int32, 3>>Indices;
 			Chaos::FConvexBuilder::BuildConvexHull(Particles, Indices);
 			EXPECT_EQ(Indices.Num(), 0);
-			Particles.Add(Chaos::TVec3<float>( 2, 3, -1 ));
+			Particles.Add(FVec3( 2, 3, -1 ));
 			Chaos::FConvexBuilder::BuildConvexHull(Particles, Indices);
 			EXPECT_EQ(Indices.Num(), 0);
 		}
 		{
-			TArray <Chaos::TVec3<T>> Particles;
+			TArray<FVec3> Particles;
 			Particles.SetNum(9);
-			Particles[0] = Chaos::TVec3<float>(-1, -1, -1);
-			Particles[1] = Chaos::TVec3<float>(-1, -1, 1);
-			Particles[2] = Chaos::TVec3<float>(-1, 1, -1);
-			Particles[3] = Chaos::TVec3<float>(-1, 1, 1);
-			Particles[4] = Chaos::TVec3<float>(1, -1, -1);
-			Particles[5] = Chaos::TVec3<float>(1, -1, 1);
-			Particles[6] = Chaos::TVec3<float>(1, 1, -1);
-			Particles[7] = Chaos::TVec3<float>(1, 1, 1);
-			Particles[8] = Chaos::TVec3<float>(0, 0, 0);
+			Particles[0] = FVec3(-1, -1, -1);
+			Particles[1] = FVec3(-1, -1, 1);
+			Particles[2] = FVec3(-1, 1, -1);
+			Particles[3] = FVec3(-1, 1, 1);
+			Particles[4] = FVec3(1, -1, -1);
+			Particles[5] = FVec3(1, -1, 1);
+			Particles[6] = FVec3(1, 1, -1);
+			Particles[7] = FVec3(1, 1, 1);
+			Particles[8] = FVec3(0, 0, 0);
 			TArray<TVector<int32, 3>>Indices;
 			Chaos::FConvexBuilder::BuildConvexHull(Particles, Indices);
 			EXPECT_EQ(Indices.Num(), 12);
@@ -1309,26 +1309,26 @@ namespace ChaosTest {
 			}
 		}
 		{
-			TArray<TVec3<T>> Particles;
+			TArray<FVec3> Particles;
 			Particles.SetNum(5);
-			Particles[0] = Chaos::TVec3<float>(-1, -1, -1);
-			Particles[1] = Chaos::TVec3<float>(1, -1, -1);
-			Particles[2] = Chaos::TVec3<float>(1, 1, -1);
-			Particles[3] = Chaos::TVec3<float>(0, 0, 0.5);
-			Particles[4] = (Particles[3]- Particles[1]) * 0.5 + Particles[1]+ Chaos::TVec3<float>(0, 0, 0.1);
+			Particles[0] = FVec3(-1, -1, -1);
+			Particles[1] = FVec3(1, -1, -1);
+			Particles[2] = FVec3(1, 1, -1);
+			Particles[3] = FVec3(0, 0, 0.5);
+			Particles[4] = (Particles[3]- Particles[1]) * 0.5 + Particles[1]+ FVec3(0, 0, 0.1);
 			TArray<TVector<int32, 3>> Indices;
 			Chaos::FConvexBuilder::BuildConvexHull(Particles, Indices);
 			EXPECT_EQ(Indices.Num(), 6);
 		}
 		{
-			TArray<TVec3<T>> Particles;
+			TArray<FVec3> Particles;
 			Particles.SetNum(6);
-			Particles[0] = Chaos::TVec3<float>(-1, -1, -1);
-			Particles[1] = Chaos::TVec3<float>(1, -1, -1);
-			Particles[2] = Chaos::TVec3<float>(1, 1, -1);
-			Particles[3] = Chaos::TVec3<float>(0, 0, 0.5);
-			Particles[4] = (Particles[3]- Particles[1]) * 0.5 + Particles[1] + Chaos::TVec3<float>(0, 0, 0.1);
-			Particles[5] = Particles[4]+ Chaos::TVec3<float>(-0.1, 0, 0);
+			Particles[0] = FVec3(-1, -1, -1);
+			Particles[1] = FVec3(1, -1, -1);
+			Particles[2] = FVec3(1, 1, -1);
+			Particles[3] = FVec3(0, 0, 0.5);
+			Particles[4] = (Particles[3]- Particles[1]) * 0.5 + Particles[1] + FVec3(0, 0, 0.1);
+			Particles[5] = Particles[4]+ FVec3(-0.1, 0, 0);
 			TArray<TVector<int32, 3>> Indices;
 			Chaos::FConvexBuilder::BuildConvexHull(Particles, Indices);
 			EXPECT_EQ(Indices.Num(), 8);
@@ -1338,17 +1338,17 @@ namespace ChaosTest {
 			// a large enough epsilon for building horizons in hull generation
 			// (tested to fail with 1e-1) we will generate a non-convex hull
 			// Using a scaled epsilon resolves this case
-			TArray<TVec3<T>> Particles;
+			TArray<FVec3> Particles;
 			Particles.SetNum(9);
-			Particles[0] = Chaos::TVec3<float>(-1, -1, -1);
-			Particles[1] = Chaos::TVec3<float>(-1, -1, 1);
-			Particles[2] = Chaos::TVec3<float>(-1, 1, -1);
-			Particles[3] = Chaos::TVec3<float>(-1, 1, 1);
-			Particles[4] = Chaos::TVec3<float>(1, -1, -1);
-			Particles[5] = Chaos::TVec3<float>(1, -1, 1);
-			Particles[6] = Chaos::TVec3<float>(1, 1, -1);
-			Particles[7] = Chaos::TVec3<float>(1, 1, 1);
-			Particles[8] = Chaos::TVec3<float>(0.966962576, -0.0577232838, 0.959515572);
+			Particles[0] = FVec3(-1, -1, -1);
+			Particles[1] = FVec3(-1, -1, 1);
+			Particles[2] = FVec3(-1, 1, -1);
+			Particles[3] = FVec3(-1, 1, 1);
+			Particles[4] = FVec3(1, -1, -1);
+			Particles[5] = FVec3(1, -1, 1);
+			Particles[6] = FVec3(1, 1, -1);
+			Particles[7] = FVec3(1, 1, 1);
+			Particles[8] = FVec3(0.966962576, -0.0577232838, 0.959515572);
 			
 			TArray<TVec3<int32>> Indices;
 			Chaos::FConvexBuilder::Params BuildParams;
@@ -1362,9 +1362,9 @@ namespace ChaosTest {
 			{
 				for (int32 i = 0; i < 3; ++i)
 				{
-					TVector<T, 3> V = Particles[Tri[i]];
-					TVector<T, 3> VAbs = V.GetAbs();
-					T Max = VAbs.GetMax();
+					FVec3 V = Particles[Tri[i]];
+					FVec3 VAbs = V.GetAbs();
+					FReal Max = VAbs.GetMax();
 					EXPECT_GE(Max, 1 - 1e-2);
 				}
 			}
@@ -1374,21 +1374,21 @@ namespace ChaosTest {
 			// only the original box - ignoring all interior and coplanar points.
 			// Note: If hull generation is changed to support non-triangular faces the conditions here
 			// will need to change as a correct hull in that method will produce only 6 faces not 12
-			TArray<TVec3<T>> Particles;
+			TArray<FVec3> Particles;
 			int32 NumParticles = 3600;
 			Particles.SetNum(NumParticles);
-			Particles[0] = Chaos::TVec3<float>(-1, -1, -1);
-			Particles[1] = Chaos::TVec3<float>(-1, -1, 1);
-			Particles[2] = Chaos::TVec3<float>(-1, 1, -1);
-			Particles[3] = Chaos::TVec3<float>(-1, 1, 1);
-			Particles[4] = Chaos::TVec3<float>(1, -1, -1);
-			Particles[5] = Chaos::TVec3<float>(1, -1, 1);
-			Particles[6] = Chaos::TVec3<float>(1, 1, -1);
-			Particles[7] = Chaos::TVec3<float>(1, 1, 1);
+			Particles[0] = FVec3(-1, -1, -1);
+			Particles[1] = FVec3(-1, -1, 1);
+			Particles[2] = FVec3(-1, 1, -1);
+			Particles[3] = FVec3(-1, 1, 1);
+			Particles[4] = FVec3(1, -1, -1);
+			Particles[5] = FVec3(1, -1, 1);
+			Particles[6] = FVec3(1, 1, -1);
+			Particles[7] = FVec3(1, 1, 1);
 			FRandomStream Stream(42);
 			for(int i = 8; i < NumParticles; ++i)
 			{
-				Particles[i]= Chaos::TVec3<float>(Stream.FRandRange(-1.f, 1.f), Stream.FRandRange(-1.f, 1.f), Stream.FRandRange(-1.f, 1.f));
+				Particles[i]= FVec3(Stream.FRandRange(-1.f, 1.f), Stream.FRandRange(-1.f, 1.f), Stream.FRandRange(-1.f, 1.f));
 			}
 			TArray<TVec3<int32>> Indices;
 
@@ -1402,9 +1402,9 @@ namespace ChaosTest {
 			{
 				for(int i = 0; i < 3; ++i)
 				{
-					TVector<T, 3> V = Particles[Tri[i]];
-					TVector<T, 3> VAbs = V.GetAbs();
-					T Max = VAbs.GetMax();
+					FVec3 V = Particles[Tri[i]];
+					FVec3 VAbs = V.GetAbs();
+					FReal Max = VAbs.GetMax();
 					EXPECT_GE(Max, 1 - 1e-2);
 				}
 			}
@@ -1413,26 +1413,26 @@ namespace ChaosTest {
 
 	void Simplify()
 	{
-		TArray<TVec3<T>> Particles;
+		TArray<FVec3> Particles;
 		Particles.SetNum(18);
-		Particles[0] = Chaos::TVec3<float>(0, 0, 12.0f);
-		Particles[1] = Chaos::TVec3<float>(-0.707f, -0.707f, 10.0f);
-		Particles[2] = Chaos::TVec3<float>(0, -1, 10.0f);
-		Particles[3] = Chaos::TVec3<float>(0.707f, -0.707f, 10.0f);
-		Particles[4] = Chaos::TVec3<float>(1, 0, 10.0f);
-		Particles[5] = Chaos::TVec3<float>(0.707f, 0.707f, 10.0f);
-		Particles[6] = Chaos::TVec3<float>(0.0f, 1.0f, 10.0f);
-		Particles[7] = Chaos::TVec3<float>(-0.707f, 0.707f, 10.0f);
-		Particles[8] = Chaos::TVec3<float>(-1.0f, 0.0f, 10.0f);
-		Particles[9] = Chaos::TVec3<float>(-0.707f, -0.707f, 0.0f);
-		Particles[10] = Chaos::TVec3<float>(0, -1, 0.0f);
-		Particles[11] = Chaos::TVec3<float>(0.707f, -0.707f, 0.0f);
-		Particles[12] = Chaos::TVec3<float>(1, 0, 0.0f);
-		Particles[13] = Chaos::TVec3<float>(0.707f, 0.707f, 0.0f);
-		Particles[14] = Chaos::TVec3<float>(0.0f, 1.0f, 0.0f);
-		Particles[15] = Chaos::TVec3<float>(-0.707f, 0.707f, 0.0f);
-		Particles[16] = Chaos::TVec3<float>(-1.0f, 0.0f, 0.0f);
-		Particles[17] = Chaos::TVec3<float>(0, 0, -2.0f);
+		Particles[0] = FVec3(0, 0, 12.0f);
+		Particles[1] = FVec3(-0.707f, -0.707f, 10.0f);
+		Particles[2] = FVec3(0, -1, 10.0f);
+		Particles[3] = FVec3(0.707f, -0.707f, 10.0f);
+		Particles[4] = FVec3(1, 0, 10.0f);
+		Particles[5] = FVec3(0.707f, 0.707f, 10.0f);
+		Particles[6] = FVec3(0.0f, 1.0f, 10.0f);
+		Particles[7] = FVec3(-0.707f, 0.707f, 10.0f);
+		Particles[8] = FVec3(-1.0f, 0.0f, 10.0f);
+		Particles[9] = FVec3(-0.707f, -0.707f, 0.0f);
+		Particles[10] = FVec3(0, -1, 0.0f);
+		Particles[11] = FVec3(0.707f, -0.707f, 0.0f);
+		Particles[12] = FVec3(1, 0, 0.0f);
+		Particles[13] = FVec3(0.707f, 0.707f, 0.0f);
+		Particles[14] = FVec3(0.0f, 1.0f, 0.0f);
+		Particles[15] = FVec3(-0.707f, 0.707f, 0.0f);
+		Particles[16] = FVec3(-1.0f, 0.0f, 0.0f);
+		Particles[17] = FVec3(0, 0, -2.0f);
 					
 		FConvex Convex(Particles, 0.0f);
 
@@ -1441,8 +1441,8 @@ namespace ChaosTest {
 		int32 OriginalNumberFaces = Convex.GetFaces().Num();
 		FAABB3 OriginalBoundingBox = Convex.BoundingBox();
 
-		const TArray<TVec3<T>>& CulledParticles = Convex.GetVertices();
-		const TArray<TPlaneConcrete<T, 3>> Planes = Convex.GetFaces();
+		const TArray<FVec3>& CulledParticles = Convex.GetVertices();
+		const TArray<TPlaneConcrete<FReal, 3>> Planes = Convex.GetFaces();
 
 		// set target number of particles in simplified convex
 		FConvexBuilder::PerformGeometryReduction = 1;
diff --git a/Engine/Source/Programs/HeadlessChaos/Private/HeadlessChaosTestMostOpposing.cpp b/Engine/Source/Programs/HeadlessChaos/Private/HeadlessChaosTestMostOpposing.cpp
--- a/Engine/Source/Programs/HeadlessChaos/Private/HeadlessChaosTestMostOpposing.cpp
+++ b/Engine/Source/Programs/HeadlessChaos/Private/HeadlessChaosTestMostOpposing.cpp
@@ -74,15 +74,15 @@ namespace ChaosTest
 		FVec3 Normal;
 		int32 FaceIndex;
 
-		TArray<TVec3<T>> Particles;
+		TArray<FVec3> Particles;
 		Particles.SetNum(6);
-		Particles[0] = TVec3<T>(1, 1, 1);
-		Particles[1] = TVec3<T>(5, 1, 1);
-		Particles[2] = TVec3<T>(1, 5, 1);
+		Particles[0] = FVec3(1, 1, 1);
+		Particles[1] = FVec3(5, 1, 1);
+		Particles[2] = FVec3(1, 5, 1);
 
-		Particles[3] = TVec3<T>(1, 1, 1);
-		Particles[4] = TVec3<T>(1, 5, 1);
-		Particles[5] = TVec3<T>(1, 1, -5);
+		Particles[3] = FVec3(1, 1, 1);
+		Particles[4] = FVec3(1, 5, 1);
+		Particles[5] = FVec3(1, 1, -5);
 
 		FConvex Convex(MoveTemp(Particles), 0.0f);
 
@@ -128,15 +128,15 @@ namespace ChaosTest
 		FVec3 Normal;
 		int32 FaceIndex;
 
-		TArray<TVec3<T>> Particles;
+		TArray<FVec3> Particles;
 		Particles.SetNum(6);
-		Particles[0] = TVec3<T>(0, -1, 1);
-		Particles[1] = TVec3<T>(1, -1, -1);
-		Particles[2] = TVec3<T>(0, 1, 1);
+		Particles[0] = FVec3(0, -1, 1);
+		Particles[1] = FVec3(1, -1, -1);
+		Particles[2] = FVec3(0, 1, 1);
 
-		Particles[3] = TVec3<T>(0, -1, 1);
-		Particles[4] = TVec3<T>(0, 1, 1);
-		Particles[5] = TVec3<T>(-1, -1, -1);
+		Particles[3] = FVec3(0, -1, 1);
+		Particles[4] = FVec3(0, 1, 1);
+		Particles[5] = FVec3(-1, -1, -1);
 
 		TUniquePtr<FImplicitObject> Convex = MakeUnique<FConvex>(MoveTemp(Particles), 0.0f);
 
diff --git a/Engine/Source/Programs/HeadlessChaos/Private/HeadlessChaosTestUtility.cpp b/Engine/Source/Programs/HeadlessChaos/Private/HeadlessChaosTestUtility.cpp
--- a/Engine/Source/Programs/HeadlessChaos/Private/HeadlessChaosTestUtility.cpp
+++ b/Engine/Source/Programs/HeadlessChaos/Private/HeadlessChaosTestUtility.cpp
@@ -579,8 +579,8 @@ namespace ChaosTest {
 		InParticles.Q() = InParticles.R();
 
 		// TODO: Change this error prone API to set bounds more automatically. This is easy to forget
-		InParticles.SetLocalBounds(TAABB<T, 3>(Cube[0], Cube[7]));
-		InParticles.SetWorldSpaceInflatedBounds(TAABB<T, 3>(Cube[0], Cube[7]));
+		InParticles.SetLocalBounds(TAABB<FReal, 3>(Cube[0], Cube[7]));
+		InParticles.SetWorldSpaceInflatedBounds(TAABB<FReal, 3>(Cube[0], Cube[7]));
 		InParticles.SetHasBounds(true);
 
 		InParticles.M() = 1.0;
