CIS failure fix (HeadlessChaos unit test)

diff --git a/Engine/Source/Programs/HeadlessChaos/Private/HeadlessChaosTestSolverProxies.cpp b/Engine/Source/Programs/HeadlessChaos/Private/HeadlessChaosTestSolverProxies.cpp
--- a/Engine/Source/Programs/HeadlessChaos/Private/HeadlessChaosTestSolverProxies.cpp
+++ b/Engine/Source/Programs/HeadlessChaos/Private/HeadlessChaosTestSolverProxies.cpp
@@ -80,6 +80,7 @@ namespace ChaosTest {
 		Particle.SetGeometry(Sphere);
 		Particle.SetX(FVec3(0, 0, 220));
 		Particle.SetV(FVec3(0, 0, -10));
+		Particle.SetCCDEnabled(true);
 		Solver->RegisterObject(Proxy);
 		Solver->AddDirtyProxy(Proxy);
 
