Headless Chaos -- Tests updated to reflect change in FConvex structure: TArray<FVec3> replacing TParticles for vertex data.

diff --git a/Engine/Source/Programs/HeadlessChaos/Private/HeadlessChaosTestConvex.cpp b/Engine/Source/Programs/HeadlessChaos/Private/HeadlessChaosTestConvex.cpp
--- a/Engine/Source/Programs/HeadlessChaos/Private/HeadlessChaosTestConvex.cpp
+++ b/Engine/Source/Programs/HeadlessChaos/Private/HeadlessChaosTestConvex.cpp
@@ -234,8 +234,8 @@ namespace ChaosTest
 		const FReal Radius = 1000.0f;
 
 		const int32 NumVertices = TestGeometry2::RawVertexArray.Num() / 3;
-		TParticles<FReal, 3> Particles;
-		Particles.AddParticles(NumVertices);
+		TArray<FVec3> Particles;
+		Particles.SetNum(NumVertices);
 		for (int32 ParticleIndex = 0; ParticleIndex < NumVertices; ++ParticleIndex)
 		{
 			const FVec3 VertexPos = FVec3(
@@ -243,7 +243,7 @@ namespace ChaosTest
 				TestGeometry2::RawVertexArray[3 * ParticleIndex + 1],
 				TestGeometry2::RawVertexArray[3 * ParticleIndex + 2]
 			);
-			Particles.X(ParticleIndex) = VertexPos;
+			Particles[ParticleIndex] = VertexPos;
 		}
 
 		FConvex Convex(Particles, 0.0f);
@@ -262,14 +262,14 @@ namespace ChaosTest
 		const int32 NumVertices = 1000;
 
 		// Make a convex with points on a sphere.
-		TParticles<FReal, 3> Particles;
-		Particles.AddParticles(NumVertices);
+		TArray<FVec3> Particles;
+		Particles.SetNum(NumVertices);
 		for (int32 ParticleIndex = 0; ParticleIndex < NumVertices; ++ParticleIndex)
 		{
 			const FReal Theta = FMath::RandRange(-PI, PI);
 			const FReal Phi = FMath::RandRange(-0.5f * PI, 0.5f * PI);
 			const FVec3 VertexPos = Radius * FVec3(FMath::Cos(Theta), FMath::Sin(Theta), FMath::Sin(Phi));
-			Particles.X(ParticleIndex) = VertexPos;
+			Particles[ParticleIndex] = VertexPos;
 		}
 		FConvex Convex(Particles, 0.0f);
 
diff --git a/Engine/Source/Programs/HeadlessChaos/Private/HeadlessChaosTestUtility.cpp b/Engine/Source/Programs/HeadlessChaos/Private/HeadlessChaosTestUtility.cpp
--- a/Engine/Source/Programs/HeadlessChaos/Private/HeadlessChaosTestUtility.cpp
+++ b/Engine/Source/Programs/HeadlessChaos/Private/HeadlessChaosTestUtility.cpp
@@ -718,9 +718,7 @@ namespace ChaosTest {
 			FVec3(HalfSize.X, -HalfSize.Y,  HalfSize.Z),
 		};
 
-		TParticles<FReal, 3> BoxParticles(MoveTemp(BoxVerts));
-
-		return FImplicitConvex3(BoxParticles, Margin);
+		return FImplicitConvex3(BoxVerts, Margin);
 	}
 
 
