ChaosVehicles: optional/alternate wheel and suspension simulation before switching over permanently, making it the default

diff --git a/Engine/Plugins/Experimental/ChaosVehiclesPlugin/Source/ChaosVehicles/Private/ChaosVehicleMovementComponent.cpp b/Engine/Plugins/Experimental/ChaosVehiclesPlugin/Source/ChaosVehicles/Private/ChaosVehicleMovementComponent.cpp
--- a/Engine/Plugins/Experimental/ChaosVehiclesPlugin/Source/ChaosVehicles/Private/ChaosVehicleMovementComponent.cpp
+++ b/Engine/Plugins/Experimental/ChaosVehiclesPlugin/Source/ChaosVehicles/Private/ChaosVehicleMovementComponent.cpp
@@ -905,7 +905,7 @@ void UChaosVehicleMovementComponent::SetChangeDownInput(bool bNewGearDown)
 
 void UChaosVehicleMovementComponent::SetTargetGear(int32 GearNum, bool bImmediate)
 {
-	if (GearNum != PVehicleOutput->TargetGear)
+	if (PVehicleOutput && GearNum != PVehicleOutput->TargetGear)
 	{
 		FBodyInstance* TargetInstance = UpdatedPrimitive->GetBodyInstance();
 
@@ -931,12 +931,12 @@ void UChaosVehicleMovementComponent::SetUseAutomaticGears(bool bUseAuto)
 
 int32 UChaosVehicleMovementComponent::GetCurrentGear() const
 {
-	return PVehicleOutput->CurrentGear;
+	return (PVehicleOutput)?PVehicleOutput->CurrentGear:0;
 }
 
 int32 UChaosVehicleMovementComponent::GetTargetGear() const
 {
-	return PVehicleOutput->TargetGear; //  or TargetGearInput
+	return (PVehicleOutput)?PVehicleOutput->TargetGear:0;
 }
 
 bool UChaosVehicleMovementComponent::GetUseAutoGears() const
@@ -1731,7 +1731,7 @@ void UChaosVehicleMovementComponent::ParallelUpdate(float DeltaSeconds)
 {
 	if (const FChaosVehicleAsyncOutput* CurrentOutput = static_cast<FChaosVehicleAsyncOutput*>(CurAsyncOutput))
 	{
-		if (CurrentOutput->bValid)
+		if (CurrentOutput->bValid && PVehicleOutput)
 		{
 			// TODO: It would be nicer to go through CurAsyncOutput rather
 			// than copying into the vehicle, think about non-async path
diff --git a/Engine/Plugins/Experimental/ChaosVehiclesPlugin/Source/ChaosVehicles/Private/ChaosVehicleWheel.cpp b/Engine/Plugins/Experimental/ChaosVehiclesPlugin/Source/ChaosVehicles/Private/ChaosVehicleWheel.cpp
--- a/Engine/Plugins/Experimental/ChaosVehiclesPlugin/Source/ChaosVehicles/Private/ChaosVehicleWheel.cpp
+++ b/Engine/Plugins/Experimental/ChaosVehiclesPlugin/Source/ChaosVehicles/Private/ChaosVehicleWheel.cpp
@@ -27,6 +27,9 @@ UChaosVehicleWheel::UChaosVehicleWheel(const FObjectInitializer& ObjectInitializ
 	//WheelMass = 20.0f;
 	LongitudinalFrictionForceMultiplier = 1.0f;
 	LateralFrictionForceMultiplier = 2.0f;
+	FrictionForceMultiplier = 2.0f;
+	CorneringStiffness = 1000.0f;
+	bNewWheelSimulation = false;
 	SideSlipModifier = 1.0f;
 	SlipThreshold = 20.0f;
 	SkidThreshold = 20.0f;
diff --git a/Engine/Plugins/Experimental/ChaosVehiclesPlugin/Source/ChaosVehicles/Private/ChaosWheeledVehicleMovementComponent.cpp b/Engine/Plugins/Experimental/ChaosVehiclesPlugin/Source/ChaosVehicles/Private/ChaosWheeledVehicleMovementComponent.cpp
--- a/Engine/Plugins/Experimental/ChaosVehiclesPlugin/Source/ChaosVehicles/Private/ChaosWheeledVehicleMovementComponent.cpp
+++ b/Engine/Plugins/Experimental/ChaosVehiclesPlugin/Source/ChaosVehicles/Private/ChaosWheeledVehicleMovementComponent.cpp
@@ -51,7 +51,7 @@ FAutoConsoleVariableRef CVarChaosVehiclesSteeringOverride(TEXT("p.Vehicle.Steeri
 
 FAutoConsoleVariableRef CVarChaosVehiclesResetMeasurements(TEXT("p.Vehicle.ResetMeasurements"), GWheeledVehicleDebugParams.ResetPerformanceMeasurements, TEXT("Reset Vehicle Performance Measurements."));
 
-FAutoConsoleVariableRef CVarChaosVehiclesDisableSuspensionConstraints(TEXT("p.Vehicle.DisableSuspensionConstraint"), GWheeledVehicleDebugParams.DisableSuspensionConstraint, TEXT("Enable/Disable Suspension Constraints."));
+//FAutoConsoleVariableRef CVarChaosVehiclesDisableSuspensionConstraints(TEXT("p.Vehicle.DisableSuspensionConstraint"), GWheeledVehicleDebugParams.DisableSuspensionConstraint, TEXT("Enable/Disable Suspension Constraints."));
 
 FAutoConsoleCommand CVarCommandVehiclesNextDebugPage(
 	TEXT("p.Vehicle.NextDebugPage"),
@@ -502,7 +502,7 @@ void UChaosWheeledVehicleSimulation::ApplySuspensionForces(float DeltaTime)
 		auto& PSuspension = PVehicle->Suspension[WheelIdx];
 		float SuspensionMovePosition = -PSuspension.Setup().MaxLength;
 
-		if (!GWheeledVehicleDebugParams.DisableSuspensionConstraint)
+		if (PWheel.Setup().NewSimulationPath)
 		{
 #if WITH_CHAOS
 			if (ConstraintHandles.Num() > 0)
@@ -540,7 +540,7 @@ void UChaosWheeledVehicleSimulation::ApplySuspensionForces(float DeltaTime)
 			FVector SusApplicationPoint = WheelState.WheelWorldLocation[WheelIdx] + PVehicle->Suspension[WheelIdx].Setup().SuspensionForceOffset;
 
 			check(PWheel.InContact());
-			if (GWheeledVehicleDebugParams.DisableSuspensionConstraint)
+			if (!PWheel.Setup().NewSimulationPath)
 			{
 				AddForceAtPosition(SuspensionForceVector, SusApplicationPoint);
 			}
@@ -995,7 +995,7 @@ void UChaosWheeledVehicleMovementComponent::FixupSkeletalMesh()
 
 					}
 
-					if (!GWheeledVehicleDebugParams.DisableSuspensionConstraint)
+					if (Wheels[WheelIdx]->bNewWheelSimulation)
 					{
 						FBodyInstance* TargetInstance = UpdatedPrimitive->GetBodyInstance();
 
@@ -1399,7 +1399,7 @@ float UChaosWheeledVehicleMovementComponent::GetEngineRotationSpeed() const
 {
 	float EngineRPM = 0.f;
 
-	if (bMechanicalSimEnabled)
+	if (bMechanicalSimEnabled && PVehicleOutput)
 	{
 		EngineRPM = PVehicleOutput->EngineRPM;
 	}
diff --git a/Engine/Plugins/Experimental/ChaosVehiclesPlugin/Source/ChaosVehicles/Public/ChaosVehicleWheel.h b/Engine/Plugins/Experimental/ChaosVehiclesPlugin/Source/ChaosVehicles/Public/ChaosVehicleWheel.h
--- a/Engine/Plugins/Experimental/ChaosVehiclesPlugin/Source/ChaosVehicles/Public/ChaosVehicleWheel.h
+++ b/Engine/Plugins/Experimental/ChaosVehiclesPlugin/Source/ChaosVehicles/Public/ChaosVehicleWheel.h
@@ -97,14 +97,25 @@ class UChaosWheeledVehicleMovementComponent;
 		//UPROPERTY(EditAnywhere, Category = Wheel, meta = (ClampMin = "0.01", UIMin = "0.01"))
 		//float WheelMass;
 
-		/** Longitudinal Friction Force Multiplier */
 		UPROPERTY(EditAnywhere, Category = Wheel)
+		bool bNewWheelSimulation;
+
+		/** Longitudinal Friction Force Multiplier */
+		UPROPERTY(EditAnywhere, Category = Wheel, meta = (EditCondition = "!bNewWheelSimulation"))
 		float LongitudinalFrictionForceMultiplier;
 
 		/** Lateral Friction Force Multiplier */
-		UPROPERTY(EditAnywhere, Category = Wheel)
+		UPROPERTY(EditAnywhere, Category = Wheel, meta = (EditCondition = "!bNewWheelSimulation"))
 		float LateralFrictionForceMultiplier;
 
+		/** Tyre Cornering Ability */
+		UPROPERTY(EditAnywhere, Category = Wheel, meta = (EditCondition = "bNewWheelSimulation"))
+		float CorneringStiffness;
+
+		/** Friction Force Multiplier */
+		UPROPERTY(EditAnywhere, Category = Wheel, meta = (EditCondition = "bNewWheelSimulation"))
+		float FrictionForceMultiplier;
+
 		/** Wheel Lateral Skid Grip Loss, lower number less grip on skid */
 		UPROPERTY(EditAnywhere, Category = Wheel, meta = (ClampMin = "0.0", UIMin = "0.0", ClampMax = "1.0", UIMax = "1.0"))
 		float SideSlipModifier;
@@ -363,6 +374,9 @@ class UChaosWheeledVehicleMovementComponent;
 			PWheelConfig.AxleType = static_cast<Chaos::FSimpleWheelConfig::EAxleType>(this->AxleType);
 			PWheelConfig.LongitudinalFrictionMultiplier = this->LongitudinalFrictionForceMultiplier;
 			PWheelConfig.LateralFrictionMultiplier = this->LateralFrictionForceMultiplier;
+			PWheelConfig.FrictionMultiplier = this->FrictionForceMultiplier;
+			PWheelConfig.CorneringStiffness = this->CorneringStiffness;
+			PWheelConfig.NewSimulationPath = this->bNewWheelSimulation;
 			PWheelConfig.SideSlipModifier = this->SideSlipModifier;
 			PWheelConfig.SlipThreshold = this->SlipThreshold;
 			PWheelConfig.SkidThreshold = this->SkidThreshold;
diff --git a/Engine/Plugins/Experimental/ChaosVehiclesPlugin/Source/ChaosVehicles/Public/ChaosWheeledVehicleMovementComponent.h b/Engine/Plugins/Experimental/ChaosVehiclesPlugin/Source/ChaosVehicles/Public/ChaosWheeledVehicleMovementComponent.h
--- a/Engine/Plugins/Experimental/ChaosVehiclesPlugin/Source/ChaosVehicles/Public/ChaosWheeledVehicleMovementComponent.h
+++ b/Engine/Plugins/Experimental/ChaosVehiclesPlugin/Source/ChaosVehicles/Public/ChaosWheeledVehicleMovementComponent.h
@@ -34,7 +34,7 @@ struct FWheeledVehicleDebugParams
 
 	bool ResetPerformanceMeasurements = false;
 
-	bool DisableSuspensionConstraint = false;
+	//bool DisableSuspensionConstraint = false;
 };
 
 /**
diff --git a/Engine/Source/Programs/HeadlessChaos/Private/HeadlessChaosTestVehicles.cpp b/Engine/Source/Programs/HeadlessChaos/Private/HeadlessChaosTestVehicles.cpp
--- a/Engine/Source/Programs/HeadlessChaos/Private/HeadlessChaosTestVehicles.cpp
+++ b/Engine/Source/Programs/HeadlessChaos/Private/HeadlessChaosTestVehicles.cpp
@@ -30,6 +30,49 @@ namespace ChaosTest
 {
 	using namespace Chaos;
 
+	TYPED_TEST(AllTraits, VehicleTest_VehicleUtilityGraph)
+	{
+		Chaos::FGraph Graph;
+
+		Graph.Add(FVector2D(10.f, 0.f));
+		Graph.Add(FVector2D(20.f, 30.f));
+		Graph.Add(FVector2D(30.f, 40.f));
+		Graph.Add(FVector2D(40.f, 60.f));
+		Graph.Add(FVector2D(50.f, 90.f));
+
+		float Tolerance = 0.001f;
+		EXPECT_NEAR(Graph.EvaluateY(-5.f), 0.f, Tolerance);
+		EXPECT_NEAR(Graph.EvaluateY(5.f), 0.f, Tolerance);
+		EXPECT_NEAR(Graph.EvaluateY(10.f), 0.f, Tolerance);
+		EXPECT_NEAR(Graph.EvaluateY(40.f), 60.f, Tolerance);
+		EXPECT_NEAR(Graph.EvaluateY(25.f), 35.f, Tolerance);
+		EXPECT_NEAR(Graph.EvaluateY(45.f), 75.f, Tolerance);
+		EXPECT_NEAR(Graph.EvaluateY(50.f), 90.f, Tolerance);
+		EXPECT_NEAR(Graph.EvaluateY(60.f), 90.f, Tolerance);
+	}
+
+	TYPED_TEST(AllTraits, VehicleTest_CalculateSlipAngle)
+	{
+		float SlipAngleFwds = FVehicleUtility::CalculateSlipAngle(0, 5);
+		float SlipAngleReverse = FVehicleUtility::CalculateSlipAngle(0, -5);
+
+		EXPECT_NEAR(SlipAngleFwds, 0.0f, 0.001f);
+		EXPECT_NEAR(SlipAngleFwds, SlipAngleReverse, 0.001f);
+
+		float SlipAngleFwdsSide = FVehicleUtility::CalculateSlipAngle(1, 5);
+		float SlipAngleReverseSide = FVehicleUtility::CalculateSlipAngle(1, -5);
+		float SlipAngleFwdsSide2 = FVehicleUtility::CalculateSlipAngle(-1, 5);
+		float SlipAngleReverseSide2 = FVehicleUtility::CalculateSlipAngle(-1, -5);
+		EXPECT_NEAR(SlipAngleFwdsSide, SlipAngleReverseSide, 0.001f);
+		EXPECT_NEAR(SlipAngleFwdsSide, SlipAngleFwdsSide2, 0.001f);
+		EXPECT_NEAR(SlipAngleFwdsSide, SlipAngleReverseSide2, 0.001f);
+
+		float SlipAngleSide1 = FVehicleUtility::CalculateSlipAngle(2, 0);
+		float SlipAngleSide2 = FVehicleUtility::CalculateSlipAngle(-2, 0);
+		EXPECT_NEAR(SlipAngleSide1, HALF_PI, 0.001f);
+		EXPECT_NEAR(SlipAngleSide1, SlipAngleSide2, 0.001f);
+
+	}
 
 	TYPED_TEST(AllTraits, VehicleTest_SteeringUtilityTurnRadius)
 	{
@@ -751,6 +794,9 @@ namespace ChaosTest
 		Setup.BrakeEnabled = true;
 		Setup.EngineEnabled = true;
 		Setup.WheelRadius = 30.0f;
+		Setup.LongitudinalFrictionMultiplier = 1.0f;
+		Setup.LateralFrictionMultiplier = 1.0f;
+		Setup.SideSlipModifier = 0.7f;
 
 		FSimpleWheelSim Wheel(&Setup);
 
@@ -799,9 +845,17 @@ namespace ChaosTest
 		Wheel.SetBrakeTorque(5000);
 		VehicleSpeedMPH = 30.f;
 		SimulateBraking(Wheel, VehicleSpeedMPH, DeltaTime, StoppingDistanceTooHeavyBreaking, SimulationTime);
-
 		EXPECT_GT(StoppingDistanceTooHeavyBreaking, StoppingDistanceA);
 
+		// Would have locked the wheels but ABS prevents skidding
+		Wheel.AccessSetup().ABSEnabled = true;
+		float StoppingDistanceTooHeavyBreakingABS = 0.f;
+		Wheel.SetBrakeTorque(5000);
+		VehicleSpeedMPH = 30.f;
+		SimulateBraking(Wheel, VehicleSpeedMPH, DeltaTime, StoppingDistanceTooHeavyBreakingABS, SimulationTime);
+		EXPECT_LT(StoppingDistanceTooHeavyBreakingABS, StoppingDistanceTooHeavyBreaking);
+		Wheel.AccessSetup().ABSEnabled = false;
+
 		// lower initial speed - stops more quickly
 		float StoppingDistanceLowerSpeed = 0.f;
 		Wheel.SetBrakeTorque(650);
@@ -832,8 +886,180 @@ namespace ChaosTest
 		Setup.TractionControlEnabled = false;
 		Setup.BrakeEnabled = true;
 		Setup.EngineEnabled = true;
+		Setup.WheelRadius = 30.0f;
+		Setup.LongitudinalFrictionMultiplier = 1.0f;
+		Setup.LateralFrictionMultiplier = 1.0f;
+		Setup.SideSlipModifier = 0.7f;
+
+		FSimpleWheelSim Wheel(&Setup);
+
+		// There could be one frame extra computation on the acceleration since the last frame of brake is not using the full 
+		// amount of torque, it's clearing the last remaining velocity without pushing the vehicle back in the opposite direction
+		// Hence a slightly larger tolerance for the result
+		float AccelerationResultsTolerance = 1.0f; // meters
+
+		// units meters
+		float Gravity = 9.8f;
+		float DeltaTime = 1.f / 30.f;
+
+		float StoppingDistanceA = 0.f;
+		float SimulationTimeBrake = 0.0f;
+		Wheel.SetSurfaceFriction(RealWorldConsts::DryRoadFriction());
+
+		// How far & what time does it take to stop from 30MPH to rest
+		Wheel.SetBrakeTorque(650);
+		SimulateBraking(Wheel, 30.0f, DeltaTime, StoppingDistanceA, SimulationTimeBrake);
+
+		// How far and what time does it take to accelerate from rest to 30MPH
+		float SimulationTimeAccel = 0.0f;
+		float DrivingDistanceA = 0.f;
+		Wheel.SetDriveTorque(650);
+		SimulateAccelerating(Wheel, Gravity, 30.0f, DeltaTime, DrivingDistanceA, SimulationTimeAccel);
+
+		// 0-30 MPH and 30-0 MPH should be the same if there's no slipping and accel torque was same as the brake torque run
+		EXPECT_LT(DrivingDistanceA - StoppingDistanceA, AccelerationResultsTolerance);
+		EXPECT_LT(SimulationTimeAccel - SimulationTimeBrake, AccelerationResultsTolerance);
+
+		// same range as braking from 30MPH
+		EXPECT_GT(DrivingDistanceA, 10.f);
+		EXPECT_LT(DrivingDistanceA, 20.f);
+
+		// Unreal units cm - Note for the same results the radius needs to remain at 0.3m and not also be scaled to 30(cm)
+		float SimulationTimeAccelCM = 0.0f;
+		float MtoCm = 100.0f;
+		float DrivingDistanceCM = 0.f;
+		Wheel.SetDriveTorque(650 * MtoCm);
+		SimulateAccelerating(Wheel, Gravity * MtoCm, 30.0f * MtoCm, DeltaTime, DrivingDistanceCM, SimulationTimeAccelCM);
+		EXPECT_GT(DrivingDistanceCM, 10.f * MtoCm);
+		EXPECT_LT(DrivingDistanceCM, 20.f * MtoCm);
+		EXPECT_LT(SimulationTimeAccel - SimulationTimeAccelCM, AccelerationResultsTolerance);
+
+		float SimulationTimeAccelSpin = 0.0f;
+		float DrivingDistanceWheelspin = 0.f;
+		Wheel.SetDriveTorque(5000); // definitely cause wheel spin
+		SimulateAccelerating(Wheel, Gravity, 30.0f, DeltaTime, DrivingDistanceWheelspin, SimulationTimeAccelSpin);
 
+		// Enable traction control should be better than both of the above
+		float SimulationTimeAccelTC = 0.0f;
+		float DrivingDistanceTC = 0.f;
+		Wheel.AccessSetup().TractionControlEnabled = true;
+		Wheel.SetDriveTorque(5000); // definitely cause wheel spin
+		SimulateAccelerating(Wheel, Gravity, 30.0f, DeltaTime, DrivingDistanceTC, SimulationTimeAccelTC);
+
+		// reaches target speed in a shorter distance
+		EXPECT_LT(DrivingDistanceTC, DrivingDistanceWheelspin);
+
+		// reaches speed quicker with TC on when wheel would be slipping from drive torque
+		EXPECT_LT(SimulationTimeAccelTC, SimulationTimeAccelSpin);
+	}
+
+	TYPED_TEST(AllTraits, VehicleTest_WheelBrakingLongitudinalSlipNew)
+	{
+		FSimpleWheelConfig Setup;
+		Setup.ABSEnabled = false;
+		Setup.TractionControlEnabled = false;
+		Setup.BrakeEnabled = true;
+		Setup.EngineEnabled = true;
 		Setup.WheelRadius = 30.0f;
+		Setup.NewSimulationPath = true;
+		Setup.FrictionMultiplier = 1.0f;
+		Setup.CorneringStiffness = 2.0f;
+		Setup.SideSlipModifier = 0.7f;
+
+		FSimpleWheelSim Wheel(&Setup);
+
+		// Google braking distance at 30mph says 14m (not interested in the thinking distance part)
+		// So using a range 10-20 to ensure we are in the correct ballpark.
+		// If specified more accurately in the test, then modifying the code would break the test all the time.
+
+		float StoppingDistanceTolerance = 0.5f; // meters
+		float DeltaTime = 1.f / 30.f;
+		float StoppingDistanceA = 0.f;
+		float SimulationTime = 0.0f;
+		Wheel.SetSurfaceFriction(RealWorldConsts::DryRoadFriction());
+
+		// reasonably ideal stopping distance - traveling forwards
+		Wheel.SetBrakeTorque(650);
+		float VehicleSpeedMPH = 30.f;
+		SimulateBraking(Wheel, VehicleSpeedMPH, DeltaTime, StoppingDistanceA, SimulationTime);
+		//UE_LOG(LogChaos, Warning, TEXT("Braking Distance %3.2fm"), StoppingDistanceB);
+		EXPECT_GT(StoppingDistanceA, 10.f);
+		EXPECT_LT(StoppingDistanceA, 20.f);
+
+		// traveling backwards stops just the same
+		float StoppingDistanceReverseDir = 0.f;
+		Wheel.SetBrakeTorque(650);
+		VehicleSpeedMPH = -30.f;
+		SimulateBraking(Wheel, VehicleSpeedMPH, DeltaTime, StoppingDistanceReverseDir, SimulationTime);
+		EXPECT_GT(StoppingDistanceReverseDir, -20.f);
+		EXPECT_LT(StoppingDistanceReverseDir, -10.f);
+		EXPECT_LT(StoppingDistanceA - FMath::Abs(StoppingDistanceReverseDir), StoppingDistanceTolerance);
+
+		// Similar results with different delta time
+		float StoppingDistanceDiffDT = 0.f;
+		VehicleSpeedMPH = 30.f;
+		SimulateBraking(Wheel, VehicleSpeedMPH, DeltaTime * 0.25f, StoppingDistanceDiffDT, SimulationTime);
+		EXPECT_LT(StoppingDistanceA - StoppingDistanceDiffDT, StoppingDistanceTolerance);
+
+		// barely touching the brake - going to take longer to stop
+		float StoppingDistanceLightBraking = 0.f;
+		Wheel.SetBrakeTorque(150);
+		VehicleSpeedMPH = 30.f;
+		SimulateBraking(Wheel, VehicleSpeedMPH, DeltaTime, StoppingDistanceLightBraking, SimulationTime);
+		EXPECT_GT(StoppingDistanceLightBraking, StoppingDistanceA);
+
+		// locking the wheels / too much brake torque -> dynamic friction rather than static friction -> going to take longer to stop
+		float StoppingDistanceTooHeavyBreaking = 0.f;
+		Wheel.SetBrakeTorque(5000);
+		VehicleSpeedMPH = 30.f;
+		SimulateBraking(Wheel, VehicleSpeedMPH, DeltaTime, StoppingDistanceTooHeavyBreaking, SimulationTime);
+		EXPECT_GT(StoppingDistanceTooHeavyBreaking, StoppingDistanceA);
+
+		// Would have locked the wheels but ABS prevents skidding
+		Wheel.AccessSetup().ABSEnabled = true;
+		float StoppingDistanceTooHeavyBreakingABS = 0.f;
+		Wheel.SetBrakeTorque(5000);
+		VehicleSpeedMPH = 30.f;
+		SimulateBraking(Wheel, VehicleSpeedMPH, DeltaTime, StoppingDistanceTooHeavyBreakingABS, SimulationTime);
+		EXPECT_LT(StoppingDistanceTooHeavyBreakingABS, StoppingDistanceTooHeavyBreaking);
+		Wheel.AccessSetup().ABSEnabled = false;
+
+		// lower initial speed - stops more quickly
+		float StoppingDistanceLowerSpeed = 0.f;
+		Wheel.SetBrakeTorque(650);
+		VehicleSpeedMPH = 20.f;
+		SimulateBraking(Wheel, VehicleSpeedMPH, DeltaTime, StoppingDistanceLowerSpeed, SimulationTime);
+		EXPECT_LT(StoppingDistanceLowerSpeed, StoppingDistanceA);
+
+		// higher initial speed - stops more slowly
+		float StoppingDistanceHigherSpeed = 0.f;
+		Wheel.SetBrakeTorque(650);
+		VehicleSpeedMPH = 60.f;
+		SimulateBraking(Wheel, VehicleSpeedMPH, DeltaTime, StoppingDistanceHigherSpeed, SimulationTime);
+		EXPECT_GT(StoppingDistanceHigherSpeed, StoppingDistanceA);
+
+		// slippy surface - stops more slowly
+		float StoppingDistanceLowFriction = 0.f;
+		Wheel.SetSurfaceFriction(0.3f);
+		Wheel.SetBrakeTorque(650);
+		VehicleSpeedMPH = 30.f;
+		SimulateBraking(Wheel, VehicleSpeedMPH, DeltaTime, StoppingDistanceLowFriction, SimulationTime);
+		EXPECT_GT(StoppingDistanceLowFriction, StoppingDistanceA);
+	}
+
+	TYPED_TEST(AllTraits, VehicleTest_WheelAcceleratingLongitudinalSlipNew)
+	{
+		FSimpleWheelConfig Setup;
+		Setup.ABSEnabled = false;
+		Setup.TractionControlEnabled = false;
+		Setup.BrakeEnabled = true;
+		Setup.EngineEnabled = true;
+		Setup.WheelRadius = 30.0f;
+		Setup.NewSimulationPath = true;
+		Setup.FrictionMultiplier = 1.0f;
+		Setup.CorneringStiffness = 2.0f;
+		Setup.SideSlipModifier = 0.7f;
+
 		FSimpleWheelSim Wheel(&Setup);
 
 		// There could be one frame extra computation on the acceleration since the last frame of brake is not using the full 
@@ -887,8 +1113,23 @@ namespace ChaosTest
 
 		// takes longer to reach the same speed
 		EXPECT_GT(SimulationTimeAccelSpin, SimulationTimeAccel);
+
+		// Enable traction control should be better than both of the above
+		float SimulationTimeAccelTC = 0.0f;
+		float DrivingDistanceTC = 0.f;
+		Wheel.AccessSetup().TractionControlEnabled = true;
+		Wheel.SetDriveTorque(5000); // definitely cause wheel spin
+		SimulateAccelerating(Wheel, Gravity, 30.0f, DeltaTime, DrivingDistanceTC, SimulationTimeAccelTC);
+
+		// reaches target speed in a shorter distance
+		EXPECT_LT(DrivingDistanceTC, DrivingDistanceWheelspin);
+
+		// reaches speed quicker with TC on when wheel would be slipping from drive torque
+		EXPECT_LT(SimulationTimeAccelTC, SimulationTimeAccelSpin);
+
 	}
 
+
 	TYPED_TEST(AllTraits, DISABLED_VehicleTest_WheelLateralSlip)
 	{
 		FSimpleWheelConfig Setup;
diff --git a/Engine/Source/Runtime/Experimental/ChaosVehicles/ChaosVehiclesCore/Private/WheelSystem.cpp b/Engine/Source/Runtime/Experimental/ChaosVehicles/ChaosVehiclesCore/Private/WheelSystem.cpp
--- a/Engine/Source/Runtime/Experimental/ChaosVehicles/ChaosVehiclesCore/Private/WheelSystem.cpp
+++ b/Engine/Source/Runtime/Experimental/ChaosVehicles/ChaosVehiclesCore/Private/WheelSystem.cpp
@@ -32,6 +32,12 @@ namespace Chaos
 
 	void FSimpleWheelSim::Simulate(float DeltaTime)
 	{
+		if (Setup().NewSimulationPath)
+		{
+			SimulateNew(DeltaTime);
+			return;
+		}
+
 		SlipAngle = FMath::Atan2(GroundVelocityVector.Y, GroundVelocityVector.X);
 
 		float FinalLongitudinalForce = 0.f;
@@ -46,7 +52,7 @@ namespace Chaos
 		bool Braking = BrakeTorque > FMath::Abs(DriveTorque);
 		float BrakeFactor = 1.0f;
 		float K = 0.4f;
-		
+
 		// are we actually touching the ground
 		if (ForceIntoSurface > SMALL_NUMBER)
 		{
@@ -91,7 +97,7 @@ namespace Chaos
 			SideSlipModifier = 1.0f;
 			bool Locked = false;
 			bool Spinning = false;
-			
+
 			// we can only obtain as much accel/decel force as the friction will allow
 			if (FMath::Abs(FinalLongitudinalForce) > LongitudinalAdhesiveLimit)
 			{
@@ -186,6 +192,145 @@ namespace Chaos
 		return;
 	}
 
+	void FSimpleWheelSim::SimulateNew(float DeltaTime)
+	{
+		float K = 0.4f;
+		float TractionControlAndABSScaling = 0.98f;	// how close to perfection is the system working
+
+		// X is longitudinal direction, Y is lateral
+		SlipAngle = FVehicleUtility::CalculateSlipAngle(GroundVelocityVector.Y, GroundVelocityVector.X);
+
+		// The physics system is mostly unit-less i.e. can work in meters or cm, however there are 
+		// a couple of places where the results are wrong if Cm is used. This is one of them, the simulated radius
+		// for torque must be real size to obtain the correct output values.
+		AppliedLinearDriveForce = DriveTorque / CmToM(Re);
+		AppliedLinearBrakeForce = BrakeTorque / CmToM(Re);
+
+		// Longitudinal multiplier now affecting both brake and steering equally
+		float AvailableGrip = ForceIntoSurface * SurfaceFriction * Setup().FrictionMultiplier;
+
+		float FinalLongitudinalForce = 0.f;
+		float FinalLateralForce = 0.f;
+
+		// currently just letting the brake override the throttle
+		bool Braking = BrakeTorque > FMath::Abs(DriveTorque);
+
+		// are we actually touching the ground
+		if (ForceIntoSurface > SMALL_NUMBER)
+		{
+			// ABS limiting brake force to match force from the grip avialable
+			if (Setup().ABSEnabled && Braking && FMath::Abs(AppliedLinearBrakeForce) > AvailableGrip)
+			{
+				if ((Braking && Setup().ABSEnabled) || (!Braking && Setup().TractionControlEnabled))
+				{
+					//	Spin = 0.0f;
+					AppliedLinearBrakeForce = AvailableGrip * TractionControlAndABSScaling;
+				}
+			}
+
+			// Traction control limiting drive force to match force from grip available
+			if (Setup().TractionControlEnabled && !Braking && FMath::Abs(AppliedLinearDriveForce) > AvailableGrip)
+			{
+				AppliedLinearDriveForce = AvailableGrip * TractionControlAndABSScaling;
+			}
+
+			if (Braking)
+			{
+				// whether the velocity is +ve or -ve when we brake we are slowing the vehicle down
+				// so force is opposing current direction of travel.
+				float ForceRequiredToBringToStop = MassPerWheel * K * (GroundVelocityVector.X) / DeltaTime;
+				FinalLongitudinalForce = AppliedLinearBrakeForce;
+
+				// check we are not applying more force than required so we end up overshooting 
+				// and accelerating in the opposite direction
+				if (FinalLongitudinalForce > FMath::Abs(ForceRequiredToBringToStop))
+				{
+					FinalLongitudinalForce = FMath::Abs(ForceRequiredToBringToStop);
+				}
+
+				// ensure the brake opposes current direction of travel
+				if (GroundVelocityVector.X > 0.0f)
+				{
+					FinalLongitudinalForce = -FinalLongitudinalForce;
+				}
+
+			}
+			else
+			{
+				FinalLongitudinalForce = AppliedLinearDriveForce;
+			}
+
+			float ForceRequiredToBringToStop = -(MassPerWheel * K * GroundVelocityVector.Y) / DeltaTime;
+
+			// use slip angle to generate a sideways force
+			float CorneringStiffness = Setup().CorneringStiffness * 10000.0f;
+
+			// Leveling off of the graph
+			float AngleLimit = FMath::DegreesToRadians(8.0f);
+			if (SlipAngle > AngleLimit)
+			{
+				SlipAngle = AngleLimit;
+			}
+			else if (SlipAngle < AngleLimit)
+			{
+				SlipAngle = AngleLimit;
+			}
+			FinalLateralForce = SlipAngle * CorneringStiffness;
+			if (FinalLateralForce > FMath::Abs(ForceRequiredToBringToStop))
+			{
+				FinalLateralForce = FMath::Abs(ForceRequiredToBringToStop);
+			}
+			if (GroundVelocityVector.Y > 0.0f)
+			{
+				FinalLateralForce = -FinalLateralForce;
+			}
+
+			// Friction circle
+			float LengthSquared = FinalLongitudinalForce * FinalLongitudinalForce + FinalLateralForce * FinalLateralForce;
+			if (LengthSquared > 0.05f)
+			{
+				float Length = FMath::Sqrt(LengthSquared);
+
+				float Clip = AvailableGrip / Length;
+				if (Clip < 1.0f)
+				{
+					FinalLongitudinalForce *= Clip;
+					FinalLateralForce *= Clip;
+
+					FinalLongitudinalForce *= Setup().SideSlipModifier;
+					FinalLateralForce *= Setup().SideSlipModifier;
+
+				}
+			}
+		}
+
+		float GroundOmega = GroundVelocityVector.X / Re;
+		Omega += (GroundOmega - Omega);
+
+		// Wheel angular position
+		AngularPosition += Omega * DeltaTime;
+
+		while (AngularPosition >= PI * 2.f)
+		{
+			AngularPosition -= PI * 2.f;
+		}
+		while (AngularPosition <= -PI * 2.f)
+		{
+			AngularPosition += PI * 2.f;
+		}
+
+		if (!bInContact)
+		{
+			ForceFromFriction = FVector::ZeroVector;
+		}
+		else
+		{
+			ForceFromFriction.X = FinalLongitudinalForce;
+			ForceFromFriction.Y = FinalLateralForce;
+		}
+
+		return;
+	}
 
 
 
@@ -198,5 +343,5 @@ namespace Chaos
 } // namespace Chaos
 
 #if VEHICLE_DEBUGGING_ENABLED
-	PRAGMA_ENABLE_OPTIMIZATION
+PRAGMA_ENABLE_OPTIMIZATION
 #endif
diff --git a/Engine/Source/Runtime/Experimental/ChaosVehicles/ChaosVehiclesCore/Public/VehicleUtility.h b/Engine/Source/Runtime/Experimental/ChaosVehicles/ChaosVehiclesCore/Public/VehicleUtility.h
--- a/Engine/Source/Runtime/Experimental/ChaosVehicles/ChaosVehiclesCore/Public/VehicleUtility.h
+++ b/Engine/Source/Runtime/Experimental/ChaosVehicles/ChaosVehiclesCore/Public/VehicleUtility.h
@@ -125,6 +125,17 @@ namespace Chaos
 
 		/** Calculate turn radius from three points. Note: this function is quite inaccurate for large radii. Return 0 if there is no answer, i.e. points lie on a line */
 		static float TurnRadiusFromThreePoints(const FVector& PtA, const FVector& PtB, const FVector& PtC);
+
+		static float CalculateSlipAngle(float Y, float X)
+		{
+			float Value = FMath::Abs(FMath::Atan2(Y, X));
+			if (Value > HALF_PI)
+			{
+				Value = PI - Value;
+			}
+
+			return Value;
+		}
 	};
 
 	FORCEINLINE float MToCmScaling()
diff --git a/Engine/Source/Runtime/Experimental/ChaosVehicles/ChaosVehiclesCore/Public/WheelSystem.h b/Engine/Source/Runtime/Experimental/ChaosVehicles/ChaosVehiclesCore/Public/WheelSystem.h
--- a/Engine/Source/Runtime/Experimental/ChaosVehicles/ChaosVehiclesCore/Public/WheelSystem.h
+++ b/Engine/Source/Runtime/Experimental/ChaosVehicles/ChaosVehiclesCore/Public/WheelSystem.h
@@ -87,9 +87,12 @@ struct CHAOSVEHICLESCORE_API FSimpleWheelConfig
 		, FrictionCombineMethod(EFrictionCombineMethod::Multiply)
 		, LateralFrictionMultiplier(1.0f)
 		, LongitudinalFrictionMultiplier(1.0f)
+		, FrictionMultiplier(2.0f)
+		, CorneringStiffness(1000.0f)
 		, SideSlipModifier(1.0f)
 		, SlipThreshold(20.0f)
 		, SkidThreshold(20.0f)
+		, NewSimulationPath(false)
 	{
 
 	}
@@ -123,11 +126,15 @@ struct CHAOSVEHICLESCORE_API FSimpleWheelConfig
 
 	float LateralFrictionMultiplier;
 	float LongitudinalFrictionMultiplier;
+	float FrictionMultiplier;
+	float CorneringStiffness;
 	float SideSlipModifier;
 
 	float SlipThreshold;
 	float SkidThreshold;
 
+	bool NewSimulationPath;
+
 	// #todo: simulated Damage
 	//EWheelDamageStatus DamageStatus;
 	//float BuckleAngle;
@@ -373,6 +380,7 @@ public:
 	 *	#todo: wheel slip angle isn't being used
 	 */
 	void Simulate(float DeltaTime);
+	void SimulateNew(float DeltaTime);
 
 
 	void SetMassPerWheel(float VehicleMassPerWheel)
