Merging fix from RES->Main: Attempt to fix compiler warning

diff --git a/Engine/Source/Programs/HeadlessChaos/Private/HeadlessChaosTestRewind.cpp b/Engine/Source/Programs/HeadlessChaos/Private/HeadlessChaosTestRewind.cpp
--- a/Engine/Source/Programs/HeadlessChaos/Private/HeadlessChaosTestRewind.cpp
+++ b/Engine/Source/Programs/HeadlessChaos/Private/HeadlessChaosTestRewind.cpp
@@ -12,6 +12,7 @@
 #include "Modules/ModuleManager.h"
 #include "RewindData.h"
 #include "GeometryCollection/GeometryCollectionTestFramework.h"
+#include "Chaos/SimCallbackObject.h"
 
 #ifndef REWIND_DESYNC
 #define REWIND_DESYNC 0
@@ -444,22 +445,22 @@ namespace ChaosTest {
 		return Result;
 	}
 
+	struct FSimCallbackHelperInput : FSimCallbackInput
+	{
+		void Reset() {}
+		int InCounter;
+	};
+
+	struct FSimCallbackHelperOutput : FSimCallbackOutput
+	{
+		void Reset() {}
+		int OutCounter;
+	};
+
 	TYPED_TEST(AllTraits, RewindTest_SimCallbackInputsOutputs)
 	{
 		TRewindHelper<TypeParam>::TestDynamicSphere([](auto* Solver, FReal SimDt, int32 Optimization, auto Proxy, auto Sphere)
 			{
-				struct FSimCallbackHelperInput : FSimCallbackInput
-				{
-					void Reset() {}
-					int InCounter;
-				};
-
-				struct FSimCallbackHelperOutput : FSimCallbackOutput
-				{
-					void Reset() {}
-					int OutCounter;
-				};
-
 				struct FSimCallbackHelper : TSimCallbackObject<FSimCallbackHelperInput, FSimCallbackHelperOutput>
 				{
 					int32 TriggerCount = 0;
