Unshelved from pending changelist '15485074':

ChaosVehicles: New code to allow the user to define the collision response of the suspension trace
Fix bug with vehicle not moving sometimes due to auto brake/auto gears and vehicle direction of travel.

diff --git a/Engine/Plugins/Experimental/ChaosVehiclesPlugin/Source/ChaosVehicles/Private/ChaosVehicleMovementComponent.cpp b/Engine/Plugins/Experimental/ChaosVehiclesPlugin/Source/ChaosVehicles/Private/ChaosVehicleMovementComponent.cpp
--- a/Engine/Plugins/Experimental/ChaosVehiclesPlugin/Source/ChaosVehicles/Private/ChaosVehicleMovementComponent.cpp
+++ b/Engine/Plugins/Experimental/ChaosVehiclesPlugin/Source/ChaosVehicles/Private/ChaosVehicleMovementComponent.cpp
@@ -966,61 +966,73 @@ float UChaosVehicleMovementComponent::CalcSteeringInput()
 	return RawSteeringInput;
 }
 
-float UChaosVehicleMovementComponent::CalcBrakeInput()
+void UChaosVehicleMovementComponent::CalcThrottleBrakeInput(float& ThrottleOut, float& BrakeOut)
 {
 	if (bReverseAsBrake)
 	{
-		float NewBrakeInput = 0.0f;
+		BrakeOut = RawBrakeInput;
+		ThrottleOut = RawThrottleInput;
 
-		// if player wants to move forwards...
 		if (RawThrottleInput > 0.f)
 		{
+			// car moving backwards but player wants to move forwards...
 			// if vehicle is moving backwards, then press brake
 			if (VehicleState.ForwardSpeed < -WrongDirectionThreshold)
 			{
-				NewBrakeInput = 1.0f;
+				BrakeOut = 1.0f; // = RawThrottleInput ?
+				ThrottleOut = 0.0f;
 			}
 
 		}
-
-		// if player wants to move backwards...
 		else if (RawBrakeInput > 0.f)
 		{
+			// car moving forwards but player wants to move backwards...
 			// if vehicle is moving forwards, then press brake
 			if (VehicleState.ForwardSpeed > WrongDirectionThreshold)
 			{
-				NewBrakeInput = 1.0f;
+				BrakeOut = 1.0f;
+				ThrottleOut = 0.0f;
+			}
+			else if (GetTargetGear() < 0)
+			{
+				ThrottleOut = RawBrakeInput;
+				BrakeOut = 0.0f;
 			}
 		}
-		// if player isn't pressing forward or backwards...
+		// straight reversing
+		else if (RawBrakeInput > 0.f && GetTargetGear() < 0)
+		{
+			ThrottleOut = RawBrakeInput;
+		}
 		else
 		{
+			// if player isn't pressing forward or backwards...
 			if (VehicleState.ForwardSpeed < StopThreshold && VehicleState.ForwardSpeed > -StopThreshold)	//auto brake 
 			{
-				NewBrakeInput = 1.f;
+				BrakeOut = 1.f;
 			}
 			else
 			{
-				NewBrakeInput = IdleBrakeInput;
+				BrakeOut = IdleBrakeInput;
 			}
 		}
 
-		return FMath::Clamp<float>(NewBrakeInput, 0.0, 1.0);
+		ThrottleOut = FMath::Clamp<float>(ThrottleOut, 0.0, 1.0);
+		BrakeOut = FMath::Clamp<float>(BrakeOut, 0.0, 1.0);
 	}
 	else
 	{
-		float NewBrakeInput = FMath::Abs(RawBrakeInput);
+		BrakeOut = FMath::Abs(RawBrakeInput);
 
 		// if player isn't pressing forward or backwards...
 		if (RawBrakeInput < SMALL_NUMBER && RawThrottleInput < SMALL_NUMBER)
 		{
 			if (VehicleState.ForwardSpeed < StopThreshold && VehicleState.ForwardSpeed > -StopThreshold)	//auto brake 
 			{
-				NewBrakeInput = 1.f;
+				BrakeOut = 1.f;
 			}
 		}
 
-		return NewBrakeInput;
 	}
 
 }
@@ -1045,26 +1057,6 @@ float UChaosVehicleMovementComponent::CalcYawInput()
 	return RawYawInput;
 }
 
-float UChaosVehicleMovementComponent::CalcThrottleInput()
-{
-	float NewThrottleInput = RawThrottleInput;
-	if (bReverseAsBrake)
-	{
-		if (RawBrakeInput > 0.f && GetTargetGear() < 0)
-		{
-			NewThrottleInput = RawBrakeInput;
-		}
-		else
-			//If the user is changing direction we should really be braking first and not applying any gas, so wait until they've changed gears
-			if ((RawThrottleInput > 0.f && GetTargetGear() < 0) || (RawBrakeInput > 0.f && GetTargetGear() > 0))
-			{
-				NewThrottleInput = 0.f;
-			}
-	}
-
-	return FMath::Abs(NewThrottleInput);
-}
-
 void UChaosVehicleMovementComponent::ClearInput()
 {
 	SteeringInput = 0.0f;
@@ -1123,11 +1115,11 @@ void UChaosVehicleMovementComponent::UpdateState(float DeltaTime)
 			{
 				if (RawBrakeInput > KINDA_SMALL_NUMBER && GetCurrentGear() >= 0 && GetTargetGear() >= 0)
 				{
-					SetTargetGear(-1, false);
+					SetTargetGear(-1, true);
 				}
 				else if (RawThrottleInput > KINDA_SMALL_NUMBER && GetCurrentGear() <= 0 && GetTargetGear() <= 0)
 				{
-					SetTargetGear(1, false);
+					SetTargetGear(1, true);
 				}
 			}
 		}
@@ -1142,9 +1134,12 @@ void UChaosVehicleMovementComponent::UpdateState(float DeltaTime)
 			}
 		}
 
+		float ModifiedThrottle = 0.f;
+		float ModifiedBrake = 0.f;
+		CalcThrottleBrakeInput(ModifiedThrottle, ModifiedBrake);
 		SteeringInput = SteeringInputRate.InterpInputValue(DeltaTime, SteeringInput, CalcSteeringInput());
-		ThrottleInput = ThrottleInputRate.InterpInputValue(DeltaTime, ThrottleInput, CalcThrottleInput());
-		BrakeInput = BrakeInputRate.InterpInputValue(DeltaTime, BrakeInput, CalcBrakeInput());
+		ThrottleInput = ThrottleInputRate.InterpInputValue(DeltaTime, ThrottleInput, ModifiedThrottle);
+		BrakeInput = BrakeInputRate.InterpInputValue(DeltaTime, BrakeInput, ModifiedBrake);
 		PitchInput = PitchInputRate.InterpInputValue(DeltaTime, PitchInput, CalcPitchInput());
 		RollInput = RollInputRate.InterpInputValue(DeltaTime, RollInput, CalcRollInput());
 		YawInput = YawInputRate.InterpInputValue(DeltaTime, YawInput, CalcYawInput());
@@ -1697,15 +1692,6 @@ void UChaosVehicleMovementComponent::Update(float DeltaTime)
 
 				AsyncInput->GravityZ = GetGravityZ();
 
-				// #TODO: don't create every frame, setup never changes dynamically
-				TArray<AActor*> ActorsToIgnore;
-				ActorsToIgnore.Add(GetPawnOwner()); // ignore self in scene query
-
-				FCollisionQueryParams TraceParams(NAME_None, FCollisionQueryParams::GetUnknownStatId(), false, nullptr);
-				TraceParams.bReturnPhysicalMaterial = true;	// we need this to get the surface friction coefficient
-				TraceParams.AddIgnoredActors(ActorsToIgnore);
-				TraceParams.bTraceComplex = true;  // (Wheels[WheelIdx]->SweepType == ESweepType::ComplexSweep); FIX this again
-				AsyncInput->TraceParams = TraceParams;
 			}
 		}
 	}
diff --git a/Engine/Plugins/Experimental/ChaosVehiclesPlugin/Source/ChaosVehicles/Private/ChaosWheeledVehicleMovementComponent.cpp b/Engine/Plugins/Experimental/ChaosVehiclesPlugin/Source/ChaosVehicles/Private/ChaosWheeledVehicleMovementComponent.cpp
--- a/Engine/Plugins/Experimental/ChaosVehiclesPlugin/Source/ChaosVehicles/Private/ChaosWheeledVehicleMovementComponent.cpp
+++ b/Engine/Plugins/Experimental/ChaosVehiclesPlugin/Source/ChaosVehicles/Private/ChaosWheeledVehicleMovementComponent.cpp
@@ -205,7 +205,7 @@ void UChaosWheeledVehicleSimulation::UpdateSimulation(float DeltaTime, const FCh
 
 		if (!GWheeledVehicleDebugParams.DisableSuspensionForces && PVehicle->bSuspensionEnabled)
 		{
-			PerformSuspensionTraces(WheelState.Trace, InputData.TraceParams);
+			PerformSuspensionTraces(WheelState.Trace, InputData.TraceParams, InputData.TraceCollisionResponse);
 		}
 
 		//////////////////////////////////////////////////////////////////////////
@@ -269,11 +269,11 @@ void UChaosWheeledVehicleSimulation::UpdateSimulation(float DeltaTime, const FCh
 
 }
 
-void UChaosWheeledVehicleSimulation::PerformSuspensionTraces(const TArray<FSuspensionTrace>& SuspensionTrace, FCollisionQueryParams& TraceParams)
+void UChaosWheeledVehicleSimulation::PerformSuspensionTraces(const TArray<FSuspensionTrace>& SuspensionTrace, FCollisionQueryParams& TraceParams, FCollisionResponseContainer& CollisionResponse)
 {
-	//#todo: SpringCollisionChannel should be a parameter setup
 	ECollisionChannel SpringCollisionChannel = ECollisionChannel::ECC_WorldDynamic;
 	FCollisionResponseParams ResponseParams;
+	ResponseParams.CollisionResponse = CollisionResponse;
 
 	// batching is about 0.5ms (25%) faster when there's 100 vehicles on a flat terrain
 	if (GVehicleDebugParams.BatchQueries)
@@ -400,7 +400,7 @@ void UChaosWheeledVehicleSimulation::PerformSuspensionTraces(const TArray<FSuspe
 					, TraceEnd + VehicleUpAxis * Radius
 					, FQuat::Identity, SpringCollisionChannel
 					, FCollisionShape::MakeSphere(Radius), TraceParams
-					, FCollisionResponseParams::DefaultResponseParam);
+					, ResponseParams);
 			}
 			break;
 
@@ -409,7 +409,7 @@ void UChaosWheeledVehicleSimulation::PerformSuspensionTraces(const TArray<FSuspe
 			case ESweepShape::Raycast:
 			default:
 			{
-				World->LineTraceSingleByChannel(HitResult, TraceStart, TraceEnd, SpringCollisionChannel, TraceParams, FCollisionResponseParams::DefaultResponseParam);
+				World->LineTraceSingleByChannel(HitResult, TraceStart, TraceEnd, SpringCollisionChannel, TraceParams, ResponseParams);
 			}
 			break;
 			}
@@ -505,7 +505,7 @@ void UChaosWheeledVehicleSimulation::ApplySuspensionForces(float DeltaTime)
 		if (PWheel.Setup().NewSimulationPath)
 		{
 #if WITH_CHAOS
-			if (ConstraintHandles.Num() > 0)
+			if (WheelIdx < ConstraintHandles.Num())
 			{
 				FPhysicsConstraintHandle& ConstraintHandle = ConstraintHandles[WheelIdx];
 				if (ConstraintHandle.IsValid())
@@ -694,17 +694,17 @@ void UChaosWheeledVehicleSimulation::ApplyInput(const FControlInputs& ControlInp
 		if (PWheel.Setup().BrakeEnabled)
 		{
 			float BrakeForce = PWheel.Setup().MaxBrakeTorque * ModifiedInputs.BrakeInput;
-			PWheel.SetBrakeTorque(MToCm(BrakeForce + EngineBrakingForce));
+			PWheel.SetBrakeTorque(TorqueMToCm(BrakeForce + EngineBrakingForce));
 		}
 		else
 		{
-			PWheel.SetBrakeTorque(MToCm(EngineBraking));
+			PWheel.SetBrakeTorque(TorqueMToCm(EngineBraking));
 		}
 
 		if ((ModifiedInputs.HandbrakeInput && PWheel.Setup().HandbrakeEnabled) || ModifiedInputs.ParkingEnabled)
 		{
 			float HandbrakeForce = ModifiedInputs.ParkingEnabled ? PWheel.Setup().HandbrakeTorque : (ModifiedInputs.HandbrakeInput * PWheel.Setup().HandbrakeTorque);
-			PWheel.SetBrakeTorque(MToCm(HandbrakeForce));
+			PWheel.SetBrakeTorque(TorqueMToCm(HandbrakeForce));
 		}
 	}
 
@@ -773,11 +773,11 @@ void UChaosWheeledVehicleSimulation::ProcessMechanicalSimulation(float DeltaTime
 						SplitTorque = PDifferential.Setup().FrontRearSplit;
 					}
 
-					PWheel.SetDriveTorque(MToCm(TransmissionTorque * SplitTorque) / (float)PVehicle->NumDrivenWheels);
+					PWheel.SetDriveTorque(TorqueMToCm(TransmissionTorque * SplitTorque) / (float)PVehicle->NumDrivenWheels);
 				}
 				else
 				{
-					PWheel.SetDriveTorque(MToCm(TransmissionTorque) / (float)PVehicle->NumDrivenWheels);
+					PWheel.SetDriveTorque(TorqueMToCm(TransmissionTorque) / (float)PVehicle->NumDrivenWheels);
 				}
 			}
 		}
@@ -930,6 +930,9 @@ UChaosWheeledVehicleMovementComponent::UChaosWheeledVehicleMovementComponent(con
 	bMechanicalSimEnabled = true;
 	bSuspensionEnabled = true;
 	bWheelFrictionEnabled = true;
+
+	WheelTraceCollisionResponses = FCollisionResponseContainer::GetDefaultResponseContainer();
+	WheelTraceCollisionResponses.Vehicle = ECR_Ignore;
 }
 
 // Public
@@ -1449,6 +1452,33 @@ FVector2D UChaosWheeledVehicleMovementComponent::CalculateWheelLayoutDimensions(
 	return MaxSize;
 }
 
+void UChaosWheeledVehicleMovementComponent::Update(float DeltaTime)
+{
+	UChaosVehicleMovementComponent::Update(DeltaTime);
+
+	if (CurAsyncInput)
+	{
+		if (const FBodyInstance* BodyInstance = GetBodyInstance())
+		{
+			if (auto Handle = BodyInstance->ActorHandle)
+			{
+				FChaosVehicleDefaultAsyncInput* AsyncInput = static_cast<FChaosVehicleDefaultAsyncInput*>(CurAsyncInput);
+
+				TArray<AActor*> ActorsToIgnore;
+				ActorsToIgnore.Add(GetPawnOwner()); // ignore self in scene query
+
+				FCollisionQueryParams TraceParams(NAME_None, FCollisionQueryParams::GetUnknownStatId(), false, nullptr);
+				TraceParams.bReturnPhysicalMaterial = true;	// we need this to get the surface friction coefficient
+				TraceParams.AddIgnoredActors(ActorsToIgnore);
+				TraceParams.bTraceComplex = true;
+				AsyncInput->TraceParams = TraceParams;
+				AsyncInput->TraceCollisionResponse = WheelTraceCollisionResponses;
+			}
+		}
+	}
+}
+
+
 // Debug
 void UChaosWheeledVehicleMovementComponent::DrawDebug(UCanvas* Canvas, float& YL, float& YPos)
 {
diff --git a/Engine/Plugins/Experimental/ChaosVehiclesPlugin/Source/ChaosVehicles/Public/ChaosVehicleMovementComponent.h b/Engine/Plugins/Experimental/ChaosVehiclesPlugin/Source/ChaosVehicles/Public/ChaosVehicleMovementComponent.h
--- a/Engine/Plugins/Experimental/ChaosVehiclesPlugin/Source/ChaosVehicles/Public/ChaosVehicleMovementComponent.h
+++ b/Engine/Plugins/Experimental/ChaosVehiclesPlugin/Source/ChaosVehicles/Public/ChaosVehicleMovementComponent.h
@@ -151,6 +151,7 @@ struct FChaosVehicleDefaultAsyncInput : public FChaosVehicleAsyncInput
 	float GravityZ;
 	FControlInputs ControlInputs;
 	mutable FCollisionQueryParams TraceParams;
+	mutable FCollisionResponseContainer TraceCollisionResponse;
 
 	FChaosVehicleDefaultAsyncInput();
 
@@ -1003,7 +1004,7 @@ public:
 	void SetCurrentAsyncInputOutputInternal(FChaosVehicleAsyncInput* CurInput, int32 InputIdx, FChaosVehicleManagerAsyncOutput* CurOutput, int32 VehicleManagerTimestamp);
 	void SetCurrentAsyncInputOutputInternal(FChaosVehicleAsyncInput* CurInput, int32 InputIdx, FChaosVehicleManagerAsyncOutput* CurOutput, FChaosVehicleManagerAsyncOutput* NextOutput, float Alpha, int32 VehicleManagerTimestamp);
 
-	void Update(float DeltaTime);
+	virtual void Update(float DeltaTime);
 
 	// Get output data from Physics Thread
 	virtual void ParallelUpdate(float DeltaSeconds);
@@ -1140,8 +1141,8 @@ protected:
 	/** Compute steering input */
 	float CalcSteeringInput();
 
-	/** Compute brake input */
-	float CalcBrakeInput();
+	/** Compute throttle & brake input */
+	void CalcThrottleBrakeInput(float& ThrottleOut, float& BrakeOut);
 
 	/** Compute handbrake input */
 	float CalcHandbrakeInput();
@@ -1156,7 +1157,6 @@ protected:
 	float CalcYawInput();
 
 	/** Compute throttle inputs */
-	virtual float CalcThrottleInput();
 	float CalcThrottleUpInput();
 	float CalcThrottleDownInput();
 
diff --git a/Engine/Plugins/Experimental/ChaosVehiclesPlugin/Source/ChaosVehicles/Public/ChaosWheeledVehicleMovementComponent.h b/Engine/Plugins/Experimental/ChaosVehiclesPlugin/Source/ChaosVehicles/Public/ChaosWheeledVehicleMovementComponent.h
--- a/Engine/Plugins/Experimental/ChaosVehiclesPlugin/Source/ChaosVehicles/Public/ChaosWheeledVehicleMovementComponent.h
+++ b/Engine/Plugins/Experimental/ChaosVehiclesPlugin/Source/ChaosVehicles/Public/ChaosWheeledVehicleMovementComponent.h
@@ -565,7 +565,7 @@ public:
 	virtual void ApplyInput(const FControlInputs& ControlInputs, float DeltaTime) override;
 
 	/** Perform suspension ray/shape traces */
-	virtual void PerformSuspensionTraces(const TArray<Chaos::FSuspensionTrace>& SuspensionTrace, FCollisionQueryParams& TraceParams);
+	virtual void PerformSuspensionTraces(const TArray<Chaos::FSuspensionTrace>& SuspensionTrace, FCollisionQueryParams& TraceParams, FCollisionResponseContainer& CollisionResponse);
 
 
 	/** Update the engine/transmission simulation */
@@ -612,6 +612,9 @@ class CHAOSVEHICLES_API UChaosWheeledVehicleMovementComponent : public UChaosVeh
 	UPROPERTY(EditAnywhere, Category = WheelSetup)
 	TArray<FChaosWheelSetup> WheelSetups;
 
+	UPROPERTY(EditAnywhere, Category = Custom)
+	struct FCollisionResponseContainer WheelTraceCollisionResponses;
+
 	UPROPERTY(EditAnywhere, Category = MechanicalSetup)
 	bool bMechanicalSimEnabled;
 
@@ -680,6 +683,22 @@ class CHAOSVEHICLES_API UChaosWheeledVehicleMovementComponent : public UChaosVeh
 		return WheelStatus[WheelIndex];
 	}
 
+
+	/** Set all channels to the specified response - for wheel raycasts */
+	void SetWheelTraceAllChannels(ECollisionResponse NewResponse)
+	{
+		WheelTraceCollisionResponses.SetAllChannels(NewResponse);
+	}
+
+	/** Set the response of this body to the supplied settings - for wheel raycasts */
+	void SetWheelTraceResponseToChannel(ECollisionChannel Channel, ECollisionResponse NewResponse)
+	{
+		WheelTraceCollisionResponses.SetResponse(Channel, NewResponse);
+	}
+
+	/** Get Collision ResponseToChannels container for this component **/
+//	FORCEINLINE_DEBUGGABLE const FCollisionResponseContainer& GetTraceResponseToChannels() const { return WheelTraceCollisionResponses.GetResponseContainer(); }
+
 	//////////////////////////////////////////////////////////////////////////
 	// Public
 
@@ -787,6 +806,9 @@ protected:
 	// Update
 	void FillWheelOutputState();
 
+	/* Fill Async input state */
+	virtual void Update(float DeltaTime) override;
+
 	//////////////////////////////////////////////////////////////////////////
 	// Debug
 
diff --git a/Engine/Source/Programs/HeadlessChaos/Private/HeadlessChaosTestVehicles.cpp b/Engine/Source/Programs/HeadlessChaos/Private/HeadlessChaosTestVehicles.cpp
--- a/Engine/Source/Programs/HeadlessChaos/Private/HeadlessChaosTestVehicles.cpp
+++ b/Engine/Source/Programs/HeadlessChaos/Private/HeadlessChaosTestVehicles.cpp
@@ -700,16 +700,17 @@ namespace ChaosTest
 
 	// Wheel
 	void SimulateBraking(FSimpleWheelSim& Wheel
-		, float VehicleSpeedMPH
+		, const float Gravity
+		, float VehicleSpeed
 		, float DeltaTime
 		, float& StoppingDistanceOut
 		, float& SimulationTimeOut
+		, bool bLoggingEnabled = false
 		)
 	{
 		StoppingDistanceOut = 0.f;
 		SimulationTimeOut = 0.f;
 		
-		const float Gravity = 9.8f;
 		float MaxSimTime = 15.0f;
 		float VehicleMass = 1300.f;
 		float VehicleMassPerWheel = 1300.f / 4.f;
@@ -718,11 +719,16 @@ namespace ChaosTest
 		Wheel.SetMassPerWheel(VehicleMassPerWheel);
 
 		// Road speed
-		FVector Velocity = FVector(MPHToMS(VehicleSpeedMPH), 0.f, 0.f);
+		FVector Velocity = FVector(VehicleSpeed, 0.f, 0.f);
 
 		// wheel rolling speed matches road speed
 		Wheel.SetMatchingSpeed(Velocity.X);
 
+		if (bLoggingEnabled)
+		{
+			UE_LOG(LogChaos, Warning, TEXT("--------------------START---------------------"));
+		}
+
 		while (SimulationTimeOut < MaxSimTime)
 		{
 			// rolling speed matches road speed
@@ -733,7 +739,11 @@ namespace ChaosTest
 			Velocity += DeltaTime * Wheel.GetForceFromFriction() / VehicleMassPerWheel;
 			StoppingDistanceOut += Velocity.X * DeltaTime;
 
-			// #todo: make this better remove the 2.0f
+			if (bLoggingEnabled)
+			{
+				UE_LOG(LogChaos, Warning, TEXT("Wheel.GetForceFromFriction() %s"), *Wheel.GetForceFromFriction().ToString());
+			}
+
 			if (FMath::Abs(Velocity.X) < 0.05f)
 			{
 				Velocity.X = 0.f;
@@ -742,11 +752,18 @@ namespace ChaosTest
 
 			SimulationTimeOut += DeltaTime;
 		}
+
+		if (bLoggingEnabled)
+		{
+			UE_LOG(LogChaos, Warning, TEXT("---------------------END----------------------"));
+		}
+
+
 	}
 
 	void SimulateAccelerating(FSimpleWheelSim& Wheel
 		, const float Gravity
-		, float FinalVehicleSpeedMPH
+		, float FinalVehicleSpeed
 		, float DeltaTime
 		, float& DistanceTravelledOut
 		, float& SimulationTimeOut
@@ -778,7 +795,7 @@ namespace ChaosTest
 
 			SimulationTimeOut += DeltaTime;
 
-			if (FMath::Abs(Velocity.X) >= MPHToMS(FinalVehicleSpeedMPH))
+			if (FMath::Abs(Velocity.X) >= FinalVehicleSpeed)
 			{
 				break; // time is up
 			}
@@ -793,7 +810,7 @@ namespace ChaosTest
 		Setup.TractionControlEnabled = false;
 		Setup.BrakeEnabled = true;
 		Setup.EngineEnabled = true;
-		Setup.WheelRadius = 30.0f;
+		Setup.WheelRadius = 0.3f;
 		Setup.LongitudinalFrictionMultiplier = 1.0f;
 		Setup.LateralFrictionMultiplier = 1.0f;
 		Setup.SideSlipModifier = 0.7f;
@@ -804,7 +821,9 @@ namespace ChaosTest
 		// So using a range 10-20 to ensure we are in the correct ballpark.
 		// If specified more accurately in the test, then modifying the code would break the test all the time.
 
-		float StoppingDistanceTolerance = 0.5f; // meters
+		// units meters
+		float Gravity = 9.8f;
+		float StoppingDistanceTolerance = 0.5f;
 		float DeltaTime = 1.f / 30.f;
 		float StoppingDistanceA = 0.f;
 		float SimulationTime = 0.0f;
@@ -812,71 +831,70 @@ namespace ChaosTest
 
 		// reasonably ideal stopping distance - traveling forwards
 		Wheel.SetBrakeTorque(650);
-		float VehicleSpeedMPH = 30.f;
-		SimulateBraking(Wheel, VehicleSpeedMPH, DeltaTime, StoppingDistanceA, SimulationTime);
-		//UE_LOG(LogChaos, Warning, TEXT("Braking Distance %3.2fm"), StoppingDistanceB);
+		SimulateBraking(Wheel, Gravity, MPHToMS(30.f), DeltaTime, StoppingDistanceA, SimulationTime);
 		EXPECT_GT(StoppingDistanceA, 10.f);
 		EXPECT_LT(StoppingDistanceA, 20.f);
 
+		// Changing to units of Cm should yield the same results
+		float MToCm = 100.0f;
+		float StoppingDistanceCm = 0.f;
+		Wheel.SetBrakeTorque(650 * MToCm * MToCm);
+		Wheel.SetWheelRadius(0.3f * MToCm);
+		SimulateBraking(Wheel, Gravity * MToCm, MPHToCmS(30.f), DeltaTime, StoppingDistanceCm, SimulationTime);
+		EXPECT_NEAR(StoppingDistanceCm, StoppingDistanceA * MToCm, StoppingDistanceTolerance * MToCm);
+
 		// traveling backwards stops just the same
 		float StoppingDistanceReverseDir = 0.f;
+		Wheel.SetWheelRadius(0.3f);
 		Wheel.SetBrakeTorque(650);
-		VehicleSpeedMPH = -30.f;
-		SimulateBraking(Wheel, VehicleSpeedMPH, DeltaTime, StoppingDistanceReverseDir, SimulationTime);
+		SimulateBraking(Wheel, Gravity, MPHToMS(-30.f), DeltaTime, StoppingDistanceReverseDir, SimulationTime);
 		EXPECT_GT(StoppingDistanceReverseDir, -20.f);
 		EXPECT_LT(StoppingDistanceReverseDir, -10.f);
 		EXPECT_LT(StoppingDistanceA - FMath::Abs(StoppingDistanceReverseDir), StoppingDistanceTolerance);
 
 		// Similar results with different delta time
 		float StoppingDistanceDiffDT = 0.f;
-		VehicleSpeedMPH = 30.f;
-		SimulateBraking(Wheel, VehicleSpeedMPH, DeltaTime * 0.25f, StoppingDistanceDiffDT, SimulationTime);
+		SimulateBraking(Wheel, Gravity, MPHToMS(30.f), DeltaTime * 0.25f, StoppingDistanceDiffDT, SimulationTime);
 		EXPECT_LT(StoppingDistanceA - StoppingDistanceDiffDT, StoppingDistanceTolerance);
 
 		// barely touching the brake - going to take longer to stop
 		float StoppingDistanceLightBraking = 0.f;
 		Wheel.SetBrakeTorque(150);
-		VehicleSpeedMPH = 30.f;
-		SimulateBraking(Wheel, VehicleSpeedMPH, DeltaTime, StoppingDistanceLightBraking, SimulationTime);
+		SimulateBraking(Wheel, Gravity, MPHToMS(30.f), DeltaTime, StoppingDistanceLightBraking, SimulationTime);
 		EXPECT_GT(StoppingDistanceLightBraking, StoppingDistanceA);
 
 		// locking the wheels / too much brake torque -> dynamic friction rather than static friction -> going to take longer to stop
 		float StoppingDistanceTooHeavyBreaking = 0.f;
 		Wheel.SetBrakeTorque(5000);
-		VehicleSpeedMPH = 30.f;
-		SimulateBraking(Wheel, VehicleSpeedMPH, DeltaTime, StoppingDistanceTooHeavyBreaking, SimulationTime);
+		SimulateBraking(Wheel, Gravity, MPHToMS(30.f), DeltaTime, StoppingDistanceTooHeavyBreaking, SimulationTime);
 		EXPECT_GT(StoppingDistanceTooHeavyBreaking, StoppingDistanceA);
 
 		// Would have locked the wheels but ABS prevents skidding
 		Wheel.AccessSetup().ABSEnabled = true;
 		float StoppingDistanceTooHeavyBreakingABS = 0.f;
 		Wheel.SetBrakeTorque(5000);
-		VehicleSpeedMPH = 30.f;
-		SimulateBraking(Wheel, VehicleSpeedMPH, DeltaTime, StoppingDistanceTooHeavyBreakingABS, SimulationTime);
+		SimulateBraking(Wheel, Gravity, MPHToMS(30.f), DeltaTime, StoppingDistanceTooHeavyBreakingABS, SimulationTime);
 		EXPECT_LT(StoppingDistanceTooHeavyBreakingABS, StoppingDistanceTooHeavyBreaking);
 		Wheel.AccessSetup().ABSEnabled = false;
 
 		// lower initial speed - stops more quickly
 		float StoppingDistanceLowerSpeed = 0.f;
 		Wheel.SetBrakeTorque(650);
-		VehicleSpeedMPH = 20.f;
-		SimulateBraking(Wheel, VehicleSpeedMPH, DeltaTime, StoppingDistanceLowerSpeed, SimulationTime);
+		SimulateBraking(Wheel, Gravity, MPHToMS(20.f), DeltaTime, StoppingDistanceLowerSpeed, SimulationTime);
 		EXPECT_LT(StoppingDistanceLowerSpeed, StoppingDistanceA);
 
 		// higher initial speed - stops more slowly
 		float StoppingDistanceHigherSpeed = 0.f;
 		Wheel.SetBrakeTorque(650);
-		VehicleSpeedMPH = 60.f;
-		SimulateBraking(Wheel, VehicleSpeedMPH, DeltaTime, StoppingDistanceHigherSpeed, SimulationTime);
+		SimulateBraking(Wheel, Gravity, MPHToMS(60.f), DeltaTime, StoppingDistanceHigherSpeed, SimulationTime);
 		EXPECT_GT(StoppingDistanceHigherSpeed, StoppingDistanceA);
 
 		// slippy surface - stops more slowly
 		float StoppingDistanceLowFriction = 0.f;
 		Wheel.SetSurfaceFriction(0.3f);
 		Wheel.SetBrakeTorque(650);
-		VehicleSpeedMPH = 30.f;
-		SimulateBraking(Wheel, VehicleSpeedMPH, DeltaTime, StoppingDistanceLowFriction, SimulationTime);
-		EXPECT_GT(StoppingDistanceLowFriction, StoppingDistanceA);
+		SimulateBraking(Wheel, Gravity, MPHToMS(30.f), DeltaTime, StoppingDistanceLowFriction, SimulationTime);
+		EXPECT_GT(StoppingDistanceLowFriction, StoppingDistanceA);		
 	}
 
 	TYPED_TEST(AllTraits, VehicleTest_WheelAcceleratingLongitudinalSlip)
@@ -886,7 +904,7 @@ namespace ChaosTest
 		Setup.TractionControlEnabled = false;
 		Setup.BrakeEnabled = true;
 		Setup.EngineEnabled = true;
-		Setup.WheelRadius = 30.0f;
+		Setup.WheelRadius = 0.3f;
 		Setup.LongitudinalFrictionMultiplier = 1.0f;
 		Setup.LateralFrictionMultiplier = 1.0f;
 		Setup.SideSlipModifier = 0.7f;
@@ -908,13 +926,13 @@ namespace ChaosTest
 
 		// How far & what time does it take to stop from 30MPH to rest
 		Wheel.SetBrakeTorque(650);
-		SimulateBraking(Wheel, 30.0f, DeltaTime, StoppingDistanceA, SimulationTimeBrake);
+		SimulateBraking(Wheel, Gravity, MPHToMS(30.0f), DeltaTime, StoppingDistanceA, SimulationTimeBrake);
 
 		// How far and what time does it take to accelerate from rest to 30MPH
 		float SimulationTimeAccel = 0.0f;
 		float DrivingDistanceA = 0.f;
 		Wheel.SetDriveTorque(650);
-		SimulateAccelerating(Wheel, Gravity, 30.0f, DeltaTime, DrivingDistanceA, SimulationTimeAccel);
+		SimulateAccelerating(Wheel, Gravity, MPHToMS(30.0f), DeltaTime, DrivingDistanceA, SimulationTimeAccel);
 
 		// 0-30 MPH and 30-0 MPH should be the same if there's no slipping and accel torque was same as the brake torque run
 		EXPECT_LT(DrivingDistanceA - StoppingDistanceA, AccelerationResultsTolerance);
@@ -928,23 +946,25 @@ namespace ChaosTest
 		float SimulationTimeAccelCM = 0.0f;
 		float MtoCm = 100.0f;
 		float DrivingDistanceCM = 0.f;
-		Wheel.SetDriveTorque(650 * MtoCm);
-		SimulateAccelerating(Wheel, Gravity * MtoCm, 30.0f * MtoCm, DeltaTime, DrivingDistanceCM, SimulationTimeAccelCM);
+		Wheel.SetDriveTorque(650 * MtoCm * MtoCm);
+		Wheel.SetWheelRadius(0.3f * MtoCm);
+		SimulateAccelerating(Wheel, Gravity * MtoCm, MPHToCmS(30.0f), DeltaTime, DrivingDistanceCM, SimulationTimeAccelCM);
 		EXPECT_GT(DrivingDistanceCM, 10.f * MtoCm);
 		EXPECT_LT(DrivingDistanceCM, 20.f * MtoCm);
-		EXPECT_LT(SimulationTimeAccel - SimulationTimeAccelCM, AccelerationResultsTolerance);
+		EXPECT_NEAR(SimulationTimeAccel, SimulationTimeAccelCM, AccelerationResultsTolerance);
 
 		float SimulationTimeAccelSpin = 0.0f;
 		float DrivingDistanceWheelspin = 0.f;
+		Wheel.SetWheelRadius(0.3f);
 		Wheel.SetDriveTorque(5000); // definitely cause wheel spin
-		SimulateAccelerating(Wheel, Gravity, 30.0f, DeltaTime, DrivingDistanceWheelspin, SimulationTimeAccelSpin);
+		SimulateAccelerating(Wheel, Gravity, MPHToMS(30.0f), DeltaTime, DrivingDistanceWheelspin, SimulationTimeAccelSpin);
 
 		// Enable traction control should be better than both of the above
 		float SimulationTimeAccelTC = 0.0f;
 		float DrivingDistanceTC = 0.f;
 		Wheel.AccessSetup().TractionControlEnabled = true;
 		Wheel.SetDriveTorque(5000); // definitely cause wheel spin
-		SimulateAccelerating(Wheel, Gravity, 30.0f, DeltaTime, DrivingDistanceTC, SimulationTimeAccelTC);
+		SimulateAccelerating(Wheel, Gravity, MPHToMS(30.0f), DeltaTime, DrivingDistanceTC, SimulationTimeAccelTC);
 
 		// reaches target speed in a shorter distance
 		EXPECT_LT(DrivingDistanceTC, DrivingDistanceWheelspin);
@@ -960,7 +980,7 @@ namespace ChaosTest
 		Setup.TractionControlEnabled = false;
 		Setup.BrakeEnabled = true;
 		Setup.EngineEnabled = true;
-		Setup.WheelRadius = 30.0f;
+		Setup.WheelRadius = 0.3f;
 		Setup.NewSimulationPath = true;
 		Setup.FrictionMultiplier = 1.0f;
 		Setup.CorneringStiffness = 2.0f;
@@ -972,7 +992,9 @@ namespace ChaosTest
 		// So using a range 10-20 to ensure we are in the correct ballpark.
 		// If specified more accurately in the test, then modifying the code would break the test all the time.
 
-		float StoppingDistanceTolerance = 0.5f; // meters
+		// units meters
+		float Gravity = 9.8f;
+		float StoppingDistanceTolerance = 0.5f;
 		float DeltaTime = 1.f / 30.f;
 		float StoppingDistanceA = 0.f;
 		float SimulationTime = 0.0f;
@@ -980,70 +1002,69 @@ namespace ChaosTest
 
 		// reasonably ideal stopping distance - traveling forwards
 		Wheel.SetBrakeTorque(650);
-		float VehicleSpeedMPH = 30.f;
-		SimulateBraking(Wheel, VehicleSpeedMPH, DeltaTime, StoppingDistanceA, SimulationTime);
-		//UE_LOG(LogChaos, Warning, TEXT("Braking Distance %3.2fm"), StoppingDistanceB);
+		SimulateBraking(Wheel, Gravity, MPHToMS(30.f), DeltaTime, StoppingDistanceA, SimulationTime);
 		EXPECT_GT(StoppingDistanceA, 10.f);
 		EXPECT_LT(StoppingDistanceA, 20.f);
 
 		// traveling backwards stops just the same
 		float StoppingDistanceReverseDir = 0.f;
 		Wheel.SetBrakeTorque(650);
-		VehicleSpeedMPH = -30.f;
-		SimulateBraking(Wheel, VehicleSpeedMPH, DeltaTime, StoppingDistanceReverseDir, SimulationTime);
+		SimulateBraking(Wheel, Gravity, MPHToMS(-30.f), DeltaTime, StoppingDistanceReverseDir, SimulationTime);
 		EXPECT_GT(StoppingDistanceReverseDir, -20.f);
 		EXPECT_LT(StoppingDistanceReverseDir, -10.f);
 		EXPECT_LT(StoppingDistanceA - FMath::Abs(StoppingDistanceReverseDir), StoppingDistanceTolerance);
 
+		// Changing to units of Cm should yield the same results
+		float MToCm = 100.0f;
+		float StoppingDistanceCm = 0.f;
+		Wheel.SetBrakeTorque(650 * MToCm * MToCm);
+		Wheel.SetWheelRadius(0.3f * MToCm);
+		SimulateBraking(Wheel, Gravity * MToCm, MPHToCmS(30.f), DeltaTime, StoppingDistanceCm, SimulationTime);
+		EXPECT_NEAR(StoppingDistanceCm, StoppingDistanceA * MToCm, 1.0f);
+
 		// Similar results with different delta time
 		float StoppingDistanceDiffDT = 0.f;
-		VehicleSpeedMPH = 30.f;
-		SimulateBraking(Wheel, VehicleSpeedMPH, DeltaTime * 0.25f, StoppingDistanceDiffDT, SimulationTime);
+		Wheel.SetWheelRadius(0.3f);
+		SimulateBraking(Wheel, Gravity, MPHToMS(30.f), DeltaTime * 0.25f, StoppingDistanceDiffDT, SimulationTime);
 		EXPECT_LT(StoppingDistanceA - StoppingDistanceDiffDT, StoppingDistanceTolerance);
 
 		// barely touching the brake - going to take longer to stop
 		float StoppingDistanceLightBraking = 0.f;
 		Wheel.SetBrakeTorque(150);
-		VehicleSpeedMPH = 30.f;
-		SimulateBraking(Wheel, VehicleSpeedMPH, DeltaTime, StoppingDistanceLightBraking, SimulationTime);
+		SimulateBraking(Wheel, Gravity, MPHToMS(30.f), DeltaTime, StoppingDistanceLightBraking, SimulationTime);
 		EXPECT_GT(StoppingDistanceLightBraking, StoppingDistanceA);
 
 		// locking the wheels / too much brake torque -> dynamic friction rather than static friction -> going to take longer to stop
 		float StoppingDistanceTooHeavyBreaking = 0.f;
 		Wheel.SetBrakeTorque(5000);
-		VehicleSpeedMPH = 30.f;
-		SimulateBraking(Wheel, VehicleSpeedMPH, DeltaTime, StoppingDistanceTooHeavyBreaking, SimulationTime);
+		SimulateBraking(Wheel, Gravity, MPHToMS(30.f), DeltaTime, StoppingDistanceTooHeavyBreaking, SimulationTime);
 		EXPECT_GT(StoppingDistanceTooHeavyBreaking, StoppingDistanceA);
 
 		// Would have locked the wheels but ABS prevents skidding
 		Wheel.AccessSetup().ABSEnabled = true;
 		float StoppingDistanceTooHeavyBreakingABS = 0.f;
 		Wheel.SetBrakeTorque(5000);
-		VehicleSpeedMPH = 30.f;
-		SimulateBraking(Wheel, VehicleSpeedMPH, DeltaTime, StoppingDistanceTooHeavyBreakingABS, SimulationTime);
+		SimulateBraking(Wheel, Gravity, MPHToMS(30.f), DeltaTime, StoppingDistanceTooHeavyBreakingABS, SimulationTime);
 		EXPECT_LT(StoppingDistanceTooHeavyBreakingABS, StoppingDistanceTooHeavyBreaking);
 		Wheel.AccessSetup().ABSEnabled = false;
 
 		// lower initial speed - stops more quickly
 		float StoppingDistanceLowerSpeed = 0.f;
 		Wheel.SetBrakeTorque(650);
-		VehicleSpeedMPH = 20.f;
-		SimulateBraking(Wheel, VehicleSpeedMPH, DeltaTime, StoppingDistanceLowerSpeed, SimulationTime);
+		SimulateBraking(Wheel, Gravity, MPHToMS(20.f), DeltaTime, StoppingDistanceLowerSpeed, SimulationTime);
 		EXPECT_LT(StoppingDistanceLowerSpeed, StoppingDistanceA);
 
 		// higher initial speed - stops more slowly
 		float StoppingDistanceHigherSpeed = 0.f;
 		Wheel.SetBrakeTorque(650);
-		VehicleSpeedMPH = 60.f;
-		SimulateBraking(Wheel, VehicleSpeedMPH, DeltaTime, StoppingDistanceHigherSpeed, SimulationTime);
+		SimulateBraking(Wheel, Gravity, MPHToMS(60.f), DeltaTime, StoppingDistanceHigherSpeed, SimulationTime);
 		EXPECT_GT(StoppingDistanceHigherSpeed, StoppingDistanceA);
 
 		// slippy surface - stops more slowly
 		float StoppingDistanceLowFriction = 0.f;
 		Wheel.SetSurfaceFriction(0.3f);
 		Wheel.SetBrakeTorque(650);
-		VehicleSpeedMPH = 30.f;
-		SimulateBraking(Wheel, VehicleSpeedMPH, DeltaTime, StoppingDistanceLowFriction, SimulationTime);
+		SimulateBraking(Wheel, Gravity, MPHToMS(30.f), DeltaTime, StoppingDistanceLowFriction, SimulationTime);
 		EXPECT_GT(StoppingDistanceLowFriction, StoppingDistanceA);
 	}
 
@@ -1054,7 +1075,7 @@ namespace ChaosTest
 		Setup.TractionControlEnabled = false;
 		Setup.BrakeEnabled = true;
 		Setup.EngineEnabled = true;
-		Setup.WheelRadius = 30.0f;
+		Setup.WheelRadius = 0.3f;
 		Setup.NewSimulationPath = true;
 		Setup.FrictionMultiplier = 1.0f;
 		Setup.CorneringStiffness = 2.0f;
@@ -1077,13 +1098,13 @@ namespace ChaosTest
 
 		// How far & what time does it take to stop from 30MPH to rest
 		Wheel.SetBrakeTorque(650);
-		SimulateBraking(Wheel, 30.0f, DeltaTime, StoppingDistanceA, SimulationTimeBrake);
+		SimulateBraking(Wheel, Gravity, MPHToMS(30.0f), DeltaTime, StoppingDistanceA, SimulationTimeBrake);
 
 		// How far and what time does it take to accelerate from rest to 30MPH
 		float SimulationTimeAccel = 0.0f;
 		float DrivingDistanceA = 0.f;
 		Wheel.SetDriveTorque(650);
-		SimulateAccelerating(Wheel, Gravity, 30.0f, DeltaTime, DrivingDistanceA, SimulationTimeAccel);
+		SimulateAccelerating(Wheel, Gravity, MPHToMS(30.0f), DeltaTime, DrivingDistanceA, SimulationTimeAccel);
 
 		// 0-30 MPH and 30-0 MPH should be the same if there's no slipping and accel torque was same as the brake torque run
 		EXPECT_LT(DrivingDistanceA - StoppingDistanceA, AccelerationResultsTolerance);
@@ -1095,16 +1116,18 @@ namespace ChaosTest
 
 		// Unreal units cm - Note for the same results the radius needs to remain at 0.3m and not also be scaled to 30(cm)
 		float SimulationTimeAccelCM = 0.0f;
-		float MtoCm = 100.0f;
+		float MToCm = 100.0f;
 		float DrivingDistanceCM = 0.f;
-		Wheel.SetDriveTorque(650 * MtoCm);
-		SimulateAccelerating(Wheel, Gravity * MtoCm, 30.0f * MtoCm, DeltaTime, DrivingDistanceCM, SimulationTimeAccelCM);
-		EXPECT_GT(DrivingDistanceCM, 10.f * MtoCm);
-		EXPECT_LT(DrivingDistanceCM, 20.f * MtoCm);
-		EXPECT_LT(SimulationTimeAccel - SimulationTimeAccelCM, AccelerationResultsTolerance);
+		Wheel.SetDriveTorque(650.0f * MToCm * MToCm);
+		Wheel.SetWheelRadius(0.3f * MToCm);
+		SimulateAccelerating(Wheel, Gravity * MToCm, MPHToCmS(30.0f), DeltaTime, DrivingDistanceCM, SimulationTimeAccelCM);
+		EXPECT_GT(DrivingDistanceCM, 10.f * MToCm);
+		EXPECT_LT(DrivingDistanceCM, 20.f * MToCm);
+		EXPECT_NEAR(DrivingDistanceCM, DrivingDistanceA * MToCm, AccelerationResultsTolerance);
 
 		float SimulationTimeAccelSpin = 0.0f;
 		float DrivingDistanceWheelspin = 0.f;
+		Wheel.SetWheelRadius(0.3f);
 		Wheel.SetDriveTorque(5000); // definitely cause wheel spin
 		SimulateAccelerating(Wheel, Gravity, 30.0f, DeltaTime, DrivingDistanceWheelspin, SimulationTimeAccelSpin);
 
@@ -1119,7 +1142,7 @@ namespace ChaosTest
 		float DrivingDistanceTC = 0.f;
 		Wheel.AccessSetup().TractionControlEnabled = true;
 		Wheel.SetDriveTorque(5000); // definitely cause wheel spin
-		SimulateAccelerating(Wheel, Gravity, 30.0f, DeltaTime, DrivingDistanceTC, SimulationTimeAccelTC);
+		SimulateAccelerating(Wheel, Gravity, MPHToMS(30.0f), DeltaTime, DrivingDistanceTC, SimulationTimeAccelTC);
 
 		// reaches target speed in a shorter distance
 		EXPECT_LT(DrivingDistanceTC, DrivingDistanceWheelspin);
@@ -1129,7 +1152,6 @@ namespace ChaosTest
 
 	}
 
-
 	TYPED_TEST(AllTraits, DISABLED_VehicleTest_WheelLateralSlip)
 	{
 		FSimpleWheelConfig Setup;
diff --git a/Engine/Source/Runtime/Experimental/ChaosVehicles/ChaosVehiclesCore/Private/WheelSystem.cpp b/Engine/Source/Runtime/Experimental/ChaosVehicles/ChaosVehiclesCore/Private/WheelSystem.cpp
--- a/Engine/Source/Runtime/Experimental/ChaosVehicles/ChaosVehiclesCore/Private/WheelSystem.cpp
+++ b/Engine/Source/Runtime/Experimental/ChaosVehicles/ChaosVehiclesCore/Private/WheelSystem.cpp
@@ -9,7 +9,7 @@ PRAGMA_DISABLE_OPTIMIZATION
 namespace Chaos
 {
 	FSimpleWheelSim::FSimpleWheelSim(const FSimpleWheelConfig* SetupIn) : TVehicleSystem<FSimpleWheelConfig>(SetupIn)
-		, Re(30.f)
+		, Re(SetupIn->WheelRadius)
 		, Omega(0.f)
 		, Sx(0.f)
 		, DriveTorque(0.f)
@@ -45,8 +45,8 @@ namespace Chaos
 		// The physics system is mostly unit-less i.e. can work in meters or cm, however there are 
 		// a couple of places where the results are wrong if Cm is used. This is one of them, the simulated radius
 		// for torque must be real size to obtain the correct output values.
-		AppliedLinearDriveForce = DriveTorque / CmToM(Re);
-		AppliedLinearBrakeForce = BrakeTorque / CmToM(Re);
+		AppliedLinearDriveForce = DriveTorque / Re;
+		AppliedLinearBrakeForce = BrakeTorque / Re;
 
 		// currently just letting the brake override the throttle
 		bool Braking = BrakeTorque > FMath::Abs(DriveTorque);
@@ -69,10 +69,7 @@ namespace Chaos
 
 				// check we are not applying more force than required so we end up overshooting 
 				// and accelerating in the opposite direction
-				if (FinalLongitudinalForce > FMath::Abs(ForceRequiredToBringToStop))
-				{
-					FinalLongitudinalForce = FMath::Abs(ForceRequiredToBringToStop);
-				}
+				FinalLongitudinalForce = FMath::Clamp(FinalLongitudinalForce, -FMath::Abs(ForceRequiredToBringToStop), FMath::Abs(ForceRequiredToBringToStop));
 
 				// ensure the brake opposes current direction of travel
 				if (GroundVelocityVector.X > 0.0f)
@@ -109,7 +106,9 @@ namespace Chaos
 				if ((Braking && Setup().ABSEnabled) || (!Braking && Setup().TractionControlEnabled))
 				{
 					Spin = 0.0f;
-					ForceFromFriction.X = LongitudinalAdhesiveLimit * TractionControlAndABSScaling;
+
+					float Sign = (FinalLongitudinalForce > 0.0f) ? 1.0f : -1.0f;
+					ForceFromFriction.X = LongitudinalAdhesiveLimit * TractionControlAndABSScaling * Sign;
 				}
 				else
 				{
@@ -123,7 +122,9 @@ namespace Chaos
 					{
 						Locked = true;
 					}
-					ForceFromFriction.X = LongitudinalAdhesiveLimit * DynamicFrictionLongitudialScaling;
+
+					float Sign = (ForceFromFriction.X >= 0.0f) ? 1.0 : -1.0f;
+					ForceFromFriction.X = LongitudinalAdhesiveLimit * DynamicFrictionLongitudialScaling * Sign;
 				}
 			}
 			else
@@ -131,11 +132,6 @@ namespace Chaos
 				Spin = 0.0f;
 			}
 
-			if (FinalLongitudinalForce < -LongitudinalAdhesiveLimit)
-			{
-				ForceFromFriction.X = -ForceFromFriction.X;
-			}
-
 			static float DynamicFrictionLateralScaling = 0.75f;
 			if (Locked || Spinning)
 			{
@@ -203,8 +199,8 @@ namespace Chaos
 		// The physics system is mostly unit-less i.e. can work in meters or cm, however there are 
 		// a couple of places where the results are wrong if Cm is used. This is one of them, the simulated radius
 		// for torque must be real size to obtain the correct output values.
-		AppliedLinearDriveForce = DriveTorque / CmToM(Re);
-		AppliedLinearBrakeForce = BrakeTorque / CmToM(Re);
+		AppliedLinearDriveForce = DriveTorque / Re;
+		AppliedLinearBrakeForce = BrakeTorque / Re;
 
 		// Longitudinal multiplier now affecting both brake and steering equally
 		float AvailableGrip = ForceIntoSurface * SurfaceFriction * Setup().FrictionMultiplier;
@@ -214,6 +210,7 @@ namespace Chaos
 
 		// currently just letting the brake override the throttle
 		bool Braking = BrakeTorque > FMath::Abs(DriveTorque);
+		bool WheelLocked = false;
 
 		// are we actually touching the ground
 		if (ForceIntoSurface > SMALL_NUMBER)
@@ -223,15 +220,16 @@ namespace Chaos
 			{
 				if ((Braking && Setup().ABSEnabled) || (!Braking && Setup().TractionControlEnabled))
 				{
-					//	Spin = 0.0f;
-					AppliedLinearBrakeForce = AvailableGrip * TractionControlAndABSScaling;
+					float Sign = (AppliedLinearBrakeForce > 0.0f) ? 1.0f : -1.0f;
+					AppliedLinearBrakeForce = AvailableGrip * TractionControlAndABSScaling * Sign;
 				}
 			}
 
 			// Traction control limiting drive force to match force from grip available
 			if (Setup().TractionControlEnabled && !Braking && FMath::Abs(AppliedLinearDriveForce) > AvailableGrip)
 			{
-				AppliedLinearDriveForce = AvailableGrip * TractionControlAndABSScaling;
+				float Sign = (AppliedLinearDriveForce > 0.0f) ? 1.0f : -1.0f;
+				AppliedLinearDriveForce = AvailableGrip * TractionControlAndABSScaling * Sign;
 			}
 
 			if (Braking)
@@ -243,10 +241,7 @@ namespace Chaos
 
 				// check we are not applying more force than required so we end up overshooting 
 				// and accelerating in the opposite direction
-				if (FinalLongitudinalForce > FMath::Abs(ForceRequiredToBringToStop))
-				{
-					FinalLongitudinalForce = FMath::Abs(ForceRequiredToBringToStop);
-				}
+				FinalLongitudinalForce = FMath::Clamp(FinalLongitudinalForce, -FMath::Abs(ForceRequiredToBringToStop), FMath::Abs(ForceRequiredToBringToStop));
 
 				// ensure the brake opposes current direction of travel
 				if (GroundVelocityVector.X > 0.0f)
@@ -300,12 +295,24 @@ namespace Chaos
 					FinalLongitudinalForce *= Setup().SideSlipModifier;
 					FinalLateralForce *= Setup().SideSlipModifier;
 
+					if (Braking)
+					{
+						WheelLocked = true;
+					}
+
 				}
 			}
 		}
 
-		float GroundOmega = GroundVelocityVector.X / Re;
-		Omega += (GroundOmega - Omega);
+		if (WheelLocked)
+		{
+			Omega = 0.0f;
+		}
+		else
+		{ 
+			float GroundOmega = GroundVelocityVector.X / Re;
+			Omega += (GroundOmega - Omega);
+		}
 
 		// Wheel angular position
 		AngularPosition += Omega * DeltaTime;
diff --git a/Engine/Source/Runtime/Experimental/ChaosVehicles/ChaosVehiclesCore/Public/VehicleUtility.h b/Engine/Source/Runtime/Experimental/ChaosVehicles/ChaosVehiclesCore/Public/VehicleUtility.h
--- a/Engine/Source/Runtime/Experimental/ChaosVehicles/ChaosVehiclesCore/Public/VehicleUtility.h
+++ b/Engine/Source/Runtime/Experimental/ChaosVehicles/ChaosVehiclesCore/Public/VehicleUtility.h
@@ -254,6 +254,11 @@ namespace Chaos
 		return Val * Val;
 	}
 
+	FORCEINLINE float TorqueMToCm(float TorqueIn)
+	{
+		return TorqueIn * 10000.0f;
+	}
+
 	class CHAOSVEHICLESCORE_API FTimeAndDistanceMeasure
 	{
 	public:
