Chaos - improve GJK for large shapes
- avoid recalculation of nearest point oin simplex from barycentric coords when we already have a more accurate value.
- enable large object unit tests (that now pass)
- disable old unit test that is no longer valid

diff --git a/Engine/Source/Programs/HeadlessChaos/Private/HeadlessChaosTestGJK.cpp b/Engine/Source/Programs/HeadlessChaos/Private/HeadlessChaosTestGJK.cpp
--- a/Engine/Source/Programs/HeadlessChaos/Private/HeadlessChaosTestGJK.cpp
+++ b/Engine/Source/Programs/HeadlessChaos/Private/HeadlessChaosTestGJK.cpp
@@ -1374,7 +1374,14 @@ namespace ChaosTest
 		}
 	}
 
-	void GJKCapsuleConvexInitialOverlapSweep()
+	// When we have a capsule and box that are reported as initially-overlapping because they are within
+	// the GJK epsilon opf each other (but actually positively separated), verify that we get a zero time of impact.
+	// Previously the slihjtly-positive separation would result in a negative penetration and a positiove TOI.
+	// Bug fix: CL 10942094.
+	// NOTE: this issue no longer manifests with this example because GJK no longer reports this case as
+	// overlapping> The GJK epsilon no longer takes part in the distance calculation when the near point
+	// is on the face of the convex.
+	GTEST_TEST(GJKTests, DISABLED_TestGJKCapsuleConvexInitialOverlapSweep_Fixed)
 	{
 		{
 			TArray<Chaos::FVec3> ConvexParticles;
@@ -1409,7 +1416,10 @@ namespace ChaosTest
 			EXPECT_TRUE(GJKRaycast2<FReal>(A,B,BToATM,LocalDir,Length,Time,Position,Normal,0,true,SearchDir,0));
 			EXPECT_FLOAT_EQ(Time,0.0);
 		}
+	}
 
+	void GJKCapsuleConvexInitialOverlapSweep()
+	{
 		{
 			TArray<FVec3> ConvexParticles;
 			ConvexParticles.SetNum(16);
diff --git a/Engine/Source/Programs/HeadlessChaos/Private/HeadlessChaosTestLargeScale.cpp b/Engine/Source/Programs/HeadlessChaos/Private/HeadlessChaosTestLargeScale.cpp
--- a/Engine/Source/Programs/HeadlessChaos/Private/HeadlessChaosTestLargeScale.cpp
+++ b/Engine/Source/Programs/HeadlessChaos/Private/HeadlessChaosTestLargeScale.cpp
@@ -17,7 +17,9 @@ namespace ChaosTest {
 	using namespace Chaos;
 
 	// Test GJKPenetration for a normal sized Capsule against a large scaled Box.
-	GTEST_TEST(LargeScaleTests, DISABLED_TestSmallCapsuleLargeBoxGJKPenetration)
+	// We used to have a large error in GJK when testing against large shapes. This was fixed
+	// in Simplex.h in the same CL that this line was added.
+	GTEST_TEST(LargeScaleTests, TestSmallCapsuleLargeBoxGJKPenetration)
 	{
 		// Shape settings
 		const FVec3 BoxSize = FVec3(100, 100, 100);
@@ -49,12 +51,14 @@ namespace ChaosTest {
 		EXPECT_TRUE(bResult);
 		EXPECT_NEAR(-Penetration, Separation, MaxPositionError);
 		EXPECT_NEAR(ClosestA.Z, 0.5f * (BoxSize * BoxScale).Z, MaxPositionError);			// Box space
-		EXPECT_NEAR(ClosestBInA.Z, (BoxSize * BoxScale).Z + Separation, MaxPositionError);
+		EXPECT_NEAR(ClosestBInA.Z, 0.5f * (BoxSize * BoxScale).Z + Separation, MaxPositionError);
 	}
 
 	// A large box centered on X,Y and positioned so that the top surface is at Z = 0.
 	// A small capsule positioned according to the parameter CapsulePos.
 	// Verify that a downward sweep with the specified initial direction returns an accurate result.
+	// We used to have a large error in GJK when testing against large shapes. This was fixed
+	// in Simplex.h in the same CL that this line was added.
 	void SmallCapsuleLargeBoxGJKRaycast(const FVec3 &CapsuleBottomPos, const FVec3& InitialDir)
 	{
 		// Shape settings
@@ -95,9 +99,9 @@ namespace ChaosTest {
 		EXPECT_NEAR(Normal.Z, 1.0f, MaxNormalError);
 	}
 
-	GTEST_TEST(LargeScaleTests, DISABLED_TestSmallCapsuleLargeBoxGJKRaycast_Vertical)
+	GTEST_TEST(LargeScaleTests, TestSmallCapsuleLargeBoxGJKRaycast_Vertical)
 	{
-		// Place capsule abovt the plane and shited laterally
+		// Place capsule above the plane and shited laterally
 		SmallCapsuleLargeBoxGJKRaycast(FVec3(3, 5, 20), FVec3(0, 0, 1));
 		SmallCapsuleLargeBoxGJKRaycast(FVec3(100, 50, 20), FVec3(0, 0, 1));
 		SmallCapsuleLargeBoxGJKRaycast(FVec3(300, 50, 20), FVec3(0, 0, 1));
@@ -108,7 +112,9 @@ namespace ChaosTest {
 
 
 	// Test GJKRaycast for a normal sized Capsule against a large scaled Box.
-	GTEST_TEST(LargeScaleTests, DISABLED_TestSmallCapsuleLargeConvexGJKRaycast)
+	// We used to have a large error in GJK when testing against large shapes. This was fixed
+	// in Simplex.h in the same CL that this line was added.
+	GTEST_TEST(LargeScaleTests, TestSmallCapsuleLargeConvexGJKRaycast)
 	{
 		// Shape settings
 		const FVec3 BoxSize = FVec3(100, 100, 100);
@@ -143,8 +149,8 @@ namespace ChaosTest {
 		const FReal MaxPositionError = 1.e-3f;
 		const FReal MaxNormalError = 1.e-4f;
 		EXPECT_TRUE(bResult);
-		EXPECT_NEAR(Time, 0.5f, MaxTimeError);
-		EXPECT_NEAR(Position.Z, (BoxSize * BoxScale).Z, MaxPositionError);
+		EXPECT_NEAR(Time, Separation, MaxTimeError);
+		EXPECT_NEAR(Position.Z, 0.5f * (BoxSize * BoxScale).Z, MaxPositionError);
 		EXPECT_NEAR(Normal.Z, 1.0f, MaxNormalError);
 	}
 
diff --git a/Engine/Source/Runtime/Experimental/Chaos/Public/Chaos/Simplex.h b/Engine/Source/Runtime/Experimental/Chaos/Public/Chaos/Simplex.h
--- a/Engine/Source/Runtime/Experimental/Chaos/Public/Chaos/Simplex.h
+++ b/Engine/Source/Runtime/Experimental/Chaos/Public/Chaos/Simplex.h
@@ -13,8 +13,8 @@ namespace Chaos
 
 		//Closest Point = (-X0 dot X1-X0) / ||(X1-X0)||^2 * (X1-X0)
 
-		const TVector<T, d> OriginToX0 = -X0;
-		const T Dot = TVector<T, d>::DotProduct(OriginToX0, X0ToX1);
+		const TVector<T, d> X0ToOrigin = -X0;
+		const T Dot = TVector<T, d>::DotProduct(X0ToOrigin, X0ToX1);
 
 		if (Dot <= 0)
 		{
@@ -33,7 +33,7 @@ namespace Chaos
 			return X1;
 		}
 
-		const T Ratio = Dot / X0ToX1Squared;
+		const T Ratio = FMath::Clamp(Dot / X0ToX1Squared, T(0), T(1));
 		const TVector<T, d> Closest = Ratio * (X0ToX1)+X0;	//note: this could pass X1 by machine epsilon, but doesn't seem worth it for now
 		OutBarycentric[Idxs[0]] = 1 - Ratio;
 		OutBarycentric[Idxs[1]] = Ratio;
@@ -206,11 +206,16 @@ namespace Chaos
 			T Lambda1 = Cofactors[1] * InvDetM;
 			//T Lambda2 = 1 - Lambda1 - Lambda0;
 			T Lambda2 = Cofactors[2] * InvDetM;
-			const TVec3<T> ClosestPoint = X0 * Lambda0 + X1 * Lambda1 + X2 * Lambda2;	//could be slightly outside if |lambda1| < 1e-7 or |lambda2| < 1e-7. Should we clamp?
 			OutBarycentric[Idx0] = Lambda0;
 			OutBarycentric[Idx1] = Lambda1;
 			OutBarycentric[Idx2] = Lambda2;
-			return ClosestPoint;
+
+			// We know that we are inside the triangle so we can use the projected point we calculated above. 
+			// The closest point can also be derived from the barycentric coordinates, but it will contain 
+			// numerical error from the determinant calculation  and can cause GJK to terminate with a poor solution.
+			// (E.g., this caused jittering when walking on box with dimensions of 100000cm or more).
+			// return X0 * Lambda0 + X1 * Lambda1 + X2 * Lambda2;
+			return ProjectedOrigin;
 		}
 		else
 		{
