ChaosVehicleTests - added delta time variation tests

diff --git a/Engine/Source/Programs/HeadlessChaos/Private/HeadlessChaosTestVehicles.cpp b/Engine/Source/Programs/HeadlessChaos/Private/HeadlessChaosTestVehicles.cpp
--- a/Engine/Source/Programs/HeadlessChaos/Private/HeadlessChaosTestVehicles.cpp
+++ b/Engine/Source/Programs/HeadlessChaos/Private/HeadlessChaosTestVehicles.cpp
@@ -870,12 +870,12 @@ namespace ChaosTest
 		EXPECT_GT(StoppingDistanceTooHeavyBreaking, StoppingDistanceA);
 
 		// Would have locked the wheels but ABS prevents skidding
-		Wheel.AccessSetup().ABSEnabled = true;
+		Wheel.ABSEnabled = true;
 		float StoppingDistanceTooHeavyBreakingABS = 0.f;
 		Wheel.SetBrakeTorque(5000);
 		SimulateBraking(Wheel, Gravity, MPHToMS(30.f), DeltaTime, StoppingDistanceTooHeavyBreakingABS, SimulationTime);
 		EXPECT_LT(StoppingDistanceTooHeavyBreakingABS, StoppingDistanceTooHeavyBreaking);
-		Wheel.AccessSetup().ABSEnabled = false;
+		Wheel.ABSEnabled = false;
 
 		// lower initial speed - stops more quickly
 		float StoppingDistanceLowerSpeed = 0.f;
@@ -962,7 +962,7 @@ namespace ChaosTest
 		// Enable traction control should be better than both of the above
 		float SimulationTimeAccelTC = 0.0f;
 		float DrivingDistanceTC = 0.f;
-		Wheel.AccessSetup().TractionControlEnabled = true;
+		Wheel.TractionControlEnabled = true;
 		Wheel.SetDriveTorque(5000); // definitely cause wheel spin
 		SimulateAccelerating(Wheel, Gravity, MPHToMS(30.0f), DeltaTime, DrivingDistanceTC, SimulationTimeAccelTC);
 
@@ -1041,12 +1041,12 @@ namespace ChaosTest
 		EXPECT_GT(StoppingDistanceTooHeavyBreaking, StoppingDistanceA);
 
 		// Would have locked the wheels but ABS prevents skidding
-		Wheel.AccessSetup().ABSEnabled = true;
+		Wheel.ABSEnabled = true;
 		float StoppingDistanceTooHeavyBreakingABS = 0.f;
 		Wheel.SetBrakeTorque(5000);
 		SimulateBraking(Wheel, Gravity, MPHToMS(30.f), DeltaTime, StoppingDistanceTooHeavyBreakingABS, SimulationTime);
 		EXPECT_LT(StoppingDistanceTooHeavyBreakingABS, StoppingDistanceTooHeavyBreaking);
-		Wheel.AccessSetup().ABSEnabled = false;
+		Wheel.ABSEnabled = false;
 
 		// lower initial speed - stops more quickly
 		float StoppingDistanceLowerSpeed = 0.f;
@@ -1140,7 +1140,7 @@ namespace ChaosTest
 		// Enable traction control should be better than both of the above
 		float SimulationTimeAccelTC = 0.0f;
 		float DrivingDistanceTC = 0.f;
-		Wheel.AccessSetup().TractionControlEnabled = true;
+		Wheel.TractionControlEnabled = true;
 		Wheel.SetDriveTorque(5000); // definitely cause wheel spin
 		SimulateAccelerating(Wheel, Gravity, MPHToMS(30.0f), DeltaTime, DrivingDistanceTC, SimulationTimeAccelTC);
 
@@ -1591,16 +1591,134 @@ namespace ChaosTest
 		EXPECT_LT(Dynamic->X().Z - ExpectedRestingPosition, Tolerance);
 	}
 
-	//TYPED_TEST(AllTraits, VehicleTest_SuspensionNaturalFrequency)
-	//{
-	//	float SprungMass = 250.f;
-	//	float SpringStiffness = 250.0f;
-	//	float SpringDamping = 0.5f;
+	TYPED_TEST(AllTraits, VehicleTest_WheelAcceleratingLongitudinalSlip_VaryingDelta)
+	{
+		FSimpleWheelConfig Setup;
+		Setup.ABSEnabled = false;
+		Setup.TractionControlEnabled = false;
+		Setup.BrakeEnabled = true;
+		Setup.EngineEnabled = true;
+		Setup.WheelRadius = 0.3f;
+		Setup.LongitudinalFrictionMultiplier = 1.0f;
+		Setup.LateralFrictionMultiplier = 1.0f;
+		Setup.SideSlipModifier = 0.7f;
+
+		FSimpleWheelSim Wheel(&Setup);
+
+		// units meters
+		float Gravity = 9.8f;
+		float VehicleMassPerWheel = 100.f;
+		float ResultsTolerance = 0.01f;
+
+		Wheel.SetSurfaceFriction(RealWorldConsts::DryRoadFriction());
+		Wheel.SetWheelLoadForce(VehicleMassPerWheel * Gravity);
+		Wheel.SetMassPerWheel(VehicleMassPerWheel);
+
+		// Road speed
+		FVector Velocity = FVector(10.f, 0.f, 0.f);
+
+		// start from stationary
+		Wheel.SetMatchingSpeed(Velocity.X);
+		Wheel.SetVehicleGroundSpeed(Velocity);
+		Wheel.SetDriveTorque(100.0f);
+
+		float DeltaTime = 1.f / 30.f;
+		Wheel.Simulate(DeltaTime);
+		FVector ForceGenerated30FPS = Wheel.GetForceFromFriction();
+
+		Wheel.Simulate(DeltaTime);
+		FVector ForceGenerated30FPS_2 = Wheel.GetForceFromFriction();
+
+		DeltaTime = 1.f / 60.f;
+		Wheel.Simulate(DeltaTime);
+
+		FVector ForceGenerated60FPS = Wheel.GetForceFromFriction();
+
+		DeltaTime = 1.f / 50.f;
+		Wheel.Simulate(DeltaTime);
 
+		FVector ForceGenerated50FPS = Wheel.GetForceFromFriction();
 
+		EXPECT_NEAR(ForceGenerated30FPS.X, ForceGenerated30FPS_2.X, ResultsTolerance);
+		EXPECT_NEAR(ForceGenerated30FPS.X, ForceGenerated60FPS.X, ResultsTolerance);
+		EXPECT_NEAR(ForceGenerated30FPS.X, ForceGenerated60FPS.X, ResultsTolerance);
 
-	//}
 
+		Wheel.AccessSetup().NewSimulationPath = true;
+
+		Wheel.SetMatchingSpeed(Velocity.X);
+		Wheel.SetVehicleGroundSpeed(Velocity);
+		Wheel.SetDriveTorque(100.0f);
+
+		DeltaTime = 1.f / 30.f;
+		Wheel.Simulate(DeltaTime);
+		FVector NForceGenerated30FPS = Wheel.GetForceFromFriction();
+
+		Wheel.Simulate(DeltaTime);
+		FVector NForceGenerated30FPS_2 = Wheel.GetForceFromFriction();
+
+		DeltaTime = 1.f / 60.f;
+		Wheel.Simulate(DeltaTime);
+
+		FVector NForceGenerated60FPS = Wheel.GetForceFromFriction();
+
+		DeltaTime = 1.f / 50.f;
+		Wheel.Simulate(DeltaTime);
+
+		FVector NForceGenerated50FPS = Wheel.GetForceFromFriction();
+
+		EXPECT_NEAR(NForceGenerated30FPS.X, NForceGenerated30FPS_2.X, ResultsTolerance);
+		EXPECT_NEAR(NForceGenerated30FPS.X, NForceGenerated60FPS.X, ResultsTolerance);
+		EXPECT_NEAR(NForceGenerated30FPS.X, NForceGenerated50FPS.X, ResultsTolerance);
+
+	}
+
+	TYPED_TEST(AllTraits, VehicleTest_Suspension_VaryingDelta)
+	{
+		FSimpleSuspensionConfig Setup;
+		Setup.MaxLength = 20.0f;
+		Setup.DampingRatio = 0.5f;
+		Setup.SpringPreload = 100.0f;
+		Setup.SpringRate = 100.0f;
+		Setup.SuspensionAxis = FVector(0.f, 0.f, -1.f);
+		Setup.RestingForce = 50.0f;
+
+		FSimpleSuspensionSim Suspension(&Setup);
+
+		// units meters
+		float Gravity = 9.8f;
+		float VehicleMassPerWheel = 100.f;
+		float WheelRadius = 32.0f;
+		float CurrentLength = 5.0f;
+
+		float DeltaTime = 1.0f / 30.0f;
+		float ResultsTolerance = 0.01f;
+
+		Suspension.SetSuspensionLength(CurrentLength, WheelRadius);
+		Suspension.SetLocalVelocity(FVector(0.f,0.f,-8.f));
+
+		// first one has to adjust for sudden change in length
+		Suspension.Simulate(DeltaTime);
+
+		Suspension.Simulate(DeltaTime);
+		float Force30FPS = Suspension.GetSuspensionForce();
+
+		Suspension.Simulate(DeltaTime);
+		float Force30FPS_2 = Suspension.GetSuspensionForce();
+
+		DeltaTime = 60.0f;
+		Suspension.Simulate(DeltaTime);
+		float Force60FPS = Suspension.GetSuspensionForce();
+
+		DeltaTime = 50.0f; 
+		Suspension.Simulate(DeltaTime);
+		float Force50FPS = Suspension.GetSuspensionForce();
+
+		EXPECT_NEAR(Force30FPS, Force30FPS_2, ResultsTolerance);
+		EXPECT_NEAR(Force30FPS, Force60FPS, ResultsTolerance);
+		EXPECT_NEAR(Force30FPS, Force50FPS, ResultsTolerance);
+
+	}
 }
 
 
