CCD improvements
- CCD can now run every iteration (solves many CCD failures)
- CCD won't cull contacts
- Old bug where rotating bodies can tunnel is fixed now
- Added an extra fail safe in case of too few iterations

diff --git a/Engine/Source/Programs/HeadlessChaos/Private/HeadlessChaosTestCCD.cpp b/Engine/Source/Programs/HeadlessChaos/Private/HeadlessChaosTestCCD.cpp
new file mode 100644
--- /dev/null
+++ b/Engine/Source/Programs/HeadlessChaos/Private/HeadlessChaosTestCCD.cpp
@@ -0,0 +1,306 @@
+// Copyright Epic Games, Inc. All Rights Reserved.
+
+#include "Chaos/Matrix.h"
+#include "Chaos/Utilities.h"
+#include "Chaos/AABB.h"
+#include "Chaos/Core.h"
+#include "Chaos/PBDRigidsEvolutionGBF.h"
+#include "Chaos/Box.h"
+#include "Chaos/Sphere.h"
+#include "ChaosSolversModule.h"
+#include "HeadlessChaos.h"
+#include "HeadlessChaosTestUtility.h"
+#include "Modules/ModuleManager.h"
+#include "PBDRigidsSolver.h"
+
+//PRAGMA_DISABLE_OPTIMIZATION
+
+namespace ChaosTest
+{
+	using namespace Chaos;
+
+	// CCD not implemented for sphere sphere
+	TYPED_TEST(AllEvolutions, DISABLED_CCDTests_CCDEnabled)
+	{
+		const FReal Dt = 1 / 30.0f;
+		const FReal Fps = 1 / Dt;
+		const FReal SphereRadius = 100; // cm
+		const FReal InitialSpeed = SphereRadius * 5 * Fps; // More than enough to tunnel
+		using TEvolution = TypeParam;
+		TPBDRigidsSOAs<FReal, 3> Particles;
+		THandleArray<FChaosPhysicsMaterial> PhysicalMaterials;
+		TEvolution Evolution(Particles, PhysicalMaterials);
+		InitEvolutionSettings(Evolution);
+
+		// Create particles
+		TGeometryParticleHandle<FReal, 3>* Static = Evolution.CreateStaticParticles(1)[0];
+		TPBDRigidParticleHandle<FReal, 3>* Dynamic = Evolution.CreateDynamicParticles(1)[0];
+
+		// Set up physics material
+		TUniquePtr<FChaosPhysicsMaterial> PhysicsMaterial = MakeUnique<FChaosPhysicsMaterial>();
+		PhysicsMaterial->SleepCounterThreshold = 1000; // Don't sleep
+		PhysicsMaterial->Restitution = 1.0f;
+
+		// Create Sphere geometry (Radius = 100)
+		TUniquePtr<FImplicitObject> Sphere(new TSphere<FReal, 3>(FVec3(0, 0, 0), SphereRadius));
+
+		// Assign sphere geometry to both particles 
+		Static->SetGeometry(MakeSerializable(Sphere));
+		Dynamic->SetGeometry(MakeSerializable(Sphere));
+
+		Evolution.SetPhysicsMaterial(Dynamic, MakeSerializable(PhysicsMaterial));
+
+		const FReal Mass = 100000.0f;;
+		Dynamic->I() = FMatrix33(Mass, Mass, Mass);
+		Dynamic->InvI() = FMatrix33(1.0f / Mass, 1.0f / Mass, 1.0f / Mass);
+
+		// Positions and velocities
+		Static->X() = FVec3(0, 0, 0);
+
+		Dynamic->X() = FVec3(0, 0, SphereRadius * 2 + 10);
+		//Dynamic->V() = FVec3(0, 0, -InitialSpeed);
+		
+		// The position of the static has changed and statics don't automatically update bounds, so update explicitly
+		Static->SetWorldSpaceInflatedBounds(Sphere->BoundingBox().TransformedAABB(TRigidTransform<FReal, 3>(Static->X(), Static->R())));
+
+		::ChaosTest::SetParticleSimDataToCollide({ Static,Dynamic });
+
+		Dynamic->SetCCDEnabled(true);
+		Dynamic->SetGravityEnabled(false);
+		
+
+		Dynamic->V() = FVec3(0, 0, -InitialSpeed);
+
+		for (int i = 0; i < 1; ++i)
+		{
+			Evolution.AdvanceOneTimeStep(Dt);
+			Evolution.EndFrame(Dt);
+		}
+
+		// Large error margin, we are testing CCD and not solver accuracy
+		const FReal LargeErrorMargin = 10.0f;
+		EXPECT_GE(Dynamic->X()[2], SphereRadius * 2 - LargeErrorMargin);
+	}
+
+
+
+	TYPED_TEST(AllEvolutions,DISABLED_CCDTests_CCDDisabled)
+	{
+		const FReal Dt = 1 / 30.0f;
+		const FReal Fps = 1 / Dt;
+		const FReal SphereRadius = 100; // cm
+		const FReal InitialSpeed = SphereRadius * 10 * Fps; // More than enough to tunnel
+		using TEvolution = TypeParam;
+		TPBDRigidsSOAs<FReal, 3> Particles;
+		THandleArray<FChaosPhysicsMaterial> PhysicalMaterials;
+		TEvolution Evolution(Particles, PhysicalMaterials);
+		InitEvolutionSettings(Evolution);
+
+		// Create particles
+		TGeometryParticleHandle<FReal, 3>* Static = Evolution.CreateStaticParticles(1)[0];
+		TPBDRigidParticleHandle<FReal, 3>* Dynamic = Evolution.CreateDynamicParticles(1)[0];
+
+		// Set up physics material
+		TUniquePtr<FChaosPhysicsMaterial> PhysicsMaterial = MakeUnique<FChaosPhysicsMaterial>();
+		PhysicsMaterial->SleepCounterThreshold = 1000; // Don't sleep
+		PhysicsMaterial->Restitution = 1.0f;
+
+		// Create Sphere geometry (Radius = 100)
+		TUniquePtr<FImplicitObject> Sphere(new TSphere<FReal, 3>(FVec3(0, 0, 0), SphereRadius));
+
+		// Assign sphere geometry to both particles 
+		Static->SetGeometry(MakeSerializable(Sphere));
+		Dynamic->SetGeometry(MakeSerializable(Sphere));
+
+		Evolution.SetPhysicsMaterial(Dynamic, MakeSerializable(PhysicsMaterial));
+
+		const FReal Mass = 100000.0f;;
+		Dynamic->I() = FMatrix33(Mass, Mass, Mass);
+		Dynamic->InvI() = FMatrix33(1.0f / Mass, 1.0f / Mass, 1.0f / Mass);
+
+		// Positions and velocities
+		Static->X() = FVec3(0, 0, 0);
+
+		Dynamic->X() = FVec3(0, 0, SphereRadius * 2 + 100);
+
+		// The position of the static has changed and statics don't automatically update bounds, so update explicitly
+		Static->SetWorldSpaceInflatedBounds(Sphere->BoundingBox().TransformedAABB(TRigidTransform<FReal, 3>(Static->X(), Static->R())));
+
+		::ChaosTest::SetParticleSimDataToCollide({ Static,Dynamic });
+
+		Dynamic->SetCCDEnabled(false);
+		Dynamic->SetGravityEnabled(false);
+
+		Dynamic->V() = FVec3(0, -InitialSpeed, 0);
+
+		for (int i = 0; i < 1; ++i)
+		{
+			Evolution.AdvanceOneTimeStep(Dt);
+			Evolution.EndFrame(Dt);
+		}
+
+		// Large error margin, we are testing CCD and not solver accuracy
+		const FReal LargeErrorMargin = 10.0f;
+		EXPECT_NEAR(Dynamic->X()[2], SphereRadius * 2 + 10 - InitialSpeed * Dt, LargeErrorMargin);
+	}
+
+	
+	TYPED_TEST(AllEvolutions, DISABLED_CCDTests_BoxStayInsideBoxBoundaries)
+	{
+		const FReal Dt = 1 / 30.0f;
+		const FReal Fps = 1 / Dt;
+		const FReal SmallBoxSize = 100; // cm
+		const FReal ContainerBoxSize = 500; // cm
+		const FReal ContainerWallThickness = 10; // cm
+		const int ContainerFaceCount = 6;
+
+		const FVec3 InitialVelocity = FVec3(0, 0, 750); // ContainerBoxSize * 5 * Fps; // More than enough to tunnel
+
+		
+
+		using TEvolution = TypeParam;
+		TPBDRigidsSOAs<FReal, 3> Particles;
+		THandleArray<FChaosPhysicsMaterial> PhysicalMaterials;
+		TEvolution Evolution(Particles, PhysicalMaterials);
+		InitEvolutionSettings(Evolution);
+
+		// Create particles
+		TArray <TGeometryParticleHandle<FReal, 3>*> ContainerFaces = Evolution.CreateStaticParticles(ContainerFaceCount); // 6 sides of box
+		TPBDRigidParticleHandle<FReal, 3>* Dynamic = Evolution.CreateDynamicParticles(1)[0]; // The small box
+
+		// Set up physics material
+		TUniquePtr<FChaosPhysicsMaterial> PhysicsMaterial = MakeUnique<FChaosPhysicsMaterial>();
+		PhysicsMaterial->SleepCounterThreshold = 1000; // Don't sleep
+		PhysicsMaterial->Restitution = 0.0f;
+
+		// Create box geometry 
+		TUniquePtr<FImplicitObject> SmallBox(new TBox<FReal, 3>(FVec3(-SmallBoxSize/2, -SmallBoxSize/2, -SmallBoxSize/2), FVec3(SmallBoxSize/2, SmallBoxSize/2, SmallBoxSize/2)));
+
+		// Just use 3 (x2) boxes for the walls of the container (avoid rotation transforms for this test)
+		TUniquePtr<FImplicitObject> ContainerFaceX(new TBox<FReal, 3>(FVec3(-ContainerWallThickness / 2, -ContainerBoxSize / 2, -ContainerBoxSize / 2), FVec3(ContainerWallThickness / 2, ContainerBoxSize / 2, ContainerBoxSize / 2)));
+		TUniquePtr<FImplicitObject> ContainerFaceY(new TBox<FReal, 3>(FVec3(-ContainerBoxSize / 2, -ContainerWallThickness / 2, -ContainerBoxSize / 2), FVec3(ContainerBoxSize / 2, ContainerWallThickness / 2, ContainerBoxSize / 2)));
+		TUniquePtr<FImplicitObject> ContainerFaceZ(new TBox<FReal, 3>(FVec3(-ContainerBoxSize / 2, -ContainerBoxSize / 2, -ContainerWallThickness / 2), FVec3(ContainerBoxSize / 2, ContainerBoxSize / 2, ContainerWallThickness / 2)));		
+
+		// Assign geometry to all particles 
+		Dynamic->SetGeometry(MakeSerializable(SmallBox));
+
+		ContainerFaces[0]->SetGeometry(MakeSerializable(ContainerFaceX));
+		ContainerFaces[1]->SetGeometry(MakeSerializable(ContainerFaceX));
+		ContainerFaces[2]->SetGeometry(MakeSerializable(ContainerFaceY));
+		ContainerFaces[3]->SetGeometry(MakeSerializable(ContainerFaceY));
+		ContainerFaces[4]->SetGeometry(MakeSerializable(ContainerFaceZ));
+		ContainerFaces[5]->SetGeometry(MakeSerializable(ContainerFaceZ));
+		
+
+		Evolution.SetPhysicsMaterial(Dynamic, MakeSerializable(PhysicsMaterial));
+
+		const FReal Mass = 100000.0f;
+		Dynamic->I() = FMatrix33(Mass, Mass, Mass);
+		Dynamic->InvI() = FMatrix33(1.0f / Mass, 1.0f / Mass, 1.0f / Mass);
+
+		// Positions and velocities
+		ContainerFaces[0]->X() = FVec3(ContainerBoxSize / 2, 0, 0);
+		ContainerFaces[1]->X() = FVec3(-ContainerBoxSize / 2, 0, 0);
+		ContainerFaces[2]->X() = FVec3(0, ContainerBoxSize / 2, 0);
+		ContainerFaces[3]->X() = FVec3(0, -ContainerBoxSize / 2, 0);
+		ContainerFaces[4]->X() = FVec3(0, 0, ContainerBoxSize / 2);
+		ContainerFaces[5]->X() = FVec3(0, 0, -ContainerBoxSize / 2);
+
+		Dynamic->X() = FVec3(0, 0, 0);
+
+		// The position of the static has changed and statics don't automatically update bounds, so update explicitly
+		ContainerFaces[0]->SetWorldSpaceInflatedBounds(ContainerFaceX->BoundingBox().TransformedAABB(TRigidTransform<FReal, 3>(ContainerFaces[0]->X(), ContainerFaces[0]->R())));
+		ContainerFaces[1]->SetWorldSpaceInflatedBounds(ContainerFaceX->BoundingBox().TransformedAABB(TRigidTransform<FReal, 3>(ContainerFaces[1]->X(), ContainerFaces[1]->R())));
+		ContainerFaces[2]->SetWorldSpaceInflatedBounds(ContainerFaceY->BoundingBox().TransformedAABB(TRigidTransform<FReal, 3>(ContainerFaces[2]->X(), ContainerFaces[2]->R())));
+		ContainerFaces[3]->SetWorldSpaceInflatedBounds(ContainerFaceY->BoundingBox().TransformedAABB(TRigidTransform<FReal, 3>(ContainerFaces[3]->X(), ContainerFaces[3]->R())));
+		ContainerFaces[4]->SetWorldSpaceInflatedBounds(ContainerFaceZ->BoundingBox().TransformedAABB(TRigidTransform<FReal, 3>(ContainerFaces[4]->X(), ContainerFaces[4]->R())));
+		ContainerFaces[5]->SetWorldSpaceInflatedBounds(ContainerFaceZ->BoundingBox().TransformedAABB(TRigidTransform<FReal, 3>(ContainerFaces[5]->X(), ContainerFaces[5]->R())));
+
+		::ChaosTest::SetParticleSimDataToCollide({ Dynamic });
+		::ChaosTest::SetParticleSimDataToCollide({ ContainerFaces });
+
+		Dynamic->SetCCDEnabled(true);
+		Dynamic->SetGravityEnabled(false);
+
+		Dynamic->V() = InitialVelocity;
+
+		for (int i = 0; i < 10; ++i)
+		{
+			Evolution.AdvanceOneTimeStep(Dt);
+			Evolution.EndFrame(Dt);
+		}
+
+		// Large error margin, we are testing CCD and not solver accuracy
+		const FReal LargeErrorMargin = 10.0f;
+		// Check that we did not escape the box!
+		for (int axis = 0; axis < 3; axis++)
+		{
+			// If this failed, the dynamic cube escaped the air tight static container
+			EXPECT_LT(FMath::Abs(Dynamic->X()[axis]), ContainerBoxSize / 2 - ContainerWallThickness / 2 - SmallBoxSize / 2 + LargeErrorMargin);
+		}
+	}
+
+	// CCD not implemented for sphere sphere
+	TYPED_TEST(AllEvolutions, DISABLED_CCDTests_Shere_Sphere)
+	{
+		const FReal Dt = 1 / 30.0f;
+		const FReal Fps = 1 / Dt;
+		const FReal SphereRadius = 100; // cm
+		const FReal InitialSpeed = SphereRadius * 5 * Fps; // More than enough to tunnel
+		using TEvolution = TypeParam;
+		TPBDRigidsSOAs<FReal, 3> Particles;
+		THandleArray<FChaosPhysicsMaterial> PhysicalMaterials;
+		TEvolution Evolution(Particles, PhysicalMaterials);
+		InitEvolutionSettings(Evolution);
+
+		// Create particles
+		TGeometryParticleHandle<FReal, 3>* Static = Evolution.CreateStaticParticles(1)[0];
+		TPBDRigidParticleHandle<FReal, 3>* Dynamic = Evolution.CreateDynamicParticles(1)[0];
+
+		// Set up physics material
+		TUniquePtr<FChaosPhysicsMaterial> PhysicsMaterial = MakeUnique<FChaosPhysicsMaterial>();
+		PhysicsMaterial->SleepCounterThreshold = 1000; // Don't sleep
+		PhysicsMaterial->Restitution = 1.0f;
+
+		// Create Sphere geometry (Radius = 100)
+		TUniquePtr<FImplicitObject> Sphere(new TSphere<FReal, 3>(FVec3(0, 0, 0), SphereRadius));
+
+		// Assign sphere geometry to both particles 
+		Static->SetGeometry(MakeSerializable(Sphere));
+		Dynamic->SetGeometry(MakeSerializable(Sphere));
+
+		Evolution.SetPhysicsMaterial(Dynamic, MakeSerializable(PhysicsMaterial));
+
+		const FReal Mass = 100000.0f;;
+		Dynamic->I() = FMatrix33(Mass, Mass, Mass);
+		Dynamic->InvI() = FMatrix33(1.0f / Mass, 1.0f / Mass, 1.0f / Mass);
+
+		// Positions and velocities
+		Static->X() = FVec3(0, 0, 0);
+
+		Dynamic->X() = FVec3(0, 0, SphereRadius * 2 + 10);
+		//Dynamic->V() = FVec3(0, 0, -InitialSpeed);
+
+		// The position of the static has changed and statics don't automatically update bounds, so update explicitly
+		Static->SetWorldSpaceInflatedBounds(Sphere->BoundingBox().TransformedAABB(TRigidTransform<FReal, 3>(Static->X(), Static->R())));
+
+		::ChaosTest::SetParticleSimDataToCollide({ Static,Dynamic });
+
+		Dynamic->SetCCDEnabled(true);
+		Dynamic->SetGravityEnabled(false);
+
+
+		Dynamic->V() = FVec3(0, 0, -InitialSpeed);
+
+		for (int i = 0; i < 1; ++i)
+		{
+			Evolution.AdvanceOneTimeStep(Dt);
+			Evolution.EndFrame(Dt);
+		}
+
+		// Large error margin, we are testing CCD and not solver accuracy
+		const FReal LargeErrorMargin = 10.0f;
+		EXPECT_GE(Dynamic->X()[2], SphereRadius * 2 - LargeErrorMargin);
+	}
+}
+
diff --git a/Engine/Source/Runtime/Experimental/Chaos/Private/Chaos/CollisionResolution.cpp b/Engine/Source/Runtime/Experimental/Chaos/Private/Chaos/CollisionResolution.cpp
--- a/Engine/Source/Runtime/Experimental/Chaos/Private/Chaos/CollisionResolution.cpp
+++ b/Engine/Source/Runtime/Experimental/Chaos/Private/Chaos/CollisionResolution.cpp
@@ -1829,10 +1829,13 @@ namespace Chaos
 		void ConstructGenericConvexConvexConstraintsSwept(TGeometryParticleHandle<FReal, 3>* Particle0, TGeometryParticleHandle<FReal, 3>* Particle1, const FImplicitObject* Implicit0, const FImplicitObject* Implicit1, const FRigidTransform3& LocalTransform0, const FRigidTransform3& LocalTransform1, const FReal CullDistance, const FReal Dt, const FVec3& Dir, FReal Length, FCollisionConstraintsArray& NewConstraints)
 		{
 			CONDITIONAL_SCOPE_CYCLE_COUNTER(STAT_Collisions_ConstructGenericConvexConvexConstraintsSwept, ConstraintsDetailedStats);
+			TGenericParticleHandle<FReal, 3> RigidParticle = TGenericParticleHandle<FReal, 3>(Particle0);
 			FRigidBodySweptPointContactConstraint Constraint = FRigidBodySweptPointContactConstraint(Particle0, Implicit0, nullptr, LocalTransform0, Particle1, Implicit1, nullptr, LocalTransform1, EContactShapesType::GenericConvexConvex);
-			FRigidTransform3 WorldTransformX0 = LocalTransform0 * FRigidTransform3(Particle0->X(), Particle0->R());
+			// Note: This is unusual but we are using a mix of the previous and current transform
+			// This is due to how CCD only rewinds the position
+			FRigidTransform3 WorldTransformXQ0 = LocalTransform0 * FRigidTransform3(RigidParticle->X(), RigidParticle->Q());
 			FRigidTransform3 WorldTransform1 = LocalTransform1 * Collisions::GetTransform(Particle1);
-			UpdateGenericConvexConvexConstraintSwept(Particle0, *Implicit0, WorldTransformX0, *Implicit1, WorldTransform1, Dir, Length, CullDistance, Dt, Constraint);
+			UpdateGenericConvexConvexConstraintSwept(Particle0, *Implicit0, WorldTransformXQ0, *Implicit1, WorldTransform1, Dir, Length, CullDistance, Dt, Constraint);
 			NewConstraints.Add(Constraint);
 		}
 
@@ -2260,6 +2263,104 @@ namespace Chaos
 			//Constraint.Manifold.ContactMoveSQRDistance = FMath::Max((NewContactPositionLocal0 - OriginalContactPositionLocal0).SizeSquared(), (NewContactPositionLocal1 - OriginalContactPositionLocal1).SizeSquared());
 		}
 
+		// Run collision detection for the swept constraits
+		// NOTE: Transforms are world space shape transforms
+		
+		template<ECollisionUpdateType UpdateType>
+		inline void UpdateConstraintFromGeometryImpl(FRigidBodySweptPointContactConstraint& Constraint, const FRigidTransform3& WorldTransform0, const FRigidTransform3& WorldTransform1, const FReal CullDistance, const FReal Dt)
+		{
+			FReal LengthCCD = 0.0f;
+			FVec3 DirCCD(0.0f);
+			bool bUseCCD = UseCCD(Constraint.Particle[0], Constraint.Particle[1], Constraint.GetManifold().Implicit[0], DirCCD, LengthCCD);
+
+			const FImplicitObject& Implicit0 = *Constraint.Manifold.Implicit[0];
+			const FImplicitObject& Implicit1 = *Constraint.Manifold.Implicit[1];
+
+			TGeometryParticleHandle<FReal, 3>* Particle0 = Constraint.Particle[0];
+
+			if (bUseCCD)
+			{
+				switch (Constraint.Manifold.ShapesType)
+				{
+				case EContactShapesType::GenericConvexConvex:
+					UpdateGenericConvexConvexConstraintSwept(Particle0, Implicit0, WorldTransform0, Implicit1, WorldTransform1, DirCCD, LengthCCD, CullDistance, Dt, Constraint);
+					return;
+				case EContactShapesType::SphereHeightField:
+				{
+					const TSphere<FReal, 3>* Object0 = Implicit0.template GetObject<const TSphere<FReal, 3>>();
+					const FHeightField* Object1 = Implicit1.template GetObject<const FHeightField>();
+					UpdateSphereHeightFieldConstraintSwept(Particle0, *Object0, WorldTransform0, *Object1, WorldTransform1, DirCCD, LengthCCD, CullDistance, Dt, Constraint);
+					return;
+				}
+				case EContactShapesType::CapsuleHeightField:
+				{
+					const TCapsule<FReal>* Object0 = Implicit0.template GetObject<const TCapsule<FReal> >();
+					const FHeightField* Object1 = Implicit1.template GetObject<const FHeightField >();
+					UpdateCapsuleHeightFieldConstraintSwept(Particle0, *Object0, WorldTransform0, *Object1, WorldTransform1, DirCCD, LengthCCD, CullDistance, Dt, Constraint);
+					return;
+				}
+				case EContactShapesType::SphereTriMesh:
+				{
+					const TSphere<FReal, 3>* Object0 = Implicit0.template GetObject<const TSphere<FReal, 3>>();
+					if (const TImplicitObjectScaled<FTriangleMeshImplicitObject>* ScaledTriangleMesh = Implicit1.template GetObject<const TImplicitObjectScaled<FTriangleMeshImplicitObject>>())
+					{
+						UpdateSphereTriangleMeshConstraintSwept(Particle0, *Object0, WorldTransform0, *ScaledTriangleMesh, WorldTransform1, DirCCD, LengthCCD, CullDistance, Dt, Constraint);
+					}
+					else if (const FTriangleMeshImplicitObject* TriangleMesh = Implicit1.template GetObject<const FTriangleMeshImplicitObject>())
+					{
+						UpdateSphereTriangleMeshConstraintSwept(Particle0, *Object0, WorldTransform0, *TriangleMesh, WorldTransform1, DirCCD, LengthCCD, CullDistance, Dt, Constraint);
+					}
+					else
+					{
+						ensure(false);
+					}
+					return;
+				}
+				case EContactShapesType::CapsuleTriMesh:
+				{
+					const TCapsule<FReal>* Object0 = Implicit0.template GetObject<const TCapsule<FReal> >();
+					if (const TImplicitObjectScaled<FTriangleMeshImplicitObject>* ScaledTriangleMesh = Implicit1.template GetObject<const TImplicitObjectScaled<FTriangleMeshImplicitObject>>())
+					{
+						UpdateCapsuleTriangleMeshConstraintSwept(Particle0, *Object0, WorldTransform0, *ScaledTriangleMesh, WorldTransform1, DirCCD, LengthCCD, CullDistance, Dt, Constraint);
+					}
+					else if (const FTriangleMeshImplicitObject* TriangleMesh = Implicit1.template GetObject<const FTriangleMeshImplicitObject>())
+					{
+						UpdateCapsuleTriangleMeshConstraintSwept(Particle0, *Object0, WorldTransform0, *TriangleMesh, WorldTransform1, DirCCD, LengthCCD, CullDistance, Dt, Constraint);
+					}
+					else
+					{
+						ensure(false);
+					}
+					return;
+				}
+				case EContactShapesType::ConvexHeightField:
+				{
+					const FHeightField* Object1 = Implicit1.template GetObject<const FHeightField >();
+					UpdateConvexHeightFieldConstraintSwept(Particle0, Implicit0, WorldTransform0, *Object1, WorldTransform1, DirCCD, LengthCCD, CullDistance, Dt, Constraint);
+					return;
+				}
+				case EContactShapesType::ConvexTriMesh:
+					if (const TImplicitObjectScaled<FTriangleMeshImplicitObject>* ScaledTriangleMesh = Implicit1.template GetObject<const TImplicitObjectScaled<FTriangleMeshImplicitObject>>())
+					{
+						UpdateConvexTriangleMeshConstraintSwept(Particle0, Implicit0, WorldTransform0, *ScaledTriangleMesh, WorldTransform1, DirCCD, LengthCCD, CullDistance, Dt, Constraint);
+					}
+					else if (const FTriangleMeshImplicitObject* TriangleMesh = Implicit1.template GetObject<const FTriangleMeshImplicitObject>())
+					{
+						UpdateConvexTriangleMeshConstraintSwept(Particle0, Implicit0, WorldTransform0, *TriangleMesh, WorldTransform1, DirCCD, LengthCCD, CullDistance, Dt, Constraint);
+					}
+					else
+					{
+						ensure(false);
+					}
+					return;
+				}
+			}
+
+			Constraint.TimeOfImpact = 1.0f; // CCD will not be used
+			// Do a normal non-swept update if we reach this point - Not required for now, since this will be done in solver
+			// UpdateConstraintFromGeometryImpl<UpdateType>(*(Constraint.As<FRigidBodyPointContactConstraint>()), WorldTransform0, WorldTransform1, CullDistance, Dt);
+		}
+
 		template<typename T_TRAITS>
 		void ConstructConstraintsImpl(TGeometryParticleHandle<FReal, 3>* Particle0, TGeometryParticleHandle<FReal, 3>* Particle1, const FImplicitObject* Implicit0, const FBVHParticles* Simplicial0, const FImplicitObject* Implicit1, const FBVHParticles* Simplicial1, const FRigidTransform3& LocalTransform0, const FRigidTransform3& LocalTransform1, const FReal CullDistance, const FReal Dt, const FCollisionContext& Context, FCollisionConstraintsArray& NewConstraints)
 		{
@@ -2529,8 +2630,8 @@ namespace Chaos
 
 		// Run collision detection for the specified constraint to update the nearest contact point.
 		// NOTE: Transforms are world space particle transforms
-		template<ECollisionUpdateType UpdateType>
-		void UpdateConstraintFromGeometry(FRigidBodyPointContactConstraint& Constraint, const FRigidTransform3& ParticleTransform0, const FRigidTransform3& ParticleTransform1, const FReal CullDistance, const FReal Dt)
+		template<ECollisionUpdateType UpdateType, typename RigidBodyContactConstraint>
+		void UpdateConstraintFromGeometry(RigidBodyContactConstraint& Constraint, const FRigidTransform3& ParticleTransform0, const FRigidTransform3& ParticleTransform1, const FReal CullDistance, const FReal Dt)
 		{
 			const FRigidTransform3 WorldTransform0 = Constraint.ImplicitTransform[0] * ParticleTransform0;
 			const FRigidTransform3 WorldTransform1 = Constraint.ImplicitTransform[1] * ParticleTransform1;
@@ -2835,8 +2936,11 @@ namespace Chaos
 		template void UpdateLevelsetLevelsetConstraint<ECollisionUpdateType::Any>(const FRigidTransform3& WorldTransform0, const FRigidTransform3& WorldTransform1, const float CullDistance, const FReal Dt, FRigidBodyPointContactConstraint& Constraint);
 		template void UpdateLevelsetLevelsetConstraint<ECollisionUpdateType::Deepest>(const FRigidTransform3& WorldTransform0, const FRigidTransform3& WorldTransform1, const float CullDistance, const FReal Dt, FRigidBodyPointContactConstraint& Constraint);
 
-		template void UpdateConstraintFromGeometry<ECollisionUpdateType::Any>(FRigidBodyPointContactConstraint& ConstraintBase, const FRigidTransform3& Transform0, const FRigidTransform3& Transform1, const float CullDistance, const FReal Dt);
-		template void UpdateConstraintFromGeometry<ECollisionUpdateType::Deepest>(FRigidBodyPointContactConstraint& ConstraintBase, const FRigidTransform3& Transform0, const FRigidTransform3& Transform1, const float CullDistance, const FReal Dt);
+		template void UpdateConstraintFromGeometry<ECollisionUpdateType::Any, FRigidBodyPointContactConstraint>(FRigidBodyPointContactConstraint& ConstraintBase, const FRigidTransform3& Transform0, const FRigidTransform3& Transform1, const float CullDistance, const FReal Dt);
+		template void UpdateConstraintFromGeometry<ECollisionUpdateType::Deepest, FRigidBodyPointContactConstraint>(FRigidBodyPointContactConstraint& ConstraintBase, const FRigidTransform3& Transform0, const FRigidTransform3& Transform1, const float CullDistance, const FReal Dt);
+
+		template void UpdateConstraintFromGeometry<ECollisionUpdateType::Any, FRigidBodySweptPointContactConstraint>(FRigidBodySweptPointContactConstraint& ConstraintBase, const FRigidTransform3& Transform0, const FRigidTransform3& Transform1, const float CullDistance, const FReal Dt);
+		template void UpdateConstraintFromGeometry<ECollisionUpdateType::Deepest, FRigidBodySweptPointContactConstraint>(FRigidBodySweptPointContactConstraint& ConstraintBase, const FRigidTransform3& Transform0, const FRigidTransform3& Transform1, const float CullDistance, const FReal Dt);
 
 	} // Collisions
 
diff --git a/Engine/Source/Runtime/Experimental/Chaos/Private/Chaos/PBDCollisionConstraintsContact.cpp b/Engine/Source/Runtime/Experimental/Chaos/Private/Chaos/PBDCollisionConstraintsContact.cpp
--- a/Engine/Source/Runtime/Experimental/Chaos/Private/Chaos/PBDCollisionConstraintsContact.cpp
+++ b/Engine/Source/Runtime/Experimental/Chaos/Private/Chaos/PBDCollisionConstraintsContact.cpp
@@ -56,14 +56,20 @@ namespace Chaos
 			const FRigidTransform3 Transform1 = GetTransform(Constraint.Particle[1]);
 
 			Constraint.ResetPhi(CullDistance);
-			UpdateConstraintFromGeometry<ECollisionUpdateType::Deepest>(Constraint, Transform0, Transform1, CullDistance, Dt);
+			UpdateConstraintFromGeometry<ECollisionUpdateType::Deepest, FRigidBodyPointContactConstraint>(Constraint, Transform0, Transform1, CullDistance, Dt);
 		}
 
-		void Update(FRigidBodySweptPointContactConstraint& Constraint, const FReal CullDistance, const FReal Dt)
+		void UpdateSwept(FRigidBodySweptPointContactConstraint& Constraint, const FReal CullDistance, const FReal Dt)
 		{
+			TGenericParticleHandle<FReal, 3> Particle0 = TGenericParticleHandle<FReal, 3>(Constraint.Particle[0]);
+			// Note: This is unusual but we are using a mix of the previous and current transform
+			// This is due to how CCD only rewinds the position
+			const FRigidTransform3 TransformXQ0(Particle0->X(), Particle0->Q());
+			const FRigidTransform3 Transform1 = GetTransform(Constraint.Particle[1]);
+			Constraint.ResetPhi(CullDistance);
+			
 			// Update as a point constraint (base class).
-			Constraint.bShouldTreatAsSinglePoint = true;
-			Update(*Constraint.As<FRigidBodyPointContactConstraint>(), CullDistance, Dt);
+			UpdateConstraintFromGeometry<ECollisionUpdateType::Deepest, FRigidBodySweptPointContactConstraint>(Constraint, TransformXQ0, Transform1, CullDistance, Dt);
 		}
 
 
@@ -257,14 +263,16 @@ namespace Chaos
 		
 		void ApplySweptImpl(FRigidBodySweptPointContactConstraint& Constraint, const FContactIterationParameters & IterationParameters, const FContactParticleParameters & ParticleParameters)
 		{
-			const FReal TimeOfImpactErrorMargin = 20.0f;  // Large error margin in cm, this is to ensure that 1) velocity is solved for at the time of impact, and  2) the contact is not disabled
 			FGenericParticleHandle Particle0 = FGenericParticleHandle(Constraint.Particle[0]);
 			FGenericParticleHandle Particle1 = FGenericParticleHandle(Constraint.Particle[1]);
 
-			if (Constraint.bShouldTreatAsSinglePoint || IterationParameters.Iteration > 0 || Constraint.TimeOfImpact == 1)
+			Collisions::UpdateSwept(Constraint, ParticleParameters.CullDistance, IterationParameters.Dt);
+
+			if (Constraint.TimeOfImpact == 1)
 			{
-				// If not on first iteration, or at TOI = 1 (normal constraint) we don't want to split timestep at TOI.
-				ApplyImpl(Constraint, IterationParameters, ParticleParameters);
+				const FContactParticleParameters CCDParticleParamatersNoContactCulling{ ParticleParameters.CullDistance, ParticleParameters.RestitutionVelocityThreshold, false, ParticleParameters.Collided };
+				// If TOI = 1 (normal constraint) we don't want to split timestep at TOI.
+				ApplyImpl(Constraint, IterationParameters, CCDParticleParamatersNoContactCulling);
 				return;
 			}
 
@@ -272,19 +280,37 @@ namespace Chaos
 			// P may have changed due to other constraints, so at TOI our manifold needs updating.
 			const FReal PartialDT = Constraint.TimeOfImpact * IterationParameters.Dt;
 			const FReal RemainingDT = (1 - Constraint.TimeOfImpact) * IterationParameters.Dt;
-			const int32 FakeIteration = IterationParameters.NumIterations / 2; // For iteration count dependent effects (like relaxation)
-			const int32 PartialPairIterations = FMath::Max(IterationParameters.NumPairIterations, 2); // Do at least 2 pair iterations
+			const int32 FakeIteration = IterationParameters.NumIterations / 2; // For iteration count dependent effects (like relaxation) // @todo: Do we still need this?
+			const int32 PartialPairIterations = FMath::Max(IterationParameters.NumPairIterations, 2); // Do at least 2 pair iterations // @todo: Do we still need this?
 			const FContactIterationParameters IterationParametersPartialDT{ PartialDT, FakeIteration, IterationParameters.NumIterations, PartialPairIterations, IterationParameters.ApplyType, IterationParameters.NeedsAnotherIteration };
 			const FContactIterationParameters IterationParametersRemainingDT{ RemainingDT, FakeIteration, IterationParameters.NumIterations, IterationParameters.NumPairIterations, IterationParameters.ApplyType, IterationParameters.NeedsAnotherIteration };
-			const FContactParticleParameters CCDParticleParamaters{ ParticleParameters.CullDistance + TimeOfImpactErrorMargin, ParticleParameters.RestitutionVelocityThreshold, ParticleParameters.bCanDisableContacts, ParticleParameters.Collided };
+			const FContactParticleParameters CCDParticleParamaters{ ParticleParameters.CullDistance, ParticleParameters.RestitutionVelocityThreshold, false, ParticleParameters.Collided };
 
 			// Rewind P to TOI and Apply
 			Particle0->P() = FMath::Lerp(Particle0->X(), Particle0->P(), Constraint.TimeOfImpact);
 			ApplyImpl(Constraint, IterationParametersPartialDT, CCDParticleParamaters);
 
 			// Advance P to end of frame from TOI, and Apply
-			Particle0->P() = Particle0->P() + Particle0->V() * RemainingDT;
-			ApplyImpl(Constraint, IterationParametersRemainingDT, CCDParticleParamaters);
+			if (IterationParameters.Iteration + 1 < IterationParameters.NumIterations)
+			{
+				Particle0->P() = Particle0->P() + Particle0->V() * RemainingDT; // If we are tunneling through something else due to this, it will be resolved in the next iteration
+				ApplyImpl(Constraint, IterationParametersRemainingDT, CCDParticleParamaters);
+			}
+			else
+			{
+				// We get here if we cannot solve CCD collisions with the given number of iterations and restitution settings.
+				// So don't do the remaining dt update. This will bleed the energy! (also: Ignore rotation)
+				// To prevent this condition: increase number of iterations and/or reduce restitution and/or reduce velocities
+				if (IterationParameters.Dt > SMALL_NUMBER)
+				{
+					// Update velocity to be consistent with PBD
+					Particle0->SetV(1 / IterationParameters.Dt * (Particle0->P() - Particle0->X()));
+				}
+				else
+				{
+					Particle0->SetV(FVec3(0));
+				}
+			}
 		}
 
 
diff --git a/Engine/Source/Runtime/Experimental/Chaos/Public/Chaos/Collision/PBDCollisionConstraint.h b/Engine/Source/Runtime/Experimental/Chaos/Public/Chaos/Collision/PBDCollisionConstraint.h
--- a/Engine/Source/Runtime/Experimental/Chaos/Public/Chaos/Collision/PBDCollisionConstraint.h
+++ b/Engine/Source/Runtime/Experimental/Chaos/Public/Chaos/Collision/PBDCollisionConstraint.h
@@ -349,12 +349,10 @@ namespace Chaos
 			EContactShapesType ShapesType)
 			: Base(Particle0, Implicit0, Simplicial0, Transform0, Particle1, Implicit1, Simplicial1, Transform1, Base::FType::SinglePointSwept, ShapesType)
 			, TimeOfImpact(0)
-		    , bShouldTreatAsSinglePoint(false)
 		{}
 
 		// Value in range [0,1] used to interpolate P between [X,P] that we will rollback to when solving at time of impact.
 		FReal TimeOfImpact;
-		bool bShouldTreatAsSinglePoint;
 		static typename Base::FType StaticType() { return Base::FType::SinglePointSwept; };
 	};
 
diff --git a/Engine/Source/Runtime/Experimental/Chaos/Public/Chaos/CollisionResolution.h b/Engine/Source/Runtime/Experimental/Chaos/Public/Chaos/CollisionResolution.h
--- a/Engine/Source/Runtime/Experimental/Chaos/Public/Chaos/CollisionResolution.h
+++ b/Engine/Source/Runtime/Experimental/Chaos/Public/Chaos/CollisionResolution.h
@@ -35,8 +35,9 @@ namespace Chaos
 		//
 
 		// Update the constraint by re-running collision detection on the shape pair.
-		template<ECollisionUpdateType UpdateType>
-		void CHAOS_API UpdateConstraintFromGeometry(FRigidBodyPointContactConstraint& Constraint, const FRigidTransform3& ParticleTransform0, const FRigidTransform3& ParticleTransform1, const FReal CullDistance, const FReal Dt);
+
+		template<ECollisionUpdateType UpdateType, typename RigidBodyContactConstraint>
+		void CHAOS_API UpdateConstraintFromGeometry(RigidBodyContactConstraint& Constraint, const FRigidTransform3& ParticleTransform0, const FRigidTransform3& ParticleTransform1, const FReal CullDistance, const FReal Dt);
 
 		// Create constraints for the particle pair. This could create multiple constraints: one for each potentially colliding shape pair in multi-shape particles.
 		void CHAOS_API ConstructConstraints(TGeometryParticleHandle<FReal, 3>* Particle0, TGeometryParticleHandle<FReal, 3>* Particle1, const FImplicitObject* Implicit0, const FBVHParticles* Simplicial0, const FImplicitObject* Implicit1, const FBVHParticles* Simplicial1, const FRigidTransform3& Transform0, const FRigidTransform3& Transform1, const FReal CullDistance, const FReal dT,const FCollisionContext& Context, FCollisionConstraintsArray& NewConstraints);
