HeadlessChaos build fix - template declaration on cpp implementation hiding function from the linker (it's no longer a template function in the declaration so the linker couldn't find it)

diff --git a/Engine/Source/Programs/HeadlessChaos/Private/HeadlessChaosTestGJK.cpp b/Engine/Source/Programs/HeadlessChaos/Private/HeadlessChaosTestGJK.cpp
--- a/Engine/Source/Programs/HeadlessChaos/Private/HeadlessChaosTestGJK.cpp
+++ b/Engine/Source/Programs/HeadlessChaos/Private/HeadlessChaosTestGJK.cpp
@@ -1418,7 +1418,6 @@ namespace ChaosTest
 		}
 	}
 
-	template <typename T>
 	void GJKCapsuleConvexInitialOverlapSweep()
 	{
 		{
