HeadlessChaos, property test fix, match UE5 specific API changes.

diff --git a/Engine/Source/Programs/HeadlessChaos/Private/HeadlessChaosTestPhysicsScene.cpp b/Engine/Source/Programs/HeadlessChaos/Private/HeadlessChaosTestPhysicsScene.cpp
--- a/Engine/Source/Programs/HeadlessChaos/Private/HeadlessChaosTestPhysicsScene.cpp
+++ b/Engine/Source/Programs/HeadlessChaos/Private/HeadlessChaosTestPhysicsScene.cpp
@@ -1088,7 +1088,7 @@ namespace ChaosTest {
 			}
 
 			FVec3 Grav(0, 0, 0);
-			Scene.SetUpForFrame(&Grav, GTDt, 99999, 99999, 10, false);
+			Scene.SetUpForFrame(&Grav, GTDt, 99999, 99999, 10, false, FixedDT);
 			Scene.StartFrame();
 			Scene.EndFrame();
 
