Fix headless chaos compilation
Used proper namespace to mach recent change in SphereGenerator.h

diff --git a/Engine/Source/Programs/HeadlessChaos/Private/GeometryCollection/GeometryCollectionTestFramework.cpp b/Engine/Source/Programs/HeadlessChaos/Private/GeometryCollection/GeometryCollectionTestFramework.cpp
--- a/Engine/Source/Programs/HeadlessChaos/Private/GeometryCollection/GeometryCollectionTestFramework.cpp
+++ b/Engine/Source/Programs/HeadlessChaos/Private/GeometryCollection/GeometryCollectionTestFramework.cpp
@@ -16,6 +16,8 @@
 #include "../Resource/SphereGeometry.h"
 
 
+using namespace UE::Geometry;
+
 namespace GeometryCollectionTest
 {
 	TSharedPtr<FGeometryCollection> MakeSphereElement(
