Chaos - Inertia fixes

diff --git a/Engine/Source/Programs/HeadlessChaos/Private/HeadlessChaosTestImplicitMassProperties.cpp b/Engine/Source/Programs/HeadlessChaos/Private/HeadlessChaosTestImplicitMassProperties.cpp
new file mode 100644
--- /dev/null
+++ b/Engine/Source/Programs/HeadlessChaos/Private/HeadlessChaosTestImplicitMassProperties.cpp
@@ -0,0 +1,126 @@
+// Copyright Epic Games, Inc. All Rights Reserved.
+
+#include "Chaos/Box.h"
+#include "Chaos/ImplicitFwd.h"
+#include "Chaos/ImplicitObjectScaled.h"
+#include "Chaos/ImplicitObjectTransformed.h"
+#include "Chaos/Matrix.h"
+#include "Chaos/MassProperties.h"
+#include "Chaos/TriangleMesh.h"
+#include "HeadlessChaos.h"
+#include "HeadlessChaosTestUtility.h"
+#include "Modules/ModuleManager.h"
+
+namespace ChaosTest 
+{
+	using namespace Chaos;
+
+	// Box inertia is correct
+	GTEST_TEST(ImplicitMassPropertiesTests, TestImplicitBoxMassPropeerties)
+	{
+		// Offset box
+		const FReal ImplicitMass = 100.0f;
+		const FVec3 ImplicitMin = FVec3(100, 200, 300);
+		const FVec3 ImplicitMax = FVec3(200, 400, 600);
+		const FImplicitBox3 Implicit(ImplicitMin, ImplicitMax);
+
+		const FMatrix33 ImplicitInertia = Implicit.GetInertiaTensor(ImplicitMass);
+		const FVec3 ImplicitCoM = Implicit.GetCenterOfMass();
+		const FRotation3 ImplicitRoM = Implicit.GetRotationOfMass();
+
+		const FVec3 ExpectedCoM = 0.5f * (ImplicitMin + ImplicitMax);
+		const FRotation3 ExpectedRoM = FRotation3::FromIdentity();
+		EXPECT_NEAR(ExpectedCoM.X, ImplicitCoM.X, 1.e-2f);
+		EXPECT_NEAR(ExpectedCoM.Y, ImplicitCoM.Y, 1.e-2f);
+		EXPECT_NEAR(ExpectedCoM.Z, ImplicitCoM.Z, 1.e-2f);
+	}
+
+
+	// Convex box inertia is correct
+	GTEST_TEST(ImplicitMassPropertiesTests, TestImplicitConvexBoxMassPropeerties)
+	{
+		// Offset box
+		const FReal ImplicitMass = 100.0f;
+		const FVec3 ImplicitMin = FVec3(100, 200, 300);
+		const FVec3 ImplicitMax = FVec3(200, 400, 600);
+		const FImplicitConvex3 Implicit = CreateConvexBox(ImplicitMin, ImplicitMax, 0.0f);
+
+		const FMatrix33 ImplicitInertia = Implicit.GetInertiaTensor(ImplicitMass);
+		const FVec3 ImplicitCoM = Implicit.GetCenterOfMass();
+		const FRotation3 ImplicitRoM = Implicit.GetRotationOfMass();
+
+		const FImplicitBox3 Box(ImplicitMin, ImplicitMax, 0.0f);
+		const FVec3 ExpectedCoM = Box.GetCenterOfMass();
+		const FRotation3 ExpectedRoM = Box.GetRotationOfMass();
+		const FMatrix33 ExectedInertia = Box.GetInertiaTensor(ImplicitMass);
+		EXPECT_NEAR(ExpectedCoM.X, ImplicitCoM.X, 1.e-2f);
+		EXPECT_NEAR(ExpectedCoM.Y, ImplicitCoM.Y, 1.e-2f);
+		EXPECT_NEAR(ExpectedCoM.Z, ImplicitCoM.Z, 1.e-2f);
+		EXPECT_NEAR(ImplicitInertia.M[0][0], ExectedInertia.M[0][0], 1.e-1f);
+		EXPECT_NEAR(ImplicitInertia.M[1][1], ExectedInertia.M[1][1], 1.e-1f);
+		EXPECT_NEAR(ImplicitInertia.M[2][2], ExectedInertia.M[2][2], 1.e-1f);
+	}
+
+
+	// Scaled convex box inertia is correct
+	GTEST_TEST(ImplicitMassPropertiesTests, TestImplicitScaledConvexBoxMassPropeerties)
+	{
+		// Offset box
+		const FReal ImplicitMass = 100.0f;
+		const FVec3 ImplicitSize = FVec3(100, 200, 300);
+		const FVec3 ImplicitScale = FVec3(1.5f, 1.8f, 1.1f);
+		const TImplicitObjectScaled<FImplicitConvex3> Implicit = CreateScaledConvexBox(ImplicitSize, ImplicitScale, 0.0f);
+
+		const FMatrix33 ImplicitInertia = Implicit.GetInertiaTensor(ImplicitMass);
+		const FVec3 ImplicitCoM = Implicit.GetCenterOfMass();
+		const FRotation3 ImplicitRoM = Implicit.GetRotationOfMass();
+
+		const FImplicitBox3 Box(-0.5f * ImplicitScale * ImplicitSize, 0.5f * ImplicitScale * ImplicitSize, 0.0f);
+		const FVec3 ExpectedCoM = Box.GetCenterOfMass();
+		const FRotation3 ExpectedRoM = Box.GetRotationOfMass();
+		const FMatrix33 ExectedInertia = Box.GetInertiaTensor(ImplicitMass);
+		EXPECT_NEAR(ExpectedCoM.X, ImplicitCoM.X, 1.e-2f);
+		EXPECT_NEAR(ExpectedCoM.Y, ImplicitCoM.Y, 1.e-2f);
+		EXPECT_NEAR(ExpectedCoM.Z, ImplicitCoM.Z, 1.e-2f);
+		EXPECT_NEAR(ImplicitInertia.M[0][0], ExectedInertia.M[0][0], 1.f);
+		EXPECT_NEAR(ImplicitInertia.M[1][1], ExectedInertia.M[1][1], 1.f);
+		EXPECT_NEAR(ImplicitInertia.M[2][2], ExectedInertia.M[2][2], 1.f);
+	}
+
+
+	// @todo(chaos): Convex shape that is not a box
+
+	// @todo(chaos): Compound transformed scaled shape inertia
+
+	// Transformed Scaled convex box inertia is correct
+	// @chaos(todo): Cannot be tested unless we make GetInertiaTensor virtual on 
+	// ImplicitObject (required to make GetInertiaTensor compile on Transformed Implicit)
+#if 0
+	GTEST_TEST(ImplicitMassPropertiesTests, TestImplicitTransformedConvexBoxMassPropeerties)
+	{
+		// Offset box
+		const FReal ImplicitMass = 100.0f;
+		const FVec3 ImplicitCenter = FVec3(100, 100, 100);
+		const FVec3 ImplicitSize = FVec3(100, 200, 300);
+		const FVec3 ImplicitScale = FVec3(1.5f, 1.8f, 1.1f);
+		TImplicitObjectScaled<FImplicitConvex3> Convex = CreateScaledConvexBox(ImplicitSize, ImplicitScale, 0.0f);
+		const TImplicitObjectTransformed<FReal, 3, false> Implicit(&Convex, FRigidTransform3(ImplicitCenter, FRotation3::FromIdentity()));
+
+		const FMatrix33 ImplicitInertia = Implicit.GetInertiaTensor(ImplicitMass);
+		const FVec3 ImplicitCoM = Implicit.GetCenterOfMass();
+		const FRotation3 ImplicitRoM = Implicit.GetRotationOfMass();
+
+		const FImplicitBox3 Box(ImplicitCenter - 0.5f * ImplicitScale * ImplicitSize, ImplicitCenter + 0.5f * ImplicitScale * ImplicitSize, 0.0f);
+		const FVec3 ExpectedCoM = Box.GetCenterOfMass();
+		const FRotation3 ExpectedRoM = Box.GetRotationOfMass();
+		const FMatrix33 ExectedInertia = Box.GetInertiaTensor(ImplicitMass);
+		EXPECT_NEAR(ExpectedCoM.X, ImplicitCoM.X, 1.e-2f);
+		EXPECT_NEAR(ExpectedCoM.Y, ImplicitCoM.Y, 1.e-2f);
+		EXPECT_NEAR(ExpectedCoM.Z, ImplicitCoM.Z, 1.e-2f);
+		EXPECT_NEAR(ImplicitInertia.M[0][0], ExectedInertia.M[0][0], 1.e-1f);
+		EXPECT_NEAR(ImplicitInertia.M[1][1], ExectedInertia.M[1][1], 1.e-1f);
+		EXPECT_NEAR(ImplicitInertia.M[2][2], ExectedInertia.M[2][2], 1.e-1f);
+	}
+#endif
+
+}
diff --git a/Engine/Source/Programs/HeadlessChaos/Private/HeadlessChaosTestUtility.cpp b/Engine/Source/Programs/HeadlessChaos/Private/HeadlessChaosTestUtility.cpp
--- a/Engine/Source/Programs/HeadlessChaos/Private/HeadlessChaosTestUtility.cpp
+++ b/Engine/Source/Programs/HeadlessChaos/Private/HeadlessChaosTestUtility.cpp
@@ -702,38 +702,46 @@ namespace ChaosTest {
 		}
 	}
 
-	TArray<FVec3> MakeBoxVerts(const FVec3& HalfSize)
+	TArray<FVec3> MakeBoxVerts(const FVec3& Center, const FVec3& HalfSize)
 	{
 		return
 		{
-			FVec3(-HalfSize.X, -HalfSize.Y, -HalfSize.Z),
-			FVec3(-HalfSize.X,  HalfSize.Y, -HalfSize.Z),
-			FVec3(HalfSize.X,  HalfSize.Y, -HalfSize.Z),
-			FVec3(HalfSize.X, -HalfSize.Y, -HalfSize.Z),
-			FVec3(-HalfSize.X, -HalfSize.Y,  HalfSize.Z),
-			FVec3(-HalfSize.X,  HalfSize.Y,  HalfSize.Z),
-			FVec3(HalfSize.X,  HalfSize.Y,  HalfSize.Z),
-			FVec3(HalfSize.X, -HalfSize.Y,  HalfSize.Z),
+			Center + FVec3(-HalfSize.X, -HalfSize.Y, -HalfSize.Z),
+			Center + FVec3(-HalfSize.X,  HalfSize.Y, -HalfSize.Z),
+			Center + FVec3(HalfSize.X,  HalfSize.Y, -HalfSize.Z),
+			Center + FVec3(HalfSize.X, -HalfSize.Y, -HalfSize.Z),
+			Center + FVec3(-HalfSize.X, -HalfSize.Y,  HalfSize.Z),
+			Center + FVec3(-HalfSize.X,  HalfSize.Y,  HalfSize.Z),
+			Center + FVec3(HalfSize.X,  HalfSize.Y,  HalfSize.Z),
+			Center + FVec3(HalfSize.X, -HalfSize.Y,  HalfSize.Z),
 		};
 	}
 
 	FImplicitConvex3 CreateConvexBox(const FVec3& BoxSize, const FReal Margin)
 	{
 		const FVec3 HalfSize = 0.5f * BoxSize;
-		return FImplicitConvex3(MakeBoxVerts(HalfSize), Margin);
+		return FImplicitConvex3(MakeBoxVerts(FVec3(0), HalfSize), Margin);
 	}
 
+	FImplicitConvex3 CreateConvexBox(const FVec3& BoxMin, const FVec3& BoxMax, const FReal Margin)
+	{
+		const FVec3 Center = 0.5f * (BoxMin + BoxMax);
+		const FVec3 HalfSize = 0.5f * (BoxMax - BoxMin);
+		return FImplicitConvex3(MakeBoxVerts(Center, HalfSize), Margin);
+	}
+
+
 	TImplicitObjectInstanced<FImplicitConvex3> CreateInstancedConvexBox(const FVec3& BoxSize, const FReal Margin)
 	{
 		const FVec3 HalfSize = 0.5f * BoxSize;
-		TSharedPtr<FImplicitConvex3, ESPMode::ThreadSafe> BoxConvex = MakeShared<FImplicitConvex3, ESPMode::ThreadSafe>(MakeBoxVerts(HalfSize), 0.0f);
+		TSharedPtr<FImplicitConvex3, ESPMode::ThreadSafe> BoxConvex = MakeShared<FImplicitConvex3, ESPMode::ThreadSafe>(MakeBoxVerts(FVec3(0), HalfSize), 0.0f);
 		return TImplicitObjectInstanced<FImplicitConvex3>(BoxConvex, Margin);
 	}
 
 	TImplicitObjectScaled<FImplicitConvex3> CreateScaledConvexBox(const FVec3& BoxSize, const FVec3 BoxScale, const FReal Margin)
 	{
 		const FVec3 HalfSize = 0.5f * BoxSize;
-		TSharedPtr<FImplicitConvex3, ESPMode::ThreadSafe> BoxConvex = MakeShared<FImplicitConvex3, ESPMode::ThreadSafe>(MakeBoxVerts(HalfSize), 0.0f);
+		TSharedPtr<FImplicitConvex3, ESPMode::ThreadSafe> BoxConvex = MakeShared<FImplicitConvex3, ESPMode::ThreadSafe>(MakeBoxVerts(FVec3(0), HalfSize), 0.0f);
 		return TImplicitObjectScaled<FImplicitConvex3>(BoxConvex, BoxScale, Margin);
 	}
 
diff --git a/Engine/Source/Programs/HeadlessChaos/Public/HeadlessChaosTestUtility.h b/Engine/Source/Programs/HeadlessChaos/Public/HeadlessChaosTestUtility.h
--- a/Engine/Source/Programs/HeadlessChaos/Public/HeadlessChaosTestUtility.h
+++ b/Engine/Source/Programs/HeadlessChaos/Public/HeadlessChaosTestUtility.h
@@ -199,6 +199,7 @@ namespace ChaosTest {
 
 
 	extern FImplicitConvex3 CreateConvexBox(const FVec3& BoxSize, const FReal Margin);
+	extern FImplicitConvex3 CreateConvexBox(const FVec3& BoxMin, const FVec3& BoxMax, const FReal Margin);
 	extern TImplicitObjectInstanced<FImplicitConvex3> CreateInstancedConvexBox(const FVec3& BoxSize, const FReal Margin);
 	extern TImplicitObjectScaled<FImplicitConvex3> CreateScaledConvexBox(const FVec3& BoxSize, const FVec3 BoxScale, const FReal Margin);
 }
diff --git a/Engine/Source/Runtime/Engine/Private/PhysicsEngine/BodyInstance.cpp b/Engine/Source/Runtime/Engine/Private/PhysicsEngine/BodyInstance.cpp
--- a/Engine/Source/Runtime/Engine/Private/PhysicsEngine/BodyInstance.cpp
+++ b/Engine/Source/Runtime/Engine/Private/PhysicsEngine/BodyInstance.cpp
@@ -3104,7 +3104,8 @@ void FBodyInstance::UpdateMassProperties()
 						FTransform RelTM;
 					};
 
-					//If we have welded children we must compute the mass properties of each individual body first and then combine them all together
+					// If we have welded children we must compute the mass properties of each individual body first and then combine them 
+					// all together because each welded body may have different density etc
 					TMap<FBodyInstance*, FWeldedBatch> BodyToShapes;
 
 					for (const FPhysicsShapeHandle& Shape : Shapes) //sort all welded children by their original bodies
@@ -3140,14 +3141,23 @@ void FBodyInstance::UpdateMassProperties()
 					{
 						const FBodyInstance* OwningBI = BodyShapesItr.Key;
 						const FWeldedBatch& WeldedBatch = BodyShapesItr.Value;
+
+						// The component scale is already built into the geometry, but if the user has set up a CoM
+						// modifier, it will need to be transformed by the component scale and the welded child's relative transform.
 						FTransform MassModifierTransform = WeldedBatch.RelTM;
-						MassModifierTransform.SetScale3D(MassModifierTransform.GetScale3D() * Scale3D);	//Ensure that any scaling that is done on the component is passed into the mass frame modifiers
+						MassModifierTransform.SetScale3D(MassModifierTransform.GetScale3D() * Scale3D);
 
 						Chaos::FMassProperties BodyMassProperties = BodyUtils::ComputeMassProperties(OwningBI, WeldedBatch.Shapes, MassModifierTransform);
 						SubMassProperties.Add(BodyMassProperties);
 					}
 
-					TotalMassProperties = Chaos::Combine(SubMassProperties);
+					// Combine all the child inertias
+					// NOTE: These leaves the inertia in diagonal form with a rotation of mass
+					if (SubMassProperties.Num() > 0)
+					{
+						TotalMassProperties = Chaos::Combine(SubMassProperties);
+					}
+
 #elif PHYSICS_INTERFACE_PHYSX
 					TArray<PxMassProperties> SubMassProperties;
 					TArray<PxTransform> MassTMs;
@@ -3169,27 +3179,32 @@ void FBodyInstance::UpdateMassProperties()
 				else
 				{
 					// If we have no shapes that affect mass we cannot compute the mass properties in a meaningful way.
-					if (Shapes.Num())
+					if (Shapes.Num() > 0)
 					{
-						//No children welded so just get this body's mass properties
-						FTransform MassModifierTransform(FQuat::Identity, FVector(0.f, 0.f, 0.f), Scale3D);	//Ensure that any scaling that is done on the component is passed into the mass frame modifiers
+						// The component scale is already built into the geometry, but if the user has set up a CoM
+						// modifier, it will need to be transformed by the component scale.
+						FTransform MassModifierTransform(FQuat::Identity, FVector(0.f, 0.f, 0.f), Scale3D);
 						TotalMassProperties = BodyUtils::ComputeMassProperties(this, Shapes, MassModifierTransform);
+
+						// Make the inertia diagonal and calculate the rotation of mass
+						TotalMassProperties.RotationOfMass = Chaos::TransformToLocalSpace(TotalMassProperties.InertiaTensor);
 					}
 				}
 
 				// #PHYS2 Refactor out PxMassProperties (Our own impl?)
 #if WITH_CHAOS
+				// Note: We expect the inertia to be diagonal at this point
 				// Only set mass properties if inertia tensor is valid. TODO Remove this once we track down cause of empty tensors.
+				// (This can happen if all shapes have bContributeToMass set to false which gives an empty Shapes array. There may be other ways).
 				const float InertiaTensorTrace = (TotalMassProperties.InertiaTensor.M[0][0] + TotalMassProperties.InertiaTensor.M[1][1] + TotalMassProperties.InertiaTensor.M[2][2]) / 3;
 				if (CHAOS_ENSURE(InertiaTensorTrace > SMALL_NUMBER))
 				{
-					const Chaos::TRotation<float, 3> Rotation = Chaos::TransformToLocalSpace(TotalMassProperties.InertiaTensor);
 					const FVector MassSpaceInertiaTensor(TotalMassProperties.InertiaTensor.M[0][0], TotalMassProperties.InertiaTensor.M[1][1], TotalMassProperties.InertiaTensor.M[2][2]);
 					FPhysicsInterface::SetMassSpaceInertiaTensor_AssumesLocked(Actor, MassSpaceInertiaTensor);
 
 					FPhysicsInterface::SetMass_AssumesLocked(Actor, TotalMassProperties.Mass);
 
-					FTransform Com(Rotation, TotalMassProperties.CenterOfMass);
+					FTransform Com(TotalMassProperties.RotationOfMass, TotalMassProperties.CenterOfMass);
 					FPhysicsInterface::SetComLocalPose_AssumesLocked(Actor, Com);
 				}
 #else
diff --git a/Engine/Source/Runtime/Engine/Private/PhysicsEngine/Experimental/ChaosInterfaceUtils.cpp b/Engine/Source/Runtime/Engine/Private/PhysicsEngine/Experimental/ChaosInterfaceUtils.cpp
--- a/Engine/Source/Runtime/Engine/Private/PhysicsEngine/Experimental/ChaosInterfaceUtils.cpp
+++ b/Engine/Source/Runtime/Engine/Private/PhysicsEngine/Experimental/ChaosInterfaceUtils.cpp
@@ -439,60 +439,22 @@ namespace ChaosInterface
 		return false;
 	}
 
-	void CalculateMassPropertiesFromShapeCollection(Chaos::FMassProperties& OutProperties, const TArray<FPhysicsShapeHandle>& InShapes, float InDensityKGPerCM)
+	void CalculateMassPropertiesFromShapeCollectionImp(
+		Chaos::FMassProperties& OutProperties, 
+		int32 InNumShapes, 
+		float InDensityKGPerCM,
+		const TArray<bool>& bContributesToMass,
+		TFunction<Chaos::FPerShapeData* (int32 ShapeIndex)> GetShapeDelegate)
 	{
-		float TotalMass = 0.f;
-		Chaos::FVec3 TotalCenterOfMass(0.f);
+		float TotalMass = 0;
+		float TotalVolume = 0;
+		Chaos::FVec3 TotalCenterOfMass(0);
 		TArray< Chaos::FMassProperties > MassPropertiesList;
-		for (const FPhysicsShapeHandle& ShapeHandle : InShapes)
-		{
-			if (const Chaos::FPerShapeData* Shape = ShapeHandle.Shape)
-			{
-				if (const Chaos::FImplicitObject* ImplicitObject = Shape->GetGeometry().Get())
-				{
-					FTransform WorldTransform(ShapeHandle.ActorRef->GetGameThreadAPI().R(), ShapeHandle.ActorRef->GetGameThreadAPI().X());
-					Chaos::FMassProperties MassProperties;
-					if (CalculateMassPropertiesOfImplicitType(MassProperties, WorldTransform, ImplicitObject, InDensityKGPerCM))
-					{
-						MassPropertiesList.Add(MassProperties);
-						TotalMass += MassProperties.Mass;
-						TotalCenterOfMass += MassProperties.CenterOfMass * MassProperties.Mass;
-					}
-				}
-			}
-		}
-
-		if (TotalMass > 0.f)
+		for (int32 ShapeIndex = 0; ShapeIndex < InNumShapes; ++ShapeIndex)
 		{
-			TotalCenterOfMass /= TotalMass;
-		}
+			const Chaos::FPerShapeData* Shape = GetShapeDelegate(ShapeIndex);
 
-		Chaos::PMatrix<float, 3, 3> Tensor;
-		if (MassPropertiesList.Num())
-		{
-			Tensor = Chaos::CombineWorldSpace(MassPropertiesList, InDensityKGPerCM).InertiaTensor;
-		}
-		else
-		{
-			// @todo : Add support for all types, but for now just hard code a unit sphere tensor {r:50cm} if the type was not processed
-			Tensor = Chaos::PMatrix<float, 3, 3>(5.24e5, 5.24e5, 5.24e5);
-			TotalMass = 523.f;
-		}
-
-		OutProperties.InertiaTensor = Tensor;
-		OutProperties.Mass = TotalMass;
-		OutProperties.CenterOfMass = TotalCenterOfMass;
-	}
-
-	void CalculateMassPropertiesFromShapeCollection(Chaos::FMassProperties& OutProperties, const Chaos::FShapesArray& InShapes, const TArray<bool>& bContributesToMass, float InDensityKGPerCM)
-	{
-		float TotalMass = 0.f;
-		Chaos::FVec3 TotalCenterOfMass(0.f);
-		TArray< Chaos::FMassProperties > MassPropertiesList;
-		for (int32 ShapeIndex = 0; ShapeIndex < InShapes.Num(); ++ShapeIndex)
-		{
-			const TUniquePtr<Chaos::FPerShapeData>& Shape = InShapes[ShapeIndex];
-			const bool bHassMass = (ShapeIndex < bContributesToMass.Num())? bContributesToMass[ShapeIndex] : true;
+			const bool bHassMass = (ShapeIndex < bContributesToMass.Num()) ? bContributesToMass[ShapeIndex] : true;
 			if (bHassMass)
 			{
 				if (const Chaos::FImplicitObject* ImplicitObject = Shape->GetGeometry().Get())
@@ -502,6 +464,7 @@ namespace ChaosInterface
 					{
 						MassPropertiesList.Add(MassProperties);
 						TotalMass += MassProperties.Mass;
+						TotalVolume += MassProperties.Volume;
 						TotalCenterOfMass += MassProperties.CenterOfMass * MassProperties.Mass;
 					}
 				}
@@ -514,20 +477,52 @@ namespace ChaosInterface
 		}
 
 		Chaos::PMatrix<float, 3, 3> Tensor;
+		Chaos::FRotation3 RotationOfMass = Chaos::FRotation3::Identity;
 		if (MassPropertiesList.Num())
 		{
 			Tensor = Chaos::CombineWorldSpace(MassPropertiesList, InDensityKGPerCM).InertiaTensor;
+			// NOTE: If multiple items in the list, rotation of mass will be zero, but if only 1 item is the list the item is returned directly and we may have a rotation of mass
+			Chaos::FMassProperties CombinedMassProperties = Chaos::CombineWorldSpace<float, 3>(MassPropertiesList);
+			Tensor = CombinedMassProperties.InertiaTensor;
+			RotationOfMass = CombinedMassProperties.RotationOfMass;
+			Tensor = Chaos::CombineWorldSpace(MassPropertiesList, InDensityKGPerCM).InertiaTensor;
 		}
 		else
 		{
 			// @todo : Add support for all types, but for now just hard code a unit sphere tensor {r:50cm} if the type was not processed
 			Tensor = Chaos::PMatrix<float, 3, 3>(5.24e5f, 5.24e5f, 5.24e5f);
 			TotalMass = 523.0f;
+			TotalVolume = 523000;
 		}
 
 		OutProperties.InertiaTensor = Tensor;
 		OutProperties.Mass = TotalMass;
+		OutProperties.Volume = TotalVolume;
 		OutProperties.CenterOfMass = TotalCenterOfMass;
+		OutProperties.RotationOfMass = RotationOfMass;
+	}
+
+
+	void CalculateMassPropertiesFromShapeCollection(Chaos::FMassProperties& OutProperties, const TArray<FPhysicsShapeHandle>& InShapes, float InDensityKGPerCM)
+	{
+		CalculateMassPropertiesFromShapeCollectionImp(
+			OutProperties,
+			InShapes.Num(),
+			InDensityKGPerCM,
+			TArray<bool>(),
+			[&InShapes](int32 ShapeIndex) { return InShapes[ShapeIndex].Shape; }
+		);
+	}
+
+	void CalculateMassPropertiesFromShapeCollection(Chaos::FMassProperties& OutProperties, const Chaos::FShapesArray& InShapes, const TArray<bool>& bContributesToMass, float InDensityKGPerCM)
+	{
+		CalculateMassPropertiesFromShapeCollectionImp(
+			OutProperties,
+			InShapes.Num(),
+			InDensityKGPerCM,
+			bContributesToMass,
+			[&InShapes](int32 ShapeIndex) { return InShapes[ShapeIndex].Get(); }
+		);
 	}
 
 #endif // WITH_CHAOS
diff --git a/Engine/Source/Runtime/Experimental/Chaos/Private/Chaos/MassProperties.cpp b/Engine/Source/Runtime/Experimental/Chaos/Private/Chaos/MassProperties.cpp
--- a/Engine/Source/Runtime/Experimental/Chaos/Private/Chaos/MassProperties.cpp
+++ b/Engine/Source/Runtime/Experimental/Chaos/Private/Chaos/MassProperties.cpp
@@ -173,14 +173,12 @@ namespace Chaos
 
 	FMassProperties Combine(const TArray<FMassProperties>& MPArray)
 	{
-		// todo(chaos) : should we be able to have different density per primitive instead of using this one default?
-		constexpr FReal DefaultDensity = (FReal)1.0;
-		FMassProperties NewMP = CombineWorldSpace(MPArray, (FReal)1.0);
+		FMassProperties NewMP = CombineWorldSpace(MPArray);
 		NewMP.RotationOfMass = TransformToLocalSpace(NewMP.InertiaTensor);
 		return NewMP;
 	}
 
-	FMassProperties CombineWorldSpace(const TArray<FMassProperties>& MPArray, FReal InDensityKGPerCM)
+	FMassProperties CombineWorldSpace(const TArray<FMassProperties>& MPArray)
 	{
 		check(MPArray.Num() > 0);
 		if (MPArray.Num() == 1)
@@ -199,7 +197,7 @@ namespace Chaos
 		NewMP.CenterOfMass /= NewMP.Mass;
 		for (const FMassProperties& Child : MPArray)
 		{
-			const FReal M = Child.Volume * InDensityKGPerCM;
+			const FReal M = Child.Mass;
 			const FVec3 ParentToChild = Child.CenterOfMass - NewMP.CenterOfMass;
 			const FReal P0 = ParentToChild[0];
 			const FReal P1 = ParentToChild[1];
diff --git a/Engine/Source/Runtime/Experimental/Chaos/Public/Chaos/ImplicitObjectScaled.h b/Engine/Source/Runtime/Experimental/Chaos/Public/Chaos/ImplicitObjectScaled.h
--- a/Engine/Source/Runtime/Experimental/Chaos/Public/Chaos/ImplicitObjectScaled.h
+++ b/Engine/Source/Runtime/Experimental/Chaos/Public/Chaos/ImplicitObjectScaled.h
@@ -6,6 +6,7 @@
 #include "Chaos/ImplicitFwd.h"
 #include "Chaos/ImplicitObject.h"
 #include "Chaos/Transform.h"
+#include "Chaos/Utilities.h"	// For ScaleInertia - pull that into mass utils
 #include "ChaosArchive.h"
 #include "Templates/ChooseClass.h"
 #include "Templates/EnableIf.h"
@@ -288,6 +289,21 @@ public:
 		return MObject->GetVertex(VertexIndex);
 	}
 
+	const FVec3 GetCenterOfMass() const
+	{
+		return MObject->GetCenterOfMass();
+	}
+
+	FRotation3 GetRotationOfMass() const
+	{
+		return GetRotationOfMass();
+	}
+
+	const FMatrix33 GetInertiaTensor(const FReal Mass) const
+	{
+		return MObject->GetInertiaTensor(Mass);
+	}
+
 protected:
 	ObjectType MObject;
 
@@ -336,18 +352,6 @@ public:
 		return MLocalBoundingBox;
 	}
 
-	const FReal GetVolume() const
-	{
-		// TODO: More precise volume!
-		return BoundingBox().GetVolume();
-	}
-
-	const FMatrix33 GetInertiaTensor(const FReal Mass) const
-	{
-		// TODO: More precise inertia!
-		return BoundingBox().GetInertiaTensor(Mass);
-	}
-
 protected:
 	FVec3 MScale;
 	FVec3 MInvScale;
@@ -806,16 +810,24 @@ public:
 		UpdateBounds();
 	}
 
+	const FReal GetVolume() const
+	{
+		return MScale.X * MScale.Y * MScale.Z * MObject->GetVolume();
+	}
+
 	const FVec3 GetCenterOfMass() const
 	{
-		// TODO: I'm not sure this is correct in all cases
 		return MScale * MObject->GetCenterOfMass();
 	}
 
 	FRotation3 GetRotationOfMass() const
 	{
-		// TODO: use the actual underlying object ( to be consistent with current implementation of GetInertiaTensor and GetVolume )
-		return BoundingBox().GetRotationOfMass();
+		return MObject->GetRotationOfMass();
+	}
+
+	const FMatrix33 GetInertiaTensor(const FReal Mass) const
+	{
+		return Utilities::ScaleInertia(MObject->GetInertiaTensor(Mass), MScale, false);
 	}
 
 	const ObjectType Object() const { return MObject; }
diff --git a/Engine/Source/Runtime/Experimental/Chaos/Public/Chaos/MassProperties.h b/Engine/Source/Runtime/Experimental/Chaos/Public/Chaos/MassProperties.h
--- a/Engine/Source/Runtime/Experimental/Chaos/Public/Chaos/MassProperties.h
+++ b/Engine/Source/Runtime/Experimental/Chaos/Public/Chaos/MassProperties.h
@@ -44,9 +44,17 @@ namespace Chaos
 	void CHAOS_API CalculateInertiaAndRotationOfMass(const FParticles& Vertices, const TSurfaces& Surfaces, const FReal Density, const FVec3& CenterOfMass,
 		FMatrix33& OutInertiaTensor, FRotation3& OutRotationOfMass);
 
+
+	// Combine a list of transformed inertia tensors into a single inertia. Also diagonalize the inertia
+	// and set the rotation of mass accordingly.
 	FMassProperties CHAOS_API Combine(const TArray<FMassProperties>& MPArray);
 
-	FMassProperties CHAOS_API CombineWorldSpace(const TArray<FMassProperties>& MPArray, FReal InDensityKGPerCM);
+	// Combine a list of transformed inertia tensors into a single inertia. 
+	// NOTE: If there is more than one item in the list, the output may be non-diagonaly and will have a zero rotation.
+	// If there is only 1 item in the list it will return it directly, so the rotation of mass may be non-zero.
+	template<class T, int d>
+	FMassProperties CHAOS_API CombineWorldSpace(const TArray<FMassProperties>& MPArray);
+
 
 	template <typename T, int d>
 	using TMassProperties UE_DEPRECATED(4.27, "Deprecated. this class is to be deleted, use FMassProperties instead") = FMassProperties;
