Chaos Physics: More CCD Fixes and improvements
- Swept constraints will now not be disabled prematurely
- In other cases cull distances are now reenabled and velocity expanded for improved performance
- CCD tests  added, updated, enabled and passing

diff --git a/Engine/Source/Programs/HeadlessChaos/Private/HeadlessChaosTestCCD.cpp b/Engine/Source/Programs/HeadlessChaos/Private/HeadlessChaosTestCCD.cpp
--- a/Engine/Source/Programs/HeadlessChaos/Private/HeadlessChaosTestCCD.cpp
+++ b/Engine/Source/Programs/HeadlessChaos/Private/HeadlessChaosTestCCD.cpp
@@ -19,13 +19,13 @@
 {
 	using namespace Chaos;
 
-	// CCD not implemented for sphere sphere
-	TYPED_TEST(AllEvolutions, DISABLED_CCDTests_CCDEnabled)
+	TYPED_TEST(AllEvolutions, CCDTests_CCDDisabled)
 	{
 		const FReal Dt = 1 / 30.0f;
 		const FReal Fps = 1 / Dt;
-		const FReal SphereRadius = 100; // cm
-		const FReal InitialSpeed = SphereRadius * 5 * Fps; // More than enough to tunnel
+		const FReal BoxHalfSize = 50; // cm
+		const FReal InitialSpeed = BoxHalfSize * 5 * Fps; // More than enough to tunnel
+		const FReal InitialPosition = BoxHalfSize * 2 + 30;
 		using TEvolution = TypeParam;
 		TPBDRigidsSOAs<FReal, 3> Particles;
 		THandleArray<FChaosPhysicsMaterial> PhysicalMaterials;
@@ -41,12 +41,10 @@
 		PhysicsMaterial->SleepCounterThreshold = 1000; // Don't sleep
 		PhysicsMaterial->Restitution = 1.0f;
 
-		// Create Sphere geometry (Radius = 100)
-		TUniquePtr<FImplicitObject> Sphere(new TSphere<FReal, 3>(FVec3(0, 0, 0), SphereRadius));
+		// Create box geometry 
+		TUniquePtr<FImplicitObject> SmallBox(new TBox<FReal, 3>(FVec3(-BoxHalfSize, -BoxHalfSize, -BoxHalfSize), FVec3(BoxHalfSize, BoxHalfSize, BoxHalfSize)));
+		Static->SetGeometry(MakeSerializable(SmallBox));
+		AppendDynamicParticleConvexBox(*Dynamic, FVec3(BoxHalfSize), 0.0f);
 
-		// Assign sphere geometry to both particles 
-		Static->SetGeometry(MakeSerializable(Sphere));
-		Dynamic->SetGeometry(MakeSerializable(Sphere));
-
 		Evolution.SetPhysicsMaterial(Dynamic, MakeSerializable(PhysicsMaterial));
 
@@ -56,18 +54,17 @@
 
 		// Positions and velocities
 		Static->X() = FVec3(0, 0, 0);
+		Dynamic->X() = FVec3(0, 0, InitialPosition); // Start 30cm above the static box
+		Dynamic->V() = FVec3(0, 0, -InitialSpeed);
 
-		Dynamic->X() = FVec3(0, 0, SphereRadius * 2 + 10);
-		//Dynamic->V() = FVec3(0, 0, -InitialSpeed);
-		
 		// The position of the static has changed and statics don't automatically update bounds, so update explicitly
-		Static->SetWorldSpaceInflatedBounds(Sphere->BoundingBox().TransformedAABB(TRigidTransform<FReal, 3>(Static->X(), Static->R())));
+		Static->SetWorldSpaceInflatedBounds(SmallBox->BoundingBox().TransformedAABB(TRigidTransform<FReal, 3>(Static->X(), Static->R())));
 
 		::ChaosTest::SetParticleSimDataToCollide({ Static,Dynamic });
 
-		Dynamic->SetCCDEnabled(true);
+		Dynamic->SetCCDEnabled(false);
 		Dynamic->SetGravityEnabled(false);
-		
+
 
 		Dynamic->V() = FVec3(0, 0, -InitialSpeed);
 
@@ -78,18 +75,77 @@
 		}
 
 		// Large error margin, we are testing CCD and not solver accuracy
+		// The Box should pass right through the other one without interacting
 		const FReal LargeErrorMargin = 10.0f;
-		EXPECT_GE(Dynamic->X()[2], SphereRadius * 2 - LargeErrorMargin);
+		EXPECT_NEAR(Dynamic->X()[2], InitialPosition - InitialSpeed * Dt, LargeErrorMargin);
 	}
 
-
-
-	TYPED_TEST(AllEvolutions,DISABLED_CCDTests_CCDDisabled)
+	TYPED_TEST(AllEvolutions,CCDTests_ConvexConvex)
 	{
 		const FReal Dt = 1 / 30.0f;
 		const FReal Fps = 1 / Dt;
+		const FReal BoxHalfSize = 50; // cm
+		const FReal InitialSpeed = BoxHalfSize * 5 * Fps; // More than enough to tunnel
+		using TEvolution = TypeParam;
+		TPBDRigidsSOAs<FReal, 3> Particles;
+		THandleArray<FChaosPhysicsMaterial> PhysicalMaterials;
+		TEvolution Evolution(Particles, PhysicalMaterials);
+		InitEvolutionSettings(Evolution);
+
+		// Create particles
+		TGeometryParticleHandle<FReal, 3>* Static = Evolution.CreateStaticParticles(1)[0];
+		TPBDRigidParticleHandle<FReal, 3>* Dynamic = Evolution.CreateDynamicParticles(1)[0];
+
+		// Set up physics material
+		TUniquePtr<FChaosPhysicsMaterial> PhysicsMaterial = MakeUnique<FChaosPhysicsMaterial>();
+		PhysicsMaterial->SleepCounterThreshold = 1000; // Don't sleep
+		PhysicsMaterial->Restitution = 1.0f;
+
+		// Create box geometry 
+		TUniquePtr<FImplicitObject> SmallBox(new TBox<FReal, 3>(FVec3(-BoxHalfSize, -BoxHalfSize, -BoxHalfSize), FVec3(BoxHalfSize, BoxHalfSize, BoxHalfSize)));
+		Static->SetGeometry(MakeSerializable(SmallBox));
+		AppendDynamicParticleConvexBox(*Dynamic, FVec3(BoxHalfSize), 0.0f);
+
+		Evolution.SetPhysicsMaterial(Dynamic, MakeSerializable(PhysicsMaterial));
+
+		const FReal Mass = 100000.0f;;
+		Dynamic->I() = FMatrix33(Mass, Mass, Mass);
+		Dynamic->InvI() = FMatrix33(1.0f / Mass, 1.0f / Mass, 1.0f / Mass);
+
+		// Positions and velocities
+		Static->X() = FVec3(0, 0, 0);
+		Dynamic->X() = FVec3(0, 0, BoxHalfSize * 2 + 30); // Start 30cm above the static box
+		Dynamic->V() = FVec3(0, 0, -InitialSpeed);
+
+		// The position of the static has changed and statics don't automatically update bounds, so update explicitly
+		Static->SetWorldSpaceInflatedBounds(SmallBox->BoundingBox().TransformedAABB(TRigidTransform<FReal, 3>(Static->X(), Static->R())));
+
+		::ChaosTest::SetParticleSimDataToCollide({ Static,Dynamic });
+
+		Dynamic->SetCCDEnabled(true);
+		Dynamic->SetGravityEnabled(false);
+
+
+		Dynamic->V() = FVec3(0, 0, -InitialSpeed);
+
+		for (int i = 0; i < 1; ++i)
+		{
+			Evolution.AdvanceOneTimeStep(Dt);
+			Evolution.EndFrame(Dt);
+		}
+
+		// Large error margin, we are testing CCD and not solver accuracy
+		const FReal LargeErrorMargin = 10.0f;
+		EXPECT_GE(Dynamic->X()[2], BoxHalfSize * 2 - LargeErrorMargin);
+	}
+
+	// CCD not implemented for sphere sphere
+	TYPED_TEST(AllEvolutions, DISABLED_CCDTests_SphereSphere)
+	{
+		const FReal Dt = 1 / 30.0f;
+		const FReal Fps = 1 / Dt;
 		const FReal SphereRadius = 100; // cm
-		const FReal InitialSpeed = SphereRadius * 10 * Fps; // More than enough to tunnel
+		const FReal InitialSpeed = SphereRadius * 5 * Fps; // More than enough to tunnel
 		using TEvolution = TypeParam;
 		TPBDRigidsSOAs<FReal, 3> Particles;
 		THandleArray<FChaosPhysicsMaterial> PhysicalMaterials;
@@ -121,17 +177,19 @@
 		// Positions and velocities
 		Static->X() = FVec3(0, 0, 0);
 
-		Dynamic->X() = FVec3(0, 0, SphereRadius * 2 + 100);
-
+		Dynamic->X() = FVec3(0, 0, SphereRadius * 2 + 10);
+		Dynamic->V() = FVec3(0, 0, -InitialSpeed);
+		
 		// The position of the static has changed and statics don't automatically update bounds, so update explicitly
 		Static->SetWorldSpaceInflatedBounds(Sphere->BoundingBox().TransformedAABB(TRigidTransform<FReal, 3>(Static->X(), Static->R())));
 
 		::ChaosTest::SetParticleSimDataToCollide({ Static,Dynamic });
 
-		Dynamic->SetCCDEnabled(false);
+		Dynamic->SetCCDEnabled(true);
 		Dynamic->SetGravityEnabled(false);
+		
 
-		Dynamic->V() = FVec3(0, -InitialSpeed, 0);
+		Dynamic->V() = FVec3(0, 0, -InitialSpeed);
 
 		for (int i = 0; i < 1; ++i)
 		{
@@ -141,22 +199,21 @@
 
 		// Large error margin, we are testing CCD and not solver accuracy
 		const FReal LargeErrorMargin = 10.0f;
-		EXPECT_NEAR(Dynamic->X()[2], SphereRadius * 2 + 10 - InitialSpeed * Dt, LargeErrorMargin);
+		EXPECT_GE(Dynamic->X()[2], SphereRadius * 2 - LargeErrorMargin);
 	}
 
 	
-	TYPED_TEST(AllEvolutions, DISABLED_CCDTests_BoxStayInsideBoxBoundaries)
+	// Bounce a box within a box container without penetrating the container
+	TYPED_TEST(AllEvolutions, CCDTests_BoxStayInsideBoxBoundaries)
 	{
 		const FReal Dt = 1 / 30.0f;
 		const FReal Fps = 1 / Dt;
-		const FReal SmallBoxSize = 100; // cm
-		const FReal ContainerBoxSize = 500; // cm
+		const FReal SmallBoxHalfSize = 50; // cm
+		const FReal ContainerBoxHalfSize = 250; // cm
 		const FReal ContainerWallThickness = 10; // cm
 		const int ContainerFaceCount = 6;
 
-		const FVec3 InitialVelocity = FVec3(0, 0, 750); // ContainerBoxSize * 5 * Fps; // More than enough to tunnel
+		const FVec3 InitialVelocity = FVec3(-ContainerBoxHalfSize * Fps * 10, 0 ,0);
 
-		
-
 		using TEvolution = TypeParam;
 		TPBDRigidsSOAs<FReal, 3> Particles;
@@ -171,18 +228,18 @@
 		// Set up physics material
 		TUniquePtr<FChaosPhysicsMaterial> PhysicsMaterial = MakeUnique<FChaosPhysicsMaterial>();
 		PhysicsMaterial->SleepCounterThreshold = 1000; // Don't sleep
-		PhysicsMaterial->Restitution = 0.0f;
+		PhysicsMaterial->Restitution = 0.2f;// Bounce against the walls a few times
 
 		// Create box geometry 
-		TUniquePtr<FImplicitObject> SmallBox(new TBox<FReal, 3>(FVec3(-SmallBoxSize/2, -SmallBoxSize/2, -SmallBoxSize/2), FVec3(SmallBoxSize/2, SmallBoxSize/2, SmallBoxSize/2)));
+		//TUniquePtr<FImplicitObject> SmallBox(new TBox<FReal, 3>(FVec3(-SmallBoxHalfSize, -SmallBoxHalfSize, -SmallBoxHalfSize), FVec3(SmallBoxHalfSize, SmallBoxHalfSize, SmallBoxHalfSize)));
+		AppendDynamicParticleConvexBox(*Dynamic, FVec3(SmallBoxHalfSize), 0.0f);
 
 		// Just use 3 (x2) boxes for the walls of the container (avoid rotation transforms for this test)
-		TUniquePtr<FImplicitObject> ContainerFaceX(new TBox<FReal, 3>(FVec3(-ContainerWallThickness / 2, -ContainerBoxSize / 2, -ContainerBoxSize / 2), FVec3(ContainerWallThickness / 2, ContainerBoxSize / 2, ContainerBoxSize / 2)));
-		TUniquePtr<FImplicitObject> ContainerFaceY(new TBox<FReal, 3>(FVec3(-ContainerBoxSize / 2, -ContainerWallThickness / 2, -ContainerBoxSize / 2), FVec3(ContainerBoxSize / 2, ContainerWallThickness / 2, ContainerBoxSize / 2)));
-		TUniquePtr<FImplicitObject> ContainerFaceZ(new TBox<FReal, 3>(FVec3(-ContainerBoxSize / 2, -ContainerBoxSize / 2, -ContainerWallThickness / 2), FVec3(ContainerBoxSize / 2, ContainerBoxSize / 2, ContainerWallThickness / 2)));		
+		TUniquePtr<FImplicitObject> ContainerFaceX(new TBox<FReal, 3>(FVec3(-ContainerWallThickness / 2, -ContainerBoxHalfSize, -ContainerBoxHalfSize), FVec3(ContainerWallThickness / 2, ContainerBoxHalfSize, ContainerBoxHalfSize)));
+		TUniquePtr<FImplicitObject> ContainerFaceY(new TBox<FReal, 3>(FVec3(-ContainerBoxHalfSize, -ContainerWallThickness / 2, -ContainerBoxHalfSize), FVec3(ContainerBoxHalfSize, ContainerWallThickness / 2, ContainerBoxHalfSize)));
+		TUniquePtr<FImplicitObject> ContainerFaceZ(new TBox<FReal, 3>(FVec3(-ContainerBoxHalfSize, -ContainerBoxHalfSize, -ContainerWallThickness / 2), FVec3(ContainerBoxHalfSize, ContainerBoxHalfSize, ContainerWallThickness / 2)));		
 
-		// Assign geometry to all particles 
-		Dynamic->SetGeometry(MakeSerializable(SmallBox));
+		
 
 		ContainerFaces[0]->SetGeometry(MakeSerializable(ContainerFaceX));
 		ContainerFaces[1]->SetGeometry(MakeSerializable(ContainerFaceX));
@@ -199,12 +256,12 @@
 		Dynamic->InvI() = FMatrix33(1.0f / Mass, 1.0f / Mass, 1.0f / Mass);
 
 		// Positions and velocities
-		ContainerFaces[0]->X() = FVec3(ContainerBoxSize / 2, 0, 0);
-		ContainerFaces[1]->X() = FVec3(-ContainerBoxSize / 2, 0, 0);
-		ContainerFaces[2]->X() = FVec3(0, ContainerBoxSize / 2, 0);
-		ContainerFaces[3]->X() = FVec3(0, -ContainerBoxSize / 2, 0);
-		ContainerFaces[4]->X() = FVec3(0, 0, ContainerBoxSize / 2);
-		ContainerFaces[5]->X() = FVec3(0, 0, -ContainerBoxSize / 2);
+		ContainerFaces[0]->X() = FVec3(ContainerBoxHalfSize, 0, 0);
+		ContainerFaces[1]->X() = FVec3(-ContainerBoxHalfSize, 0, 0);
+		ContainerFaces[2]->X() = FVec3(0, ContainerBoxHalfSize, 0);
+		ContainerFaces[3]->X() = FVec3(0, -ContainerBoxHalfSize, 0);
+		ContainerFaces[4]->X() = FVec3(0, 0, ContainerBoxHalfSize);
+		ContainerFaces[5]->X() = FVec3(0, 0, -ContainerBoxHalfSize);
 
 		Dynamic->X() = FVec3(0, 0, 0);
 
@@ -223,7 +280,8 @@
 		Dynamic->SetGravityEnabled(false);
 
 		Dynamic->V() = InitialVelocity;
-
+		///////////////////////////////////
+		// Test 1: bouncing from two opposite walls
 		for (int i = 0; i < 10; ++i)
 		{
 			Evolution.AdvanceOneTimeStep(Dt);
@@ -236,71 +294,56 @@
 		for (int axis = 0; axis < 3; axis++)
 		{
 			// If this failed, the dynamic cube escaped the air tight static container
-			EXPECT_LT(FMath::Abs(Dynamic->X()[axis]), ContainerBoxSize / 2 - ContainerWallThickness / 2 - SmallBoxSize / 2 + LargeErrorMargin);
+			const  FReal MaxCoordinates = ContainerBoxHalfSize - ContainerWallThickness / 2 - SmallBoxHalfSize;
+			EXPECT_LT(FMath::Abs(Dynamic->X()[axis]), MaxCoordinates + LargeErrorMargin);
 		}
-	}
+		/////////////////////////////////////////////
+		// Test2: Now launch to cube to a corner
+		Dynamic->V() = FVec3(-ContainerBoxHalfSize * Fps * 10);
+		Dynamic->W() = FVec3(0);
+		Dynamic->X() = FVec3(0);
+		Dynamic->P() = FVec3(0);
+		Dynamic->R() = TRotation<FReal, 3>::FromIdentity();
+		Dynamic->Q() = TRotation<FReal, 3>::FromIdentity();
 
-	// CCD not implemented for sphere sphere
-	TYPED_TEST(AllEvolutions, DISABLED_CCDTests_Shere_Sphere)
-	{
-		const FReal Dt = 1 / 30.0f;
-		const FReal Fps = 1 / Dt;
-		const FReal SphereRadius = 100; // cm
-		const FReal InitialSpeed = SphereRadius * 5 * Fps; // More than enough to tunnel
-		using TEvolution = TypeParam;
-		TPBDRigidsSOAs<FReal, 3> Particles;
-		THandleArray<FChaosPhysicsMaterial> PhysicalMaterials;
-		TEvolution Evolution(Particles, PhysicalMaterials);
-		InitEvolutionSettings(Evolution);
+		for (int i = 0; i < 10; ++i)
+		{
+			Evolution.AdvanceOneTimeStep(Dt);
+			Evolution.EndFrame(Dt);
+		}
 
-		// Create particles
-		TGeometryParticleHandle<FReal, 3>* Static = Evolution.CreateStaticParticles(1)[0];
-		TPBDRigidParticleHandle<FReal, 3>* Dynamic = Evolution.CreateDynamicParticles(1)[0];
+		// Check that we did not escape the box!
+		for (int axis = 0; axis < 3; axis++)
+		{
+			// If this failed, the dynamic cube escaped the air tight static container
+			const  FReal MaxCoordinates = ContainerBoxHalfSize - ContainerWallThickness / 2 - SmallBoxHalfSize;
+			EXPECT_LT(FMath::Abs(Dynamic->X()[axis]), MaxCoordinates + LargeErrorMargin);
+		}
 
-		// Set up physics material
-		TUniquePtr<FChaosPhysicsMaterial> PhysicsMaterial = MakeUnique<FChaosPhysicsMaterial>();
-		PhysicsMaterial->SleepCounterThreshold = 1000; // Don't sleep
+		/////////////////////////////////////////////////////////////////////
+		// Test 3: Now we give it something impossible to solve with PBD solver (restitution of 1, high velocities causes final position to be outside of box). 
+		// Make sure it still stays inside the box (albeit with a very reduced velocity) 
+		Dynamic->V() = InitialVelocity;
+		Dynamic->W() = FVec3(0);
+		Dynamic->X() = FVec3(0);
+		Dynamic->P() = FVec3(0);
+		Dynamic->R() = TRotation<FReal, 3>::FromIdentity();
+		Dynamic->Q() = TRotation<FReal, 3>::FromIdentity();
 		PhysicsMaterial->Restitution = 1.0f;
 
-		// Create Sphere geometry (Radius = 100)
-		TUniquePtr<FImplicitObject> Sphere(new TSphere<FReal, 3>(FVec3(0, 0, 0), SphereRadius));
-
-		// Assign sphere geometry to both particles 
-		Static->SetGeometry(MakeSerializable(Sphere));
-		Dynamic->SetGeometry(MakeSerializable(Sphere));
-
-		Evolution.SetPhysicsMaterial(Dynamic, MakeSerializable(PhysicsMaterial));
-
-		const FReal Mass = 100000.0f;;
-		Dynamic->I() = FMatrix33(Mass, Mass, Mass);
-		Dynamic->InvI() = FMatrix33(1.0f / Mass, 1.0f / Mass, 1.0f / Mass);
-
-		// Positions and velocities
-		Static->X() = FVec3(0, 0, 0);
-
-		Dynamic->X() = FVec3(0, 0, SphereRadius * 2 + 10);
-		//Dynamic->V() = FVec3(0, 0, -InitialSpeed);
-
-		// The position of the static has changed and statics don't automatically update bounds, so update explicitly
-		Static->SetWorldSpaceInflatedBounds(Sphere->BoundingBox().TransformedAABB(TRigidTransform<FReal, 3>(Static->X(), Static->R())));
-
-		::ChaosTest::SetParticleSimDataToCollide({ Static,Dynamic });
-
-		Dynamic->SetCCDEnabled(true);
-		Dynamic->SetGravityEnabled(false);
-
-
-		Dynamic->V() = FVec3(0, 0, -InitialSpeed);
-
-		for (int i = 0; i < 1; ++i)
+		for (int i = 0; i < 10; ++i)
 		{
 			Evolution.AdvanceOneTimeStep(Dt);
 			Evolution.EndFrame(Dt);
 		}
 
-		// Large error margin, we are testing CCD and not solver accuracy
-		const FReal LargeErrorMargin = 10.0f;
-		EXPECT_GE(Dynamic->X()[2], SphereRadius * 2 - LargeErrorMargin);
+		// Check that we did not escape the box!
+		for (int axis = 0; axis < 3; axis++)
+		{
+			// If this failed, the dynamic cube escaped the air tight static container
+			const  FReal MaxCoordinates = ContainerBoxHalfSize - ContainerWallThickness / 2 - SmallBoxHalfSize;
+			EXPECT_LT(FMath::Abs(Dynamic->X()[axis]), MaxCoordinates + LargeErrorMargin);
+		}
 	}
 }
 
diff --git a/Engine/Source/Runtime/Experimental/Chaos/Private/Chaos/CollisionResolution.cpp b/Engine/Source/Runtime/Experimental/Chaos/Private/Chaos/CollisionResolution.cpp
--- a/Engine/Source/Runtime/Experimental/Chaos/Private/Chaos/CollisionResolution.cpp
+++ b/Engine/Source/Runtime/Experimental/Chaos/Private/Chaos/CollisionResolution.cpp
@@ -237,11 +237,11 @@ namespace Chaos
 			}
 		}
 
-		void UpdateContactPoint(FRigidBodyPointContactConstraint& Constraint, const FContactPoint& ContactPoint, const FReal CullDistance, const FReal Dt)
+		void UpdateContactPoint(FRigidBodyPointContactConstraint& Constraint, const FContactPoint& ContactPoint, const FReal CullDistance, const FReal Dt, bool bAllowContactDisable = true)
 		{
 			// Permanently disable contacts beyond the CullDistance
 #if CHAOS_COLLISION_CREATE_BOUNDSCHECK
-			if (ContactPoint.Phi > CullDistance)
+			if (ContactPoint.Phi > CullDistance && bAllowContactDisable)
 			{
 				Constraint.SetDisabled(true);
 				return;
@@ -1804,7 +1804,7 @@ namespace Chaos
 		{
 			CONDITIONAL_SCOPE_CYCLE_COUNTER(STAT_Collisions_UpdateGenericConvexConvexConstraintSwept, ConstraintsDetailedStats);
 			FReal TOI = 1.0f;
-			UpdateContactPoint(Constraint, GenericConvexConvexContactPointSwept(Implicit0, WorldTransform0, Implicit1, WorldTransform1, Dir, Length, CullDistance, TOI), CullDistance, Dt);
+			UpdateContactPoint(Constraint, GenericConvexConvexContactPointSwept(Implicit0, WorldTransform0, Implicit1, WorldTransform1, Dir, Length, CullDistance, TOI), CullDistance, Dt, false);
 			SetSweptConstraintTOI(Particle0, TOI, Length, Dir, Constraint);
 		}
 
@@ -2263,7 +2263,7 @@ namespace Chaos
 			//Constraint.Manifold.ContactMoveSQRDistance = FMath::Max((NewContactPositionLocal0 - OriginalContactPositionLocal0).SizeSquared(), (NewContactPositionLocal1 - OriginalContactPositionLocal1).SizeSquared());
 		}
 
-		// Run collision detection for the swept constraits
+		// Run collision detection for the swept constraints
 		// NOTE: Transforms are world space shape transforms
 		
 		template<ECollisionUpdateType UpdateType>
diff --git a/Engine/Source/Runtime/Experimental/Chaos/Private/Chaos/PBDCollisionConstraintsContact.cpp b/Engine/Source/Runtime/Experimental/Chaos/Private/Chaos/PBDCollisionConstraintsContact.cpp
--- a/Engine/Source/Runtime/Experimental/Chaos/Private/Chaos/PBDCollisionConstraintsContact.cpp
+++ b/Engine/Source/Runtime/Experimental/Chaos/Private/Chaos/PBDCollisionConstraintsContact.cpp
@@ -266,13 +266,14 @@ namespace Chaos
 			FGenericParticleHandle Particle0 = FGenericParticleHandle(Constraint.Particle[0]);
 			FGenericParticleHandle Particle1 = FGenericParticleHandle(Constraint.Particle[1]);
 
-			Collisions::UpdateSwept(Constraint, ParticleParameters.CullDistance, IterationParameters.Dt);
-
+			const FReal ExpandedCullDistance = ParticleParameters.CullDistance + IterationParameters.Dt * Particle0->V().Size() * 2.0f; // Particle will move at most by this much during PBD solve
+			Collisions::UpdateSwept(Constraint, ExpandedCullDistance, IterationParameters.Dt);
+			
+			const FContactParticleParameters CCDParticleParamaters{ ExpandedCullDistance, ParticleParameters.RestitutionVelocityThreshold, true, ParticleParameters.Collided };
 			if (Constraint.TimeOfImpact == 1)
 			{
-				const FContactParticleParameters CCDParticleParamatersNoContactCulling{ ParticleParameters.CullDistance, ParticleParameters.RestitutionVelocityThreshold, false, ParticleParameters.Collided };
 				// If TOI = 1 (normal constraint) we don't want to split timestep at TOI.
-				ApplyImpl(Constraint, IterationParameters, CCDParticleParamatersNoContactCulling);
+				ApplyImpl(Constraint, IterationParameters, CCDParticleParamaters);
 				return;
 			}
 
@@ -284,7 +285,6 @@ namespace Chaos
 			const int32 PartialPairIterations = FMath::Max(IterationParameters.NumPairIterations, 2); // Do at least 2 pair iterations // @todo: Do we still need this?
 			const FContactIterationParameters IterationParametersPartialDT{ PartialDT, FakeIteration, IterationParameters.NumIterations, PartialPairIterations, IterationParameters.ApplyType, IterationParameters.NeedsAnotherIteration };
 			const FContactIterationParameters IterationParametersRemainingDT{ RemainingDT, FakeIteration, IterationParameters.NumIterations, IterationParameters.NumPairIterations, IterationParameters.ApplyType, IterationParameters.NeedsAnotherIteration };
-			const FContactParticleParameters CCDParticleParamaters{ ParticleParameters.CullDistance, ParticleParameters.RestitutionVelocityThreshold, false, ParticleParameters.Collided };
 
 			// Rewind P to TOI and Apply
 			Particle0->P() = FMath::Lerp(Particle0->X(), Particle0->P(), Constraint.TimeOfImpact);
