HeadlessChaos compilation fixes, new and changed tests in UE5 requiring updates to API changes.

diff --git a/Engine/Source/Programs/HeadlessChaos/Private/GeometryCollection/GeometryCollectionTestSimulationStreaming.cpp b/Engine/Source/Programs/HeadlessChaos/Private/GeometryCollection/GeometryCollectionTestSimulationStreaming.cpp
--- a/Engine/Source/Programs/HeadlessChaos/Private/GeometryCollection/GeometryCollectionTestSimulationStreaming.cpp
+++ b/Engine/Source/Programs/HeadlessChaos/Private/GeometryCollection/GeometryCollectionTestSimulationStreaming.cpp
@@ -52,7 +52,7 @@ namespace GeometryCollectionTest
 				// todo: this is crashing
 				UnitTest.Solver->RegisterObject(Collection->PhysObject);
 				UnitTest.Solver->AddDirtyProxy(Collection->PhysObject);
-				Collection->PhysObject->Initialize();
+				Collection->PhysObject->Initialize(UnitTest.Solver->GetEvolution());
 			}
 		}
 		
@@ -98,7 +98,7 @@ namespace GeometryCollectionTest
 				// todo: this is crashing
 				UnitTest.Solver->RegisterObject(Collection->PhysObject);
 				UnitTest.Solver->AddDirtyProxy(Collection->PhysObject);
-				Collection->PhysObject->Initialize();
+				Collection->PhysObject->Initialize(UnitTest.Solver->GetEvolution());
 			}
 		}
 		
@@ -159,7 +159,7 @@ namespace GeometryCollectionTest
 				// todo: this is crashing
 				UnitTest.Solver->RegisterObject(Collection->PhysObject);
 				UnitTest.Solver->AddDirtyProxy(Collection->PhysObject);
-				Collection->PhysObject->Initialize();
+				Collection->PhysObject->Initialize(UnitTest.Solver->GetEvolution());
 			}
 		}
 	
diff --git a/Engine/Source/Programs/HeadlessChaos/Private/HeadlessChaosTestJointPlasticity.cpp b/Engine/Source/Programs/HeadlessChaos/Private/HeadlessChaosTestJointPlasticity.cpp
--- a/Engine/Source/Programs/HeadlessChaos/Private/HeadlessChaosTestJointPlasticity.cpp
+++ b/Engine/Source/Programs/HeadlessChaos/Private/HeadlessChaosTestJointPlasticity.cpp
@@ -59,9 +59,9 @@ namespace ChaosTest {
 		
 	}
 
-	TYPED_TEST(AllEvolutions, JointPlasticity_UnderAngularPlasticityThreshold)
+	GTEST_TEST(AllEvolutions, JointPlasticity_UnderAngularPlasticityThreshold)
 	{
-		JointPlasticity_UnderAngularPlasticityThreshold<TypeParam>();
+		JointPlasticity_UnderAngularPlasticityThreshold<FPBDRigidsEvolution>();
 	}
 
 
@@ -116,9 +116,9 @@ namespace ChaosTest {
 
 	}
 
-	TYPED_TEST(AllEvolutions, JointPlasticity_OverAngularPlasticityThreshold)
+	GTEST_TEST(AllEvolutions, JointPlasticity_OverAngularPlasticityThreshold)
 	{
-		JointPlasticity_OverAngularPlasticityThreshold<TypeParam>();
+		JointPlasticity_OverAngularPlasticityThreshold<FPBDRigidsEvolution>();
 	}
 
 
@@ -171,9 +171,9 @@ namespace ChaosTest {
 
 	}
 
-	TYPED_TEST(AllEvolutions, JointPlasticity_UnderLinearPlasticityThreshold)
+	GTEST_TEST(AllEvolutions, JointPlasticity_UnderLinearPlasticityThreshold)
 	{
-		JointPlasticity_UnderLinearPlasticityThreshold<TypeParam>();
+		JointPlasticity_UnderLinearPlasticityThreshold<FPBDRigidsEvolution>();
 	}
 
 
@@ -226,11 +226,8 @@ namespace ChaosTest {
 		EXPECT_TRUE(Test.SOAs.GetDynamicParticles().X(0).Z > 0.f);
 	}
 
-	TYPED_TEST(AllEvolutions, JointPlasticity_OverLinearPlasticityThreshold)
+	GTEST_TEST(AllEvolutions, JointPlasticity_OverLinearPlasticityThreshold)
 	{
-		JointPlasticity_OverLinearPlasticityThreshold<TypeParam>();
+		JointPlasticity_OverLinearPlasticityThreshold<FPBDRigidsEvolution>();
 	}
-
-
-
-}
\ No newline at end of file
+}
diff --git a/Engine/Source/Programs/HeadlessChaos/Private/HeadlessChaosTestVehicles.cpp b/Engine/Source/Programs/HeadlessChaos/Private/HeadlessChaosTestVehicles.cpp
--- a/Engine/Source/Programs/HeadlessChaos/Private/HeadlessChaosTestVehicles.cpp
+++ b/Engine/Source/Programs/HeadlessChaos/Private/HeadlessChaosTestVehicles.cpp
@@ -30,7 +30,7 @@ namespace ChaosTest
 {
 	using namespace Chaos;
 
-	TYPED_TEST(AllTraits, VehicleTest_VehicleUtilityGraph)
+	GTEST_TEST(AllTraits, VehicleTest_VehicleUtilityGraph)
 	{
 		Chaos::FGraph Graph;
 
@@ -51,7 +51,7 @@ namespace ChaosTest
 		EXPECT_NEAR(Graph.EvaluateY(60.f), 90.f, Tolerance);
 	}
 
-	TYPED_TEST(AllTraits, VehicleTest_CalculateSlipAngle)
+	GTEST_TEST(AllTraits, VehicleTest_CalculateSlipAngle)
 	{
 		float SlipAngleFwds = FVehicleUtility::CalculateSlipAngle(0, 5);
 		float SlipAngleReverse = FVehicleUtility::CalculateSlipAngle(0, -5);
@@ -973,7 +973,7 @@ namespace ChaosTest
 		EXPECT_LT(SimulationTimeAccelTC, SimulationTimeAccelSpin);
 	}
 
-	TYPED_TEST(AllTraits, VehicleTest_WheelBrakingLongitudinalSlipNew)
+	GTEST_TEST(AllTraits, VehicleTest_WheelBrakingLongitudinalSlipNew)
 	{
 		FSimpleWheelConfig Setup;
 		Setup.ABSEnabled = false;
@@ -1068,7 +1068,7 @@ namespace ChaosTest
 		EXPECT_GT(StoppingDistanceLowFriction, StoppingDistanceA);
 	}
 
-	TYPED_TEST(AllTraits, VehicleTest_WheelAcceleratingLongitudinalSlipNew)
+	GTEST_TEST(AllTraits, VehicleTest_WheelAcceleratingLongitudinalSlipNew)
 	{
 		FSimpleWheelConfig Setup;
 		Setup.ABSEnabled = false;
@@ -1671,7 +1671,7 @@ namespace ChaosTest
 
 	}
 
-	TYPED_TEST(AllTraits, VehicleTest_Suspension_VaryingDelta)
+	GTEST_TEST(AllTraits, VehicleTest_Suspension_VaryingDelta)
 	{
 		FSimpleSuspensionConfig Setup;
 		Setup.MaxLength = 20.0f;
