HeadlessChaos: Update callsites to Geometry functions to match recent API updates

diff --git a/Engine/Source/Programs/HeadlessChaos/Private/GeometryCollection/GeometryCollectionTestFramework.cpp b/Engine/Source/Programs/HeadlessChaos/Private/GeometryCollection/GeometryCollectionTestFramework.cpp
--- a/Engine/Source/Programs/HeadlessChaos/Private/GeometryCollection/GeometryCollectionTestFramework.cpp
+++ b/Engine/Source/Programs/HeadlessChaos/Private/GeometryCollection/GeometryCollectionTestFramework.cpp
@@ -119,8 +119,9 @@ namespace GeometryCollectionTest
 		const FIndex3i& EdgeVertices = FIndex3i(4, 4, 4),
 		const int NumberOfMaterials = 2)
 	{
+		using GeomVector = UE::Geometry::FVector3<double>;
 		FGridBoxMeshGenerator BoxGen;
-		BoxGen.Box = FOrientedBox3d(FVector(0), Extents); // box center, box dimensions
+		BoxGen.Box = FOrientedBox3d(GeomVector::Zero(), GeomVector(Extents)); // box center, box dimensions
 		BoxGen.EdgeVertices = EdgeVertices;
 		BoxGen.Generate();
 
