Fix CIS Headless Chaos Tests

diff --git a/Engine/Source/Programs/HeadlessChaos/Private/HeadlessChaosTestVehicles.cpp b/Engine/Source/Programs/HeadlessChaos/Private/HeadlessChaosTestVehicles.cpp
--- a/Engine/Source/Programs/HeadlessChaos/Private/HeadlessChaosTestVehicles.cpp
+++ b/Engine/Source/Programs/HeadlessChaos/Private/HeadlessChaosTestVehicles.cpp
@@ -44,10 +44,12 @@ namespace ChaosTest
 		EXPECT_NEAR(Graph.EvaluateY(-5.f), 0.f, Tolerance);
 		EXPECT_NEAR(Graph.EvaluateY(5.f), 0.f, Tolerance);
 		EXPECT_NEAR(Graph.EvaluateY(10.f), 0.f, Tolerance);
+		EXPECT_NEAR(Graph.EvaluateY(15.f), 15.f, Tolerance);
 		EXPECT_NEAR(Graph.EvaluateY(40.f), 60.f, Tolerance);
 		EXPECT_NEAR(Graph.EvaluateY(25.f), 35.f, Tolerance);
 		EXPECT_NEAR(Graph.EvaluateY(45.f), 75.f, Tolerance);
 		EXPECT_NEAR(Graph.EvaluateY(50.f), 90.f, Tolerance);
+		EXPECT_NEAR(Graph.EvaluateY(55.f), 90.f, Tolerance);
 		EXPECT_NEAR(Graph.EvaluateY(60.f), 90.f, Tolerance);
 	}
 
@@ -811,179 +813,8 @@ namespace ChaosTest
 		Setup.BrakeEnabled = true;
 		Setup.EngineEnabled = true;
 		Setup.WheelRadius = 0.3f;
-		Setup.LongitudinalFrictionMultiplier = 1.0f;
-		Setup.LateralFrictionMultiplier = 1.0f;
-		Setup.SideSlipModifier = 0.7f;
-
-		FSimpleWheelSim Wheel(&Setup);
-
-		// Google braking distance at 30mph says 14m (not interested in the thinking distance part)
-		// So using a range 10-20 to ensure we are in the correct ballpark.
-		// If specified more accurately in the test, then modifying the code would break the test all the time.
-
-		// units meters
-		float Gravity = 9.8f;
-		float StoppingDistanceTolerance = 0.5f;
-		float DeltaTime = 1.f / 30.f;
-		float StoppingDistanceA = 0.f;
-		float SimulationTime = 0.0f;
-		Wheel.SetSurfaceFriction(RealWorldConsts::DryRoadFriction());
-
-		// reasonably ideal stopping distance - traveling forwards
-		Wheel.SetBrakeTorque(650);
-		SimulateBraking(Wheel, Gravity, MPHToMS(30.f), DeltaTime, StoppingDistanceA, SimulationTime);
-		EXPECT_GT(StoppingDistanceA, 10.f);
-		EXPECT_LT(StoppingDistanceA, 20.f);
-
-		// Changing to units of Cm should yield the same results
-		float MToCm = 100.0f;
-		float StoppingDistanceCm = 0.f;
-		Wheel.SetBrakeTorque(650 * MToCm * MToCm);
-		Wheel.SetWheelRadius(0.3f * MToCm);
-		SimulateBraking(Wheel, Gravity * MToCm, MPHToCmS(30.f), DeltaTime, StoppingDistanceCm, SimulationTime);
-		EXPECT_NEAR(StoppingDistanceCm, StoppingDistanceA * MToCm, StoppingDistanceTolerance * MToCm);
-
-		// traveling backwards stops just the same
-		float StoppingDistanceReverseDir = 0.f;
-		Wheel.SetWheelRadius(0.3f);
-		Wheel.SetBrakeTorque(650);
-		SimulateBraking(Wheel, Gravity, MPHToMS(-30.f), DeltaTime, StoppingDistanceReverseDir, SimulationTime);
-		EXPECT_GT(StoppingDistanceReverseDir, -20.f);
-		EXPECT_LT(StoppingDistanceReverseDir, -10.f);
-		EXPECT_LT(StoppingDistanceA - FMath::Abs(StoppingDistanceReverseDir), StoppingDistanceTolerance);
-
-		// Similar results with different delta time
-		float StoppingDistanceDiffDT = 0.f;
-		SimulateBraking(Wheel, Gravity, MPHToMS(30.f), DeltaTime * 0.25f, StoppingDistanceDiffDT, SimulationTime);
-		EXPECT_LT(StoppingDistanceA - StoppingDistanceDiffDT, StoppingDistanceTolerance);
-
-		// barely touching the brake - going to take longer to stop
-		float StoppingDistanceLightBraking = 0.f;
-		Wheel.SetBrakeTorque(150);
-		SimulateBraking(Wheel, Gravity, MPHToMS(30.f), DeltaTime, StoppingDistanceLightBraking, SimulationTime);
-		EXPECT_GT(StoppingDistanceLightBraking, StoppingDistanceA);
-
-		// locking the wheels / too much brake torque -> dynamic friction rather than static friction -> going to take longer to stop
-		float StoppingDistanceTooHeavyBreaking = 0.f;
-		Wheel.SetBrakeTorque(5000);
-		SimulateBraking(Wheel, Gravity, MPHToMS(30.f), DeltaTime, StoppingDistanceTooHeavyBreaking, SimulationTime);
-		EXPECT_GT(StoppingDistanceTooHeavyBreaking, StoppingDistanceA);
-
-		// Would have locked the wheels but ABS prevents skidding
-		Wheel.ABSEnabled = true;
-		float StoppingDistanceTooHeavyBreakingABS = 0.f;
-		Wheel.SetBrakeTorque(5000);
-		SimulateBraking(Wheel, Gravity, MPHToMS(30.f), DeltaTime, StoppingDistanceTooHeavyBreakingABS, SimulationTime);
-		EXPECT_LT(StoppingDistanceTooHeavyBreakingABS, StoppingDistanceTooHeavyBreaking);
-		Wheel.ABSEnabled = false;
-
-		// lower initial speed - stops more quickly
-		float StoppingDistanceLowerSpeed = 0.f;
-		Wheel.SetBrakeTorque(650);
-		SimulateBraking(Wheel, Gravity, MPHToMS(20.f), DeltaTime, StoppingDistanceLowerSpeed, SimulationTime);
-		EXPECT_LT(StoppingDistanceLowerSpeed, StoppingDistanceA);
-
-		// higher initial speed - stops more slowly
-		float StoppingDistanceHigherSpeed = 0.f;
-		Wheel.SetBrakeTorque(650);
-		SimulateBraking(Wheel, Gravity, MPHToMS(60.f), DeltaTime, StoppingDistanceHigherSpeed, SimulationTime);
-		EXPECT_GT(StoppingDistanceHigherSpeed, StoppingDistanceA);
-
-		// slippy surface - stops more slowly
-		float StoppingDistanceLowFriction = 0.f;
-		Wheel.SetSurfaceFriction(0.3f);
-		Wheel.SetBrakeTorque(650);
-		SimulateBraking(Wheel, Gravity, MPHToMS(30.f), DeltaTime, StoppingDistanceLowFriction, SimulationTime);
-		EXPECT_GT(StoppingDistanceLowFriction, StoppingDistanceA);		
-	}
-
-	GTEST_TEST(AllTraits, VehicleTest_WheelAcceleratingLongitudinalSlip)
-	{
-		FSimpleWheelConfig Setup;
-		Setup.ABSEnabled = false;
-		Setup.TractionControlEnabled = false;
-		Setup.BrakeEnabled = true;
-		Setup.EngineEnabled = true;
-		Setup.WheelRadius = 0.3f;
-		Setup.LongitudinalFrictionMultiplier = 1.0f;
-		Setup.LateralFrictionMultiplier = 1.0f;
-		Setup.SideSlipModifier = 0.7f;
-
-		FSimpleWheelSim Wheel(&Setup);
-
-		// There could be one frame extra computation on the acceleration since the last frame of brake is not using the full 
-		// amount of torque, it's clearing the last remaining velocity without pushing the vehicle back in the opposite direction
-		// Hence a slightly larger tolerance for the result
-		float AccelerationResultsTolerance = 1.0f; // meters
-
-		// units meters
-		float Gravity = 9.8f;
-		float DeltaTime = 1.f / 30.f;
-
-		float StoppingDistanceA = 0.f;
-		float SimulationTimeBrake = 0.0f;
-		Wheel.SetSurfaceFriction(RealWorldConsts::DryRoadFriction());
-
-		// How far & what time does it take to stop from 30MPH to rest
-		Wheel.SetBrakeTorque(650);
-		SimulateBraking(Wheel, Gravity, MPHToMS(30.0f), DeltaTime, StoppingDistanceA, SimulationTimeBrake);
-
-		// How far and what time does it take to accelerate from rest to 30MPH
-		float SimulationTimeAccel = 0.0f;
-		float DrivingDistanceA = 0.f;
-		Wheel.SetDriveTorque(650);
-		SimulateAccelerating(Wheel, Gravity, MPHToMS(30.0f), DeltaTime, DrivingDistanceA, SimulationTimeAccel);
-
-		// 0-30 MPH and 30-0 MPH should be the same if there's no slipping and accel torque was same as the brake torque run
-		EXPECT_LT(DrivingDistanceA - StoppingDistanceA, AccelerationResultsTolerance);
-		EXPECT_LT(SimulationTimeAccel - SimulationTimeBrake, AccelerationResultsTolerance);
-
-		// same range as braking from 30MPH
-		EXPECT_GT(DrivingDistanceA, 10.f);
-		EXPECT_LT(DrivingDistanceA, 20.f);
-
-		// Unreal units cm - Note for the same results the radius needs to remain at 0.3m and not also be scaled to 30(cm)
-		float SimulationTimeAccelCM = 0.0f;
-		float MtoCm = 100.0f;
-		float DrivingDistanceCM = 0.f;
-		Wheel.SetDriveTorque(650 * MtoCm * MtoCm);
-		Wheel.SetWheelRadius(0.3f * MtoCm);
-		SimulateAccelerating(Wheel, Gravity * MtoCm, MPHToCmS(30.0f), DeltaTime, DrivingDistanceCM, SimulationTimeAccelCM);
-		EXPECT_GT(DrivingDistanceCM, 10.f * MtoCm);
-		EXPECT_LT(DrivingDistanceCM, 20.f * MtoCm);
-		EXPECT_NEAR(SimulationTimeAccel, SimulationTimeAccelCM, AccelerationResultsTolerance);
-
-		float SimulationTimeAccelSpin = 0.0f;
-		float DrivingDistanceWheelspin = 0.f;
-		Wheel.SetWheelRadius(0.3f);
-		Wheel.SetDriveTorque(5000); // definitely cause wheel spin
-		SimulateAccelerating(Wheel, Gravity, MPHToMS(30.0f), DeltaTime, DrivingDistanceWheelspin, SimulationTimeAccelSpin);
-
-		// Enable traction control should be better than both of the above
-		float SimulationTimeAccelTC = 0.0f;
-		float DrivingDistanceTC = 0.f;
-		Wheel.TractionControlEnabled = true;
-		Wheel.SetDriveTorque(5000); // definitely cause wheel spin
-		SimulateAccelerating(Wheel, Gravity, MPHToMS(30.0f), DeltaTime, DrivingDistanceTC, SimulationTimeAccelTC);
-
-		// reaches target speed in a shorter distance
-		EXPECT_LT(DrivingDistanceTC, DrivingDistanceWheelspin);
-
-		// reaches speed quicker with TC on when wheel would be slipping from drive torque
-		EXPECT_LT(SimulationTimeAccelTC, SimulationTimeAccelSpin);
-	}
-
-	GTEST_TEST(AllTraits, VehicleTest_WheelBrakingLongitudinalSlipNew)
-	{
-		FSimpleWheelConfig Setup;
-		Setup.ABSEnabled = false;
-		Setup.TractionControlEnabled = false;
-		Setup.BrakeEnabled = true;
-		Setup.EngineEnabled = true;
-		Setup.WheelRadius = 0.3f;
-		Setup.NewSimulationPath = true;
 		Setup.FrictionMultiplier = 1.0f;
-		Setup.CorneringStiffness = 2.0f;
+		Setup.CorneringStiffness = 1000.0f;
 		Setup.SideSlipModifier = 0.7f;
 
 		FSimpleWheelSim Wheel(&Setup);
@@ -1068,7 +899,7 @@ namespace ChaosTest
 		EXPECT_GT(StoppingDistanceLowFriction, StoppingDistanceA);
 	}
 
-	GTEST_TEST(AllTraits, VehicleTest_WheelAcceleratingLongitudinalSlipNew)
+	GTEST_TEST(AllTraits, VehicleTest_WheelAcceleratingLongitudinalSlip)
 	{
 		FSimpleWheelConfig Setup;
 		Setup.ABSEnabled = false;
@@ -1076,9 +907,8 @@ namespace ChaosTest
 		Setup.BrakeEnabled = true;
 		Setup.EngineEnabled = true;
 		Setup.WheelRadius = 0.3f;
-		Setup.NewSimulationPath = true;
 		Setup.FrictionMultiplier = 1.0f;
-		Setup.CorneringStiffness = 2.0f;
+		Setup.CorneringStiffness = 1000.0f;
 		Setup.SideSlipModifier = 0.7f;
 
 		FSimpleWheelSim Wheel(&Setup);
@@ -1597,8 +1427,8 @@ namespace ChaosTest
 		Setup.BrakeEnabled = true;
 		Setup.EngineEnabled = true;
 		Setup.WheelRadius = 0.3f;
-		Setup.LongitudinalFrictionMultiplier = 1.0f;
-		Setup.LateralFrictionMultiplier = 1.0f;
+		Setup.FrictionMultiplier = 1.0f;
+		Setup.CorneringStiffness = 1000.0f;
 		Setup.SideSlipModifier = 0.7f;
 
 		FSimpleWheelSim Wheel(&Setup);
@@ -1642,33 +1472,9 @@ namespace ChaosTest
 		EXPECT_NEAR(ForceGenerated30FPS.X, ForceGenerated60FPS.X, ResultsTolerance);
 
 
-		Wheel.AccessSetup().NewSimulationPath = true;
-
 		Wheel.SetMatchingSpeed(Velocity.X);
 		Wheel.SetVehicleGroundSpeed(Velocity);
 		Wheel.SetDriveTorque(100.0f);
-
-		DeltaTime = 1.f / 30.f;
-		Wheel.Simulate(DeltaTime);
-		FVector NForceGenerated30FPS = Wheel.GetForceFromFriction();
-
-		Wheel.Simulate(DeltaTime);
-		FVector NForceGenerated30FPS_2 = Wheel.GetForceFromFriction();
-
-		DeltaTime = 1.f / 60.f;
-		Wheel.Simulate(DeltaTime);
-
-		FVector NForceGenerated60FPS = Wheel.GetForceFromFriction();
-
-		DeltaTime = 1.f / 50.f;
-		Wheel.Simulate(DeltaTime);
-
-		FVector NForceGenerated50FPS = Wheel.GetForceFromFriction();
-
-		EXPECT_NEAR(NForceGenerated30FPS.X, NForceGenerated30FPS_2.X, ResultsTolerance);
-		EXPECT_NEAR(NForceGenerated30FPS.X, NForceGenerated60FPS.X, ResultsTolerance);
-		EXPECT_NEAR(NForceGenerated30FPS.X, NForceGenerated50FPS.X, ResultsTolerance);
-
 	}
 
 	GTEST_TEST(AllTraits, VehicleTest_Suspension_VaryingDelta)
