CIS - fix headlesschaos unit tests

diff --git a/Engine/Source/Programs/HeadlessChaos/Private/HeadlessChaosTestSim.cpp b/Engine/Source/Programs/HeadlessChaos/Private/HeadlessChaosTestSim.cpp
--- a/Engine/Source/Programs/HeadlessChaos/Private/HeadlessChaosTestSim.cpp
+++ b/Engine/Source/Programs/HeadlessChaos/Private/HeadlessChaosTestSim.cpp
@@ -276,7 +276,6 @@ namespace ChaosTest {
 			FSingleParticlePhysicsProxy* Proxy = nullptr;
 			FVector PTX = FVector::ZeroVector;
 			int32 Hits = 0;
-			asdasd accumulate dt and check it's 4 * 1/60
 		};
 
 		FCallback* Callback = Solver->CreateAndRegisterSimCallbackObject_External<FCallback>();
