Fix headlesschaos not compiling

diff --git a/Engine/Source/Programs/HeadlessChaos/Private/HeadlessChaosTestSim.cpp b/Engine/Source/Programs/HeadlessChaos/Private/HeadlessChaosTestSim.cpp
--- a/Engine/Source/Programs/HeadlessChaos/Private/HeadlessChaosTestSim.cpp
+++ b/Engine/Source/Programs/HeadlessChaos/Private/HeadlessChaosTestSim.cpp
@@ -241,12 +241,12 @@ namespace ChaosTest {
 		EXPECT_TRUE(Dynamic1HasWokeAgain);
 	}
 
-	TYPED_TEST(AllTraits, SimTests_MidSubstepSleep)
+	GTEST_TEST(AllTraits, SimTests_MidSubstepSleep)
 	{
 		TSharedPtr<FImplicitObject, ESPMode::ThreadSafe> Sphere{new TSphere<FReal, 3>(FVec3(0), 10)};
 
 		FChaosSolversModule* Module = FChaosSolversModule::GetModule();
-		auto Solver = Module->CreateSolver<TypeParam>(nullptr);
+		auto Solver = Module->CreateSolver(nullptr);
 
 		FSingleParticlePhysicsProxy* Proxy = FSingleParticlePhysicsProxy::Create(Chaos::TPBDRigidParticle<FReal, 3>::CreateParticle());
 		Chaos::FRigidBodyHandle_External& Particle = Proxy->GetGameThreadAPI();
