Chaos: Constraints
Separated the idea of disabling a constraint vs. disconnecting a constraint from its particles during teardown and deletion.
- Added particle construction and destruction handshake for the suspension constraint.
- Added plumbing for DisconnectConstraint through the constraint containers and rules.
- Units tests for stressing the constraint states.

diff --git a/Engine/Source/Programs/HeadlessChaos/Private/HeadlessChaosTestGraph.cpp b/Engine/Source/Programs/HeadlessChaos/Private/HeadlessChaosTestGraph.cpp
--- a/Engine/Source/Programs/HeadlessChaos/Private/HeadlessChaosTestGraph.cpp
+++ b/Engine/Source/Programs/HeadlessChaos/Private/HeadlessChaosTestGraph.cpp
@@ -6,6 +6,10 @@
 #include "Chaos/PBDConstraintGraph.h"
 #include "Chaos/PBDRigidsEvolution.h"
 #include "Chaos/PBDRigidsEvolutionGBF.h"
+#include "Chaos/PBDPositionConstraints.h"
+#include "Chaos/PBDSuspensionConstraints.h"
+#include "Chaos/PBDRigidDynamicSpringConstraints.h"
+#include "Chaos/PBDJointConstraints.h"
 #include "Chaos/Utilities.h"
 #include "HeadlessChaos.h"
 #include "HeadlessChaosTestUtility.h"
@@ -1078,6 +1082,187 @@ namespace ChaosTest {
 		}
 	}
 
+	/**
+	 * Test validity of all Particle Handle pointers held by ConstraintGraph.
+	 */
+	class ConstraintGraphValidation
+	{
+	public:
+		static bool IsParticleHandleValid(const FGeometryParticleHandle* Handle)
+		{
+			// Null is valid because we are allowing holes in the Nodes array that get re-filled using a Free List of empty slots 
+			if (Handle == nullptr)
+			{
+				return true;
+			}
+
+			// We're really testing if Handle is pointing to a valid particle handle; the return bool is inconsequential.
+			if (Handle->HasBounds())
+			{
+				return true;
+			}
+
+			return true;
+		}
+
+		static void ValidateConstraintGraphParticleHandles(const FPBDConstraintGraph& ConstraintGraph)
+		{
+			// Are all keys, values in the ParticleToNodeIndex Valid?
+			const int32 NodeCount = ConstraintGraph.Nodes.Num();
+			for (const TPair<FGeometryParticleHandle*, int32>& HandleAndIndex : ConstraintGraph.ParticleToNodeIndex)
+			{
+				EXPECT_TRUE(IsParticleHandleValid(HandleAndIndex.Key));
+				EXPECT_GE(HandleAndIndex.Value, 0);
+				EXPECT_LT(HandleAndIndex.Value, NodeCount);
+			}
+
+			// Are all Islands holding valid particle handles?
+			const int32 NumIslands = ConstraintGraph.NumIslands();
+			for (int32 IslandIdx = 0; IslandIdx < NumIslands; ++IslandIdx)
+			{
+				const TArray<FGeometryParticleHandle*>& IslandParticles = ConstraintGraph.GetIslandParticles(IslandIdx);
+				for (const FGeometryParticleHandle* Handle : IslandParticles)
+				{
+					EXPECT_TRUE(IsParticleHandleValid(Handle));
+				}
+			}
+
+			// Are all Graph Nodes valid?
+			for (const FPBDConstraintGraph::FGraphNode& Node : ConstraintGraph.Nodes)
+			{
+				ensure(IsParticleHandleValid(Node.Particle));
+				EXPECT_TRUE(IsParticleHandleValid(Node.Particle));
+			}
+
+			// Are all Graph Edges valid?
+			for (const FPBDConstraintGraph::FGraphEdge& Edge : ConstraintGraph.Edges)
+			{
+				EXPECT_GE(Edge.FirstNode, 0);
+				EXPECT_LT(Edge.FirstNode, NodeCount);
+				EXPECT_GE(Edge.SecondNode, -1); // position and suspension constraints don't have a valid second node
+				EXPECT_LT(Edge.SecondNode, NodeCount);
+			}
+		}
+	};
+
+	/**
+	 * Create a particles variously constrained.
+	 * Eliminate blocks of particles and test that ConstraintGraph is not holding dangling particles.
+	 */
+	void GraphDestroyParticles()
+	{
+		// Create some particles - doesn't matter what position or other state they have
+		TArray<TGeometryParticleHandle<FReal, 3>*> AllParticles;
+
+		FPBDRigidsSOAs SOAs;
+		TArray<FPBDRigidParticleHandle*> Dynamics = SOAs.CreateDynamicParticles(500);
+		for (auto Dynamic : Dynamics) { AllParticles.Add(Dynamic); }
+
+		// Set up the particle graph
+		FPBDConstraintGraph Graph;
+		Graph.InitializeGraph(SOAs.GetNonDisabledView());
+		
+		const int32 NumParticles = AllParticles.Num();
+
+		// Put together constraint lists
+		TArray<TVec2<int32>> ConstrainedParticles0;
+		ConstrainedParticles0.Reserve(150);
+		TArray<TVec2<int32>> ConstrainedParticles1;
+		ConstrainedParticles1.Reserve(150);
+		TArray<TVec2<int32>> ConstrainedParticles2;
+		ConstrainedParticles2.Reserve(150);
+		for (int32 ConstraintIdx = 0; ConstraintIdx < 150; ++ConstraintIdx)
+		{
+			ConstrainedParticles0.Add({ (ConstraintIdx * 2) % NumParticles, (ConstraintIdx * 3) % NumParticles });
+			ConstrainedParticles1.Add({ (ConstraintIdx * 4) % NumParticles, (ConstraintIdx * 5) % NumParticles });
+			ConstrainedParticles2.Add({ (ConstraintIdx * 6) % NumParticles, (ConstraintIdx * 7) % NumParticles });
+		}
+		
+		// Create some position constraints
+		Graph.ReserveConstraints(ConstrainedParticles0.Num());
+		FPBDPositionConstraints PositionConstraints;
+		for (const TVec2<int32>& Indices : ConstrainedParticles0)
+		{
+			FPBDRigidParticleHandle* Particle0 = AllParticles[Indices[0]]->CastToRigidParticle();
+			auto* NewConstraint = PositionConstraints.AddConstraint(Particle0, FVec3(0, 0, 0 ));
+			Graph.AddConstraint(0, NewConstraint, TVec2<FGeometryParticleHandle*>(Particle0, nullptr));
+		}
+
+		// Create some suspension constraints
+		Graph.ReserveConstraints(ConstrainedParticles1.Num());
+		FPBDSuspensionConstraints SuspensionConstraints;
+		for (const TVec2<int32>& Indices : ConstrainedParticles1)
+		{
+			FPBDSuspensionSettings Settings;
+			FPBDRigidParticleHandle* Particle0 = AllParticles[Indices[0]]->CastToRigidParticle();
+			FPBDRigidParticleHandle* Particle1 = AllParticles[Indices[1]]->CastToRigidParticle();
+			auto* NewConstraint = SuspensionConstraints.AddConstraint(Particle0, FVec3(0, 0, 0), Settings);
+			Graph.AddConstraint(1, NewConstraint, TVec2<FGeometryParticleHandle*>(Particle0, Particle1));
+		}
+
+		// Create some dynamic spring constraints.
+		Graph.ReserveConstraints(ConstrainedParticles2.Num());
+		FPBDRigidDynamicSpringConstraints SpringConstraints;
+		for (const TVec2<int32>& Indices : ConstrainedParticles2)
+		{
+			FPBDRigidParticleHandle* Particle0 = AllParticles[Indices[0]]->CastToRigidParticle();
+			FPBDRigidParticleHandle* Particle1 = AllParticles[Indices[1]]->CastToRigidParticle();
+			auto* NewConstraint = SpringConstraints.AddConstraint(FPBDRigidDynamicSpringConstraints::FConstrainedParticlePair(Particle0, Particle1));
+			Graph.AddConstraint(2, NewConstraint, TVec2<FGeometryParticleHandle*>(Particle0, Particle1));
+		}
+		
+		ConstraintGraphValidation::ValidateConstraintGraphParticleHandles(Graph);
+
+		// Add new particles
+		TArray<TPBDRigidParticleHandle<FReal, 3>*> NewDynamics = SOAs.CreateDynamicParticles(20);
+		ConstraintGraphValidation::ValidateConstraintGraphParticleHandles(Graph);
+
+		// Remove some constraints
+		PositionConstraints.RemoveConstraint(0);
+		SuspensionConstraints.RemoveConstraint(0);
+		SpringConstraints.RemoveConstraint(0);
+		ConstraintGraphValidation::ValidateConstraintGraphParticleHandles(Graph);
+		
+		// Add new constraints
+		Graph.ReserveConstraints(ConstrainedParticles1.Num());
+		FPBDJointConstraints JointConstraints;
+		for (const TVec2<int32>& Indices : ConstrainedParticles1)
+		{
+			FPBDRigidParticleHandle* Particle0 = AllParticles[Indices[0]]->CastToRigidParticle();
+			FPBDRigidParticleHandle* Particle1 = AllParticles[Indices[1]]->CastToRigidParticle();
+			auto* NewConstraint = JointConstraints.AddConstraint(FPBDJointConstraints::FParticlePair(Particle0,Particle1), FRigidTransform3());
+			Graph.AddConstraint(3, NewConstraint, TVec2<FGeometryParticleHandle*>(Particle0, Particle1));
+		}
+		ConstraintGraphValidation::ValidateConstraintGraphParticleHandles(Graph);
+
+		// Remove a block of particles
+		for (int32 Idx = 0; Idx < 10; ++Idx)
+		{
+			Graph.RemoveParticle(Dynamics[Idx]);
+			SOAs.DestroyParticle(Dynamics[Idx]);
+		}
+		ConstraintGraphValidation::ValidateConstraintGraphParticleHandles(Graph);
+
+		// Disable a block of particles
+		for (int32 Idx = 10; Idx < 20; ++Idx)
+		{
+			Graph.DisableParticle({ Dynamics[Idx] });
+			SOAs.DisableParticle(Dynamics[Idx]);
+		}
+		ConstraintGraphValidation::ValidateConstraintGraphParticleHandles(Graph);
+
+		// Re-enable the block of particles
+		for (int32 Idx = 10; Idx < 20; ++Idx)
+		{
+			Graph.EnableParticle(Dynamics[Idx], nullptr);
+			SOAs.EnableParticle(Dynamics[Idx]);
+		}
+		ConstraintGraphValidation::ValidateConstraintGraphParticleHandles(Graph);
+		
+	}
+
+
+
 
 	TEST(GraphTests,TestGraphIslands)
 	{
@@ -1124,5 +1309,9 @@ namespace ChaosTest {
 
 		SUCCEED();
 	}
-}
 
+	TEST(GraphTests, TestParticleDestruction)
+	{
+		GraphDestroyParticles();
+	}
+}
diff --git a/Engine/Source/Runtime/Experimental/Chaos/Private/Chaos/PBDConstraintRule.cpp b/Engine/Source/Runtime/Experimental/Chaos/Private/Chaos/PBDConstraintRule.cpp
--- a/Engine/Source/Runtime/Experimental/Chaos/Private/Chaos/PBDConstraintRule.cpp
+++ b/Engine/Source/Runtime/Experimental/Chaos/Private/Chaos/PBDConstraintRule.cpp
@@ -50,13 +50,6 @@ namespace Chaos
 
 	}
 
-	template<class T_CONSTRAINTS>
-	void TPBDConstraintGraphRuleImpl<T_CONSTRAINTS>::DisableConstraints(const TSet<FGeometryParticleHandle*>& RemovedParticles)
-	{
-		Constraints.DisableConstraints(RemovedParticles);
-	}
-
-
 	template<class T_CONSTRAINTS>
 	int32 TPBDConstraintGraphRuleImpl<T_CONSTRAINTS>::NumConstraints() const
 	{ 
diff --git a/Engine/Source/Runtime/Experimental/Chaos/Private/Chaos/PBDJointConstraints.cpp b/Engine/Source/Runtime/Experimental/Chaos/Private/Chaos/PBDJointConstraints.cpp
--- a/Engine/Source/Runtime/Experimental/Chaos/Private/Chaos/PBDJointConstraints.cpp
+++ b/Engine/Source/Runtime/Experimental/Chaos/Private/Chaos/PBDJointConstraints.cpp
@@ -432,24 +432,29 @@ namespace Chaos
 		}
 	}
 
+
+
 	
-	void FPBDJointConstraints::DisableConstraints(const TSet<TGeometryParticleHandle<FReal, 3>*>& RemovedParticles)
+	void FPBDJointConstraints::DisconnectConstraints(const TSet<TGeometryParticleHandle<FReal, 3>*>& RemovedParticles)
 	{
-		for (TGeometryParticleHandle<FReal, 3>* RemovedParticle : RemovedParticles)
+		for (TGeometryParticleHandle<FReal, 3>*RemovedParticle : RemovedParticles)
 		{
 			for (FConstraintHandle* ConstraintHandle : RemovedParticle->ParticleConstraints())
 			{
-				ConstraintHandle->SetEnabled(false); // constraint lifespan is managed by the proxy
-				int ConstraintIndex = ConstraintHandle->GetConstraintIndex();
-				if (ConstraintIndex != INDEX_NONE)
+				if (ConstraintHandle->As<FPBDJointConstraintHandle>())
 				{
-					if (ConstraintParticles[ConstraintIndex][0] == RemovedParticle)
-					{
-						ConstraintParticles[ConstraintIndex][0] = nullptr;
-					}
-					if (ConstraintParticles[ConstraintIndex][1] == RemovedParticle)
+					ConstraintHandle->SetEnabled(false); // constraint lifespan is managed by the proxy
+					int ConstraintIndex = ConstraintHandle->GetConstraintIndex();
+					if (ConstraintIndex != INDEX_NONE)
 					{
-						ConstraintParticles[ConstraintIndex][1] = nullptr;
+						if (ConstraintParticles[ConstraintIndex][0] == RemovedParticle)
+						{
+							ConstraintParticles[ConstraintIndex][0] = nullptr;
+						}
+						if (ConstraintParticles[ConstraintIndex][1] == RemovedParticle)
+						{
+							ConstraintParticles[ConstraintIndex][1] = nullptr;
+						}
 					}
 				}
 			}
diff --git a/Engine/Source/Runtime/Experimental/Chaos/Private/PhysicsProxy/SuspensionConstraintProxy.cpp b/Engine/Source/Runtime/Experimental/Chaos/Private/PhysicsProxy/SuspensionConstraintProxy.cpp
--- a/Engine/Source/Runtime/Experimental/Chaos/Private/PhysicsProxy/SuspensionConstraintProxy.cpp
+++ b/Engine/Source/Runtime/Experimental/Chaos/Private/PhysicsProxy/SuspensionConstraintProxy.cpp
@@ -54,8 +54,8 @@ void FSuspensionConstraintPhysicsProxy::InitializeOnPhysicsThread(Chaos::FPBDRig
 			Chaos::TGeometryParticleHandle<Chaos::FReal, 3>* Handle0 = GetParticleHandleFromProxy(BasePairs[0]);
 			if (Handle0)
 			{
-				Handle = SuspensionConstraints.AddConstraint(Handle0, Constraint->GetLocation()
-					, SuspensionSettingsBuffer);
+				Handle = SuspensionConstraints.AddConstraint(Handle0, Constraint->GetLocation(), SuspensionSettingsBuffer);
+				Handle0->AddConstraintHandle(Handle);
 			}
 			
 		}
diff --git a/Engine/Source/Runtime/Experimental/Chaos/Public/Chaos/PBDConstraintContainer.h b/Engine/Source/Runtime/Experimental/Chaos/Public/Chaos/PBDConstraintContainer.h
--- a/Engine/Source/Runtime/Experimental/Chaos/Public/Chaos/PBDConstraintContainer.h
+++ b/Engine/Source/Runtime/Experimental/Chaos/Public/Chaos/PBDConstraintContainer.h
@@ -20,6 +20,7 @@ namespace Chaos
 
 		virtual void SetConstraintEnabled(int32 ConstraintIndex, bool bEnabled) { }
 		virtual bool IsConstraintEnabled(int32 ConstraintIndex) const { return true; }
+		virtual void DisconnectConstraints(const TSet<TGeometryParticleHandle<FReal, 3>*>&) {}
 
 	protected:
 		// friend access to the Constraint Handle's container API
diff --git a/Engine/Source/Runtime/Experimental/Chaos/Public/Chaos/PBDConstraintGraph.h b/Engine/Source/Runtime/Experimental/Chaos/Public/Chaos/PBDConstraintGraph.h
--- a/Engine/Source/Runtime/Experimental/Chaos/Public/Chaos/PBDConstraintGraph.h
+++ b/Engine/Source/Runtime/Experimental/Chaos/Public/Chaos/PBDConstraintGraph.h
@@ -3,6 +3,11 @@
 
 #include "ParticleHandle.h"
 
+namespace ChaosTest
+{
+	class ConstraintGraphValidation;
+}
+
 namespace Chaos
 {
 	template<typename T, int d>
@@ -37,6 +42,9 @@ namespace Chaos
 	public:
 		friend class FPBDConstraintColor;
 
+		//Make friend with unit test code so we can verify some behavior
+		friend class ChaosTest::ConstraintGraphValidation;
+
 		/** Information required to map a graph edge back to its constraint */
 		struct FConstraintData
 		{
diff --git a/Engine/Source/Runtime/Experimental/Chaos/Public/Chaos/PBDConstraintRule.h b/Engine/Source/Runtime/Experimental/Chaos/Public/Chaos/PBDConstraintRule.h
--- a/Engine/Source/Runtime/Experimental/Chaos/Public/Chaos/PBDConstraintRule.h
+++ b/Engine/Source/Runtime/Experimental/Chaos/Public/Chaos/PBDConstraintRule.h
@@ -153,7 +153,19 @@ namespace Chaos
 		virtual void SetUseContactGraph(const bool InUseContactGraph) {}
 
 		/** Disable all constraints associated with the specified particles */
-		virtual void DisableConstraints(const TSet<TGeometryParticleHandle<FReal, 3>*>& RemovedParticles) { }
+		virtual void SetConstraintsEnabled(const TSet<TGeometryParticleHandle<FReal, 3>*>& DisabledParticles, bool bInEnabled) 
+		{ 
+			for (TGeometryParticleHandle<FReal, 3>*DisabledParticle : DisabledParticles)
+			{
+				for( FConstraintHandle* Constraint : DisabledParticle->ParticleConstraints() )
+				{
+					Constraint->SetEnabled(bInEnabled);
+				}
+			}
+		}
+
+		/** Disconnect all constraints associated with the specified particles */
+		virtual void DisconnectConstraints(const TSet<TGeometryParticleHandle<FReal, 3>*>& RemovedParticles) {  }
 
 		/** Remove all constraints */
 		virtual void ResetConstraints() {}
@@ -199,8 +211,12 @@ namespace Chaos
 		virtual void UpdatePositionBasedState(const FReal Dt) override;
 
 		virtual void AddToGraph() override;
-		
-		virtual void DisableConstraints(const TSet<TGeometryParticleHandle<FReal, 3>*>& RemovedParticles) override;
+
+		/** Disconnect all constraints associated with the specified particles */
+		virtual void DisconnectConstraints(const TSet<TGeometryParticleHandle<FReal, 3>*>& RemovedParticles)
+		{
+			Constraints.DisconnectConstraints(RemovedParticles);
+		}
 
 		virtual int32 NumConstraints() const override;
 
diff --git a/Engine/Source/Runtime/Experimental/Chaos/Public/Chaos/PBDJointConstraints.h b/Engine/Source/Runtime/Experimental/Chaos/Public/Chaos/PBDJointConstraints.h
--- a/Engine/Source/Runtime/Experimental/Chaos/Public/Chaos/PBDJointConstraints.h
+++ b/Engine/Source/Runtime/Experimental/Chaos/Public/Chaos/PBDJointConstraints.h
@@ -117,9 +117,13 @@ namespace Chaos
 		void RemoveConstraints(const TSet<TGeometryParticleHandle<FReal, 3>*>& RemovedParticles) {}
 
 		/*
-		* Disable the constraints attached to the input particles. This will set the constrained Particle elements to nullptr and set the Enable flag to false
+		* Disconnect the constraints from the attached input particles. 
+		* This will set the constrained Particle elements to nullptr and 
+		* set the Enable flag to false.
+		* 
+		* The constraint is unuseable at this point and pending deletion. 
 		*/
-		void DisableConstraints(const TSet<TGeometryParticleHandle<FReal, 3>*>& RemovedParticles);
+		void DisconnectConstraints(const TSet<TGeometryParticleHandle<FReal, 3>*>& RemovedParticles);
 
 		/*
 		 * Whether the constraint is enabled
diff --git a/Engine/Source/Runtime/Experimental/Chaos/Public/Chaos/PBDRigidsEvolution.h b/Engine/Source/Runtime/Experimental/Chaos/Public/Chaos/PBDRigidsEvolution.h
--- a/Engine/Source/Runtime/Experimental/Chaos/Public/Chaos/PBDRigidsEvolution.h
+++ b/Engine/Source/Runtime/Experimental/Chaos/Public/Chaos/PBDRigidsEvolution.h
@@ -328,7 +328,7 @@ public:
 		RemoveParticleFromAccelerationStructure(*Particle);
 		Particles.DisableParticle(Particle);
 		// disable constraint before disabling the particle in the graph because of dependencies 
-		RemoveConstraintsFromConstraintGraph(Particle->ParticleConstraints());
+		DisableConstraints({ Particle });
 		ConstraintGraph.DisableParticle(Particle);
 	}
 
@@ -338,14 +338,7 @@ public:
 	{
 		for (FGeometryParticleHandle* Particle : ParticlesIn)
 		{
-			RemoveParticleFromAccelerationStructure(*Particle);
-			Particles.DisableParticle(Particle);
-			RemoveConstraintsFromConstraintGraph(Particle->ParticleConstraints());
-		}
-		// make sure we do this has a second loop to make sure constraints have valid particles when removed from the constraint graph in the first loop
-		for (FGeometryParticleHandle* Particle : ParticlesIn)
-		{
-			ConstraintGraph.DisableParticle(Particle);
+			DisableParticle(Particle);
 		}
 	}
 
@@ -396,7 +389,7 @@ public:
 	{
 		RemoveParticleFromAccelerationStructure(*Particle);
 		UniqueIndicesPendingRelease.Add(Particle->UniqueIdx());
-		DisableConstraints(TSet<FGeometryParticleHandle*>({ Particle }));
+		DisconnectConstraints(TSet<FGeometryParticleHandle*>({ Particle }));
 		ConstraintGraph.RemoveParticle(Particle);
 		Particles.DestroyParticle(Particle);
 	}
@@ -444,18 +437,6 @@ public:
 		}
 	}
 
-	CHAOS_API void DisableParticles(const TSet<FGeometryParticleHandle*>& InParticles)
-	{
-		for (FGeometryParticleHandle* Particle : InParticles)
-		{
-			Particles.DisableParticle(Particle);
-			RemoveParticleFromAccelerationStructure(*Particle);
-		}
-
-		ConstraintGraph.DisableParticles(InParticles);
-		DisableConstraints(InParticles);
-	}
-
 	CHAOS_API void WakeIslands()
 	{
 		TArray<int32> UniqueIslands;
@@ -508,10 +489,10 @@ public:
 		}
 	}
 
-	/** Disable constraints from a set of particles to be removed (or destroyed) 
+	/** Disconnect constraints from a set of particles to be removed (or destroyed) 
 	* this will set the constraints to Enbaled = false and set their respective bodies handles to nullptr
 	*/
-	CHAOS_API void DisableConstraints(const TSet<FGeometryParticleHandle*>& RemovedParticles)
+	CHAOS_API void DisconnectConstraints(const TSet<FGeometryParticleHandle*>& RemovedParticles)
 	{
 		for (FGeometryParticleHandle* ParticleHandle : RemovedParticles)
 		{
@@ -520,7 +501,24 @@ public:
 
 		for (FPBDConstraintGraphRule* ConstraintRule : ConstraintRules)
 		{
-			ConstraintRule->DisableConstraints(RemovedParticles);
+			ConstraintRule->DisconnectConstraints(RemovedParticles);
+		}
+	}
+
+	/** Disconnect constraints from a set of particles to be removed (or destroyed)
+	* this will set the constraints to Enbaled = false, but leave connections to the 
+	* particles. 
+	*/
+	CHAOS_API void DisableConstraints(const TSet<FGeometryParticleHandle*>& DisabledParticles)
+	{
+		for (FGeometryParticleHandle* ParticleHandle : DisabledParticles)
+		{
+			RemoveConstraintsFromConstraintGraph(ParticleHandle->ParticleConstraints());
+		}
+
+		for (FPBDConstraintGraphRule* ConstraintRule : ConstraintRules)
+		{
+			ConstraintRule->SetConstraintsEnabled(DisabledParticles, false);
 		}
 	}
 
diff --git a/Engine/Source/Runtime/Experimental/Chaos/Public/Chaos/PBDSuspensionConstraints.h b/Engine/Source/Runtime/Experimental/Chaos/Public/Chaos/PBDSuspensionConstraints.h
--- a/Engine/Source/Runtime/Experimental/Chaos/Public/Chaos/PBDSuspensionConstraints.h
+++ b/Engine/Source/Runtime/Experimental/Chaos/Public/Chaos/PBDSuspensionConstraints.h
@@ -88,16 +88,33 @@ namespace Chaos
 		void RemoveConstraint(int ConstraintIndex);
 
 
-		/**
-		 * Disabled the specified constraint.
-		 */
-		void DisableConstraints(const TSet<TGeometryParticleHandle<FReal, 3>*>& RemovedParticles)
+		/*
+		* Disconnect the constraints from the attached input particles.
+		* This will set the constrained Particle elements to nullptr and
+		* set the Enable flag to false.
+		*
+		* The constraint is unuseable at this point and pending deletion.
+		*/
+
+		void DisconnectConstraints(const TSet<TGeometryParticleHandle<FReal, 3>*>& RemovedParticles)
 		{
-			for (TGeometryParticleHandle<FReal, 3>* RemovedParticle : RemovedParticles)
+			for (FGeometryParticleHandle* RemovedParticle : RemovedParticles)
 			{
 				for (FConstraintHandle* ConstraintHandle : RemovedParticle->ParticleConstraints())
 				{
-					ConstraintHandle->SetEnabled(false); // constraint lifespan is managed by the proxy
+					if (ConstraintHandle->As<FPBDSuspensionConstraintHandle>())
+					{
+						ConstraintHandle->SetEnabled(false); // constraint lifespan is managed by the proxy
+
+						int ConstraintIndex = ConstraintHandle->GetConstraintIndex();
+						if (ConstraintIndex != INDEX_NONE)
+						{
+							if (ConstrainedParticles[ConstraintIndex] == RemovedParticle)
+							{
+								ConstrainedParticles[ConstraintIndex] = nullptr;
+							}
+						}
+					}
 				}
 			}
 		}
