Chaos - fix headless chaos tests after disabling manifolds for level sets
- increase threshold in GeometryCollection_CollisionResolution_SimplicialTetrahedronWithNonUniformMassToFloor
- don't assume we always have a collision in GeometryCollection_Solver_CollisionEventFilter

diff --git a/Engine/Source/Programs/HeadlessChaos/Private/GeometryCollection/GeometryCollectionTestCollisionResolution.cpp b/Engine/Source/Programs/HeadlessChaos/Private/GeometryCollection/GeometryCollectionTestCollisionResolution.cpp
--- a/Engine/Source/Programs/HeadlessChaos/Private/GeometryCollection/GeometryCollectionTestCollisionResolution.cpp
+++ b/Engine/Source/Programs/HeadlessChaos/Private/GeometryCollection/GeometryCollectionTestCollisionResolution.cpp
@@ -529,7 +529,8 @@ namespace GeometryCollectionTest
 			//const FReal ExpectedRestingDistance = UnitTest.Solver->GetEvolution()->GetGravityForces().GetAcceleration().Size() * UnitTest.Dt * UnitTest.Dt;
 
 			// The error in the resting distance depends on the number of pushout iterations, which is very low by default
-			const FReal RestingDistanceTolerance = 1.0f;
+			// It also depends on whether manifolds are used for collision (no manifolds means larger errors)
+			const FReal RestingDistanceTolerance = 2.0f;
 
 			// validate the tetahedron collides and moved away from the static floor
 			FVec3 RestTranslation = Collection->RestCollection->Transform[0].GetTranslation();
diff --git a/Engine/Source/Programs/HeadlessChaos/Private/GeometryCollection/GeometryCollectionTestSimulationSolver.cpp b/Engine/Source/Programs/HeadlessChaos/Private/GeometryCollection/GeometryCollectionTestSimulationSolver.cpp
--- a/Engine/Source/Programs/HeadlessChaos/Private/GeometryCollection/GeometryCollectionTestSimulationSolver.cpp
+++ b/Engine/Source/Programs/HeadlessChaos/Private/GeometryCollection/GeometryCollectionTestSimulationSolver.cpp
@@ -319,7 +319,7 @@ using namespace ChaosTest;
 			if (Impact)
 			{
 				// any objects with a mass of less than 6 are removed from returned collision data
-				EXPECT_EQ(AllCollisionsArray.Num(), 4);
+				EXPECT_LE(AllCollisionsArray.Num(), 4);
 
 				for (const auto& Collision : AllCollisionsArray)
 				{ 
