Update plasticity-related code:
- Update HeadlessChaos tests
- Change warning to veryverbose

diff --git a/Engine/Source/Programs/HeadlessChaos/Private/HeadlessChaosTestJointPlasticity.cpp b/Engine/Source/Programs/HeadlessChaos/Private/HeadlessChaosTestJointPlasticity.cpp
--- a/Engine/Source/Programs/HeadlessChaos/Private/HeadlessChaosTestJointPlasticity.cpp
+++ b/Engine/Source/Programs/HeadlessChaos/Private/HeadlessChaosTestJointPlasticity.cpp
@@ -127,7 +127,7 @@ namespace ChaosTest {
 	template <typename TEvolution>
 	void JointPlasticity_UnderLinearPlasticityThreshold()
 	{
-		const FReal PlasticityRatio = 0.3;
+		const FReal PlasticityLimit = 10;
 		const int32 NumIterations = 1;
 		const FReal Gravity = 980;
 		const FReal Dt = 0.01f;
@@ -140,12 +140,13 @@ namespace ChaosTest {
 		Test.JointSettings[0].AngularMotionTypes = { EJointMotionType::Locked, EJointMotionType::Locked , EJointMotionType::Locked };
 		Test.JointSettings[0].LinearMotionTypes = { EJointMotionType::Limited, EJointMotionType::Limited , EJointMotionType::Limited };
 		Test.JointSettings[0].bSoftLinearLimitsEnabled = true;
+		Test.JointSettings[0].bLinearPositionDriveEnabled = { true, true, true };
 		Test.JointSettings[0].LinearLimit = 0;
 		Test.JointSettings[0].LinearSoftForceMode = EJointForceMode::Force;
-		Test.JointSettings[0].SoftLinearStiffness = 100000;
-		Test.JointSettings[0].SoftLinearDamping = 100;
+		Test.JointSettings[0].SoftLinearStiffness = 1000000;
+		Test.JointSettings[0].SoftLinearDamping = 1000;
 
-		Test.JointSettings[0].LinearPlasticityLimit = PlasticityRatio;
+		Test.JointSettings[0].LinearPlasticityLimit = PlasticityLimit;
 
 		Test.Create();
 		Test.AddParticleBox(FVec3(0, 0, 100), FRotation3::Identity, FVec3(10.f), 100.f);
@@ -182,7 +183,7 @@ namespace ChaosTest {
 	template <typename TEvolution>
 	void JointPlasticity_OverLinearPlasticityThreshold()
 	{
-		const FReal PlasticityRatio = 0.15;
+		const FReal PlasticityLimit = 1;
 		const int32 NumIterations = 1;
 		const FReal Gravity = 980;
 		const FReal Dt = 0.01f;
@@ -195,12 +196,13 @@ namespace ChaosTest {
 		Test.JointSettings[0].AngularMotionTypes = { EJointMotionType::Locked, EJointMotionType::Locked , EJointMotionType::Locked };
 		Test.JointSettings[0].LinearMotionTypes = { EJointMotionType::Limited, EJointMotionType::Limited , EJointMotionType::Limited };
 		Test.JointSettings[0].bSoftLinearLimitsEnabled = true;
+		Test.JointSettings[0].bLinearPositionDriveEnabled = {true, true, true};
 		Test.JointSettings[0].LinearLimit = 0;
 		Test.JointSettings[0].LinearSoftForceMode = EJointForceMode::Force;
-		Test.JointSettings[0].SoftLinearStiffness = 100000;
+		Test.JointSettings[0].SoftLinearStiffness = 10000;
 		Test.JointSettings[0].SoftLinearDamping = 100;
 
-		Test.JointSettings[0].LinearPlasticityLimit = PlasticityRatio;
+		Test.JointSettings[0].LinearPlasticityLimit = PlasticityLimit;
 
 		Test.Create();
 		Test.AddParticleBox(FVec3(0, 0, 100), FRotation3::Identity, FVec3(10.f), 100.f);
@@ -222,7 +224,7 @@ namespace ChaosTest {
 		FReal DeltaPosPost = (Test.SOAs.GetDynamicParticles().X(0)).Size();
 
 		// The linear spring should have reset. 
-		EXPECT_TRUE(DeltaPosPost < DeltaPos * (1.f-PlasticityRatio) );
+		EXPECT_TRUE(DeltaPosPost < DeltaPos - PlasticityLimit);
 		EXPECT_TRUE(Test.SOAs.GetDynamicParticles().X(0).Z > 0.f);
 	}
 
diff --git a/Engine/Source/Runtime/Experimental/Chaos/Private/Chaos/PBDJointConstraints.cpp b/Engine/Source/Runtime/Experimental/Chaos/Private/Chaos/PBDJointConstraints.cpp
--- a/Engine/Source/Runtime/Experimental/Chaos/Private/Chaos/PBDJointConstraints.cpp
+++ b/Engine/Source/Runtime/Experimental/Chaos/Private/Chaos/PBDJointConstraints.cpp
@@ -1450,23 +1450,24 @@ namespace Chaos
 
 		if(!Settings.bEnableDrives)
 		{
-			UE_LOG(LogChaosJoint, Warning, TEXT("Linear or angular drive needs to be enabled in order to use plasticity."));
+			UE_LOG(LogChaosJoint, VeryVerbose, TEXT("Linear or angular drive needs to be enabled in order to use plasticity."));
 			return;
 		}
 
+		// todo(chaos): provide constraint names for these warnings
 		bool bApplyLinearPlasticityLimits = bHasLinearPlasticityLimit;
 		if(bHasLinearPlasticityLimit)
 		{
 			const bool bLinearPositionDriveEnabled = JointSettings.bLinearPositionDriveEnabled[0] && JointSettings.bLinearPositionDriveEnabled[1] && JointSettings.bLinearPositionDriveEnabled[2];
 			if(!bLinearPositionDriveEnabled)
 			{
-				UE_LOG(LogChaosJoint, Warning, TEXT("Plasticity only supports three-dimensional linear position drives."));
+				UE_LOG(LogChaosJoint, VeryVerbose, TEXT("Plasticity only supports three-dimensional linear position drives."));
 				bApplyLinearPlasticityLimits = false;
 			}
 			const bool bLinearMotionLocked = JointSettings.LinearMotionTypes[0] == EJointMotionType::Locked || JointSettings.LinearMotionTypes[1] == EJointMotionType::Locked || JointSettings.LinearMotionTypes[2] == EJointMotionType::Locked;
 			if(bLinearMotionLocked)
 			{
-				UE_LOG(LogChaosJoint, Warning, TEXT("Plasticity does not support locked linear motion."));
+				UE_LOG(LogChaosJoint, VeryVerbose, TEXT("Plasticity does not support locked linear motion."));
 				bApplyLinearPlasticityLimits = false;
 			}
 		}
@@ -1477,13 +1478,13 @@ namespace Chaos
 			const bool bAngularPositionDriveEnabled = JointSettings.bAngularSLerpPositionDriveEnabled || (JointSettings.bAngularSwingPositionDriveEnabled && JointSettings.bAngularTwistPositionDriveEnabled);
 			if(!bAngularPositionDriveEnabled)
 			{
-				UE_LOG(LogChaosJoint, Warning, TEXT("Plasticity only supports three-dimensional angular position drives"));
+				UE_LOG(LogChaosJoint, VeryVerbose, TEXT("Plasticity only supports three-dimensional angular position drives"));
 				bApplyAngularPlasticityLimits = false;
 			}
 			const bool bAngularMotionLocked = JointSettings.AngularMotionTypes[(int32)EJointAngularConstraintIndex::Twist] == EJointMotionType::Locked || JointSettings.AngularMotionTypes[(int32)EJointAngularConstraintIndex::Swing1] == EJointMotionType::Locked || JointSettings.AngularMotionTypes[(int32)EJointAngularConstraintIndex::Swing2] == EJointMotionType::Locked;
 			if(bAngularMotionLocked)
 			{
-				UE_LOG(LogChaosJoint, Warning, TEXT("Plasticity does not support locked angular motion."));
+				UE_LOG(LogChaosJoint, VeryVerbose, TEXT("Plasticity does not support locked angular motion."));
 				bApplyAngularPlasticityLimits = false;
 			}
 		}
