Update headlesschaos plasticity tests
Use gravity to test linear plasticity. This way we don't need to tune stiffness/damping based on collision behaviors.

diff --git a/Engine/Source/Programs/HeadlessChaos/Private/HeadlessChaosTestJointPlasticity.cpp b/Engine/Source/Programs/HeadlessChaos/Private/HeadlessChaosTestJointPlasticity.cpp
--- a/Engine/Source/Programs/HeadlessChaos/Private/HeadlessChaosTestJointPlasticity.cpp
+++ b/Engine/Source/Programs/HeadlessChaos/Private/HeadlessChaosTestJointPlasticity.cpp
@@ -135,6 +135,7 @@ namespace ChaosTest {
 
 		FJointChainTest<TEvolution> Test(NumIterations, Gravity);
 		Test.InitChain(2, FVec3(0, 0, 1), 10.f, 50.f);
+		Test.ParticleMasses[1] = 10.f;
 
 		Test.JointSettings[0].bCollisionEnabled = false;
 		Test.JointSettings[0].AngularMotionTypes = { EJointMotionType::Locked, EJointMotionType::Locked , EJointMotionType::Locked };
@@ -143,32 +144,29 @@ namespace ChaosTest {
 		Test.JointSettings[0].bLinearPositionDriveEnabled = { true, true, true };
 		Test.JointSettings[0].LinearLimit = 0;
 		Test.JointSettings[0].LinearSoftForceMode = EJointForceMode::Force;
-		Test.JointSettings[0].SoftLinearStiffness = 1000000;
-		Test.JointSettings[0].SoftLinearDamping = 1000;
+		Test.JointSettings[0].SoftLinearStiffness = 10000;
+		Test.JointSettings[0].SoftLinearDamping = 100;
 
 		Test.JointSettings[0].LinearPlasticityLimit = PlasticityLimit;
 
 		Test.Create();
-		Test.AddParticleBox(FVec3(0, 0, 100), FRotation3::Identity, FVec3(10.f), 100.f);
 
-		FReal DeltaPos = (Test.SOAs.GetDynamicParticles().X(0)).Size();
-		EXPECT_TRUE(FMath::IsNearlyEqual(DeltaPos, (FReal)50.));
+		FReal Z = Test.SOAs.GetDynamicParticles().X(0)[2];
+		EXPECT_TRUE(FMath::IsNearlyEqual(Z, (FReal)50.));
 
 		// Run the sim
 		for (int32 i = 0; i < NumIts; ++i)
 		{
-
 			Test.Evolution.AdvanceOneTimeStep(Dt);
 			Test.Evolution.EndFrame(Dt);
 
-			//FReal DeltaPosI = (Test.SOAs.GetDynamicParticles().X(0)).Size();
 			//FVector Pos = Test.SOAs.GetDynamicParticles().X(0);
-			//std::cout << "["<< DeltaPosI <<"]" << Pos.X << "," << Pos.Y << "," << Pos.Z << std::endl;
+			//std::cout << Pos.X << "," << Pos.Y << "," << Pos.Z << std::endl;
 		}
-		FReal DeltaPosPost = (Test.SOAs.GetDynamicParticles().X(0)).Size();
+		FReal ZPost = Test.SOAs.GetDynamicParticles().X(0)[2];
 
 		// Nothing should have reset
-		EXPECT_TRUE(FMath::IsNearlyEqual(DeltaPosPost, DeltaPos, (FReal)5.));
+		EXPECT_TRUE(FMath::IsNearlyEqual(Z, ZPost, (FReal)5.));
 
 	}
 
@@ -183,7 +181,7 @@ namespace ChaosTest {
 	template <typename TEvolution>
 	void JointPlasticity_OverLinearPlasticityThreshold()
 	{
-		const FReal PlasticityLimit = 1;
+		const FReal PlasticityLimit = 0.1;
 		const int32 NumIterations = 1;
 		const FReal Gravity = 980;
 		const FReal Dt = 0.01f;
@@ -191,6 +189,7 @@ namespace ChaosTest {
 
 		FJointChainTest<TEvolution> Test(NumIterations, Gravity);
 		Test.InitChain(2, FVec3(0, 0, 1), 10.f, 50.f);
+		Test.ParticleMasses[1] = 10.f;
 
 		Test.JointSettings[0].bCollisionEnabled = false;
 		Test.JointSettings[0].AngularMotionTypes = { EJointMotionType::Locked, EJointMotionType::Locked , EJointMotionType::Locked };
@@ -205,10 +204,9 @@ namespace ChaosTest {
 		Test.JointSettings[0].LinearPlasticityLimit = PlasticityLimit;
 
 		Test.Create();
-		Test.AddParticleBox(FVec3(0, 0, 100), FRotation3::Identity, FVec3(10.f), 100.f);
 
-		FReal DeltaPos = (Test.SOAs.GetDynamicParticles().X(0)).Size();
-		EXPECT_TRUE(FMath::IsNearlyEqual(DeltaPos, (FReal)50.));
+		FReal Z = Test.SOAs.GetDynamicParticles().X(0)[2];
+		EXPECT_TRUE(FMath::IsNearlyEqual(Z, (FReal)50.));
 
 		// Run the sim
 		for (int32 i = 0; i < NumIts; ++i)
@@ -217,15 +215,14 @@ namespace ChaosTest {
 			Test.Evolution.AdvanceOneTimeStep(Dt);
 			Test.Evolution.EndFrame(Dt);
 
-			//FReal DeltaPosI = (Test.SOAs.GetDynamicParticles().X(0)).Size();
 			//FVector Pos = Test.SOAs.GetDynamicParticles().X(0);
-			//std::cout << "["<< DeltaPosI <<"]" << Pos.X << "," << Pos.Y << "," << Pos.Z << std::endl;
+			//std::cout << Pos.X << "," << Pos.Y << "," << Pos.Z << std::endl;
 		}
-		FReal DeltaPosPost = (Test.SOAs.GetDynamicParticles().X(0)).Size();
+		FReal ZPost = Test.SOAs.GetDynamicParticles().X(0)[2];
 
 		// The linear spring should have reset. 
-		EXPECT_TRUE(DeltaPosPost < DeltaPos - PlasticityLimit);
-		EXPECT_TRUE(Test.SOAs.GetDynamicParticles().X(0).Z > 0.f);
+		EXPECT_TRUE(ZPost < Z - PlasticityLimit);
+		EXPECT_TRUE(ZPost > 0.f);
 	}
 
 	GTEST_TEST(AllEvolutions, JointPlasticity_OverLinearPlasticityThreshold)
