Chaos: Modify convex small triangle check to work with newer changes to error recovery in convex build.

diff --git a/Engine/Source/Programs/HeadlessChaos/Private/HeadlessChaosTestConvex.cpp b/Engine/Source/Programs/HeadlessChaos/Private/HeadlessChaosTestConvex.cpp
--- a/Engine/Source/Programs/HeadlessChaos/Private/HeadlessChaosTestConvex.cpp
+++ b/Engine/Source/Programs/HeadlessChaos/Private/HeadlessChaosTestConvex.cpp
@@ -287,7 +287,7 @@ namespace ChaosTest
 	{
 		// Create a long mesh with a extremely small end (YZ plane) 
 		// so that it generate extremely sized triangle that will produce extremely small (unormalized) normals
-		const float SmallNumber = 0.00001f;
+		const float SmallNumber = 0.001f;
 		const FVec3 Range{ 100.0f, SmallNumber, SmallNumber };
 
 		const FVec3 Vertices[] =
