Chaos: Modified planar convex build test to disable Warning output on LogChaos during its scope. This test is making sure we can safely fail to build a convex without crashing and as such is expecting warning to be fired. CIS still picks that up so we need to suppress it for that one scope.

diff --git a/Engine/Source/Programs/HeadlessChaos/Private/HeadlessChaosTestConvex.cpp b/Engine/Source/Programs/HeadlessChaos/Private/HeadlessChaosTestConvex.cpp
--- a/Engine/Source/Programs/HeadlessChaos/Private/HeadlessChaosTestConvex.cpp
+++ b/Engine/Source/Programs/HeadlessChaos/Private/HeadlessChaosTestConvex.cpp
@@ -10,6 +10,7 @@
 #include "Chaos/ImplicitObjectScaled.h"
 #include "Chaos/Particles.h"
 #include "../Resource/TestGeometry2.h"
+#include "Logging/LogScopedVerbosityOverride.h"
 
 namespace ChaosTest
 {
@@ -367,7 +368,11 @@ namespace ChaosTest
 		TArray<FVec3> FinalVertices;
 		TAABB<FReal, 3> LocalBounds;
 
-		FConvexBuilder::Build(Vertices, Planes, FaceIndices, FinalVertices, LocalBounds);
+		{
+			// Temporarily set LogChaos to error, we're expecting this to fire warnings and don't want that to fail a CIS run.
+			LOG_SCOPE_VERBOSITY_OVERRIDE(LogChaos, ELogVerbosity::Error);
+			FConvexBuilder::Build(Vertices, Planes, FaceIndices, FinalVertices, LocalBounds);
+		}
 
 		// Check that we've failed to build a 3D convex hull and safely returned
 		EXPECT_EQ(Planes.Num(), 0);
