Fix UE5 CIS

diff --git a/Engine/Source/Programs/HeadlessChaos/Private/GeometryCollection/GeometryCollectionTestFields.cpp b/Engine/Source/Programs/HeadlessChaos/Private/GeometryCollection/GeometryCollectionTestFields.cpp
--- a/Engine/Source/Programs/HeadlessChaos/Private/GeometryCollection/GeometryCollectionTestFields.cpp
+++ b/Engine/Source/Programs/HeadlessChaos/Private/GeometryCollection/GeometryCollectionTestFields.cpp
@@ -50,7 +50,7 @@ namespace GeometryCollectionTest
 
 		TArray<float> ResultsArray;
 		ResultsArray.Init(false, Bounds*Bounds);
-		TArrayView<float> ResultsView(&(ResultsArray.operator[](0)), ResultsArray.Num());
+		TFieldArrayView<float> ResultsView(ResultsArray, 0, ResultsArray.Num());
 		NoiseField->Evaluate(Context, ResultsView);
 
 		Chaos::FReal min = FLT_MAX, max = -FLT_MAX;
@@ -92,7 +92,7 @@ namespace GeometryCollectionTest
 
 		TArray<int32> ResultsArray;
 		ResultsArray.Init(false, 10.0);
-		TArrayView<int32> ResultsView(&(ResultsArray.operator[](0)), ResultsArray.Num());
+		TFieldArrayView<int32> ResultsView(ResultsArray, 0, ResultsArray.Num());
 		RadialMask->Evaluate(Context, ResultsView);
 
 		for (int32 Index = 0; Index < 10; Index++)
@@ -135,7 +135,7 @@ namespace GeometryCollectionTest
 
 		TArray<float> ResultsArray;
 		ResultsArray.Init(false, 10);
-		TArrayView<float> ResultsView(&(ResultsArray.operator[](0)), ResultsArray.Num());
+		TFieldArrayView<float> ResultsView(ResultsArray, 0, ResultsArray.Num());
 		RadialFalloff->Evaluate(Context, ResultsView);
 
 		for (int32 Index = 0; Index < 10; Index++)
@@ -183,7 +183,7 @@ namespace GeometryCollectionTest
 
 		TArray<float> ResultsArray;
 		ResultsArray.Init(false, 10);
-		TArrayView<float> ResultsView(&(ResultsArray.operator[](0)), ResultsArray.Num());
+		TFieldArrayView<float> ResultsView(ResultsArray, 0, ResultsArray.Num());
 		PlaneFalloff->Evaluate(Context, ResultsView);
 
 		FPlane Plane(PlaneFalloff->Position, PlaneFalloff->Normal);
@@ -224,7 +224,7 @@ namespace GeometryCollectionTest
 
 		TArray<FVector> ResultsArray;
 		ResultsArray.Init(FVector(0), 10);
-		TArrayView<FVector> ResultsView(&(ResultsArray.operator[](0)), ResultsArray.Num());
+		TFieldArrayView<FVector> ResultsView(ResultsArray, 0, ResultsArray.Num());
 		UniformVector->Evaluate(Context, ResultsView);
 
 		for (int32 Index = 0; Index < 10; Index++)
@@ -261,7 +261,7 @@ namespace GeometryCollectionTest
 
 		TArray<FVector> ResultsArray;
 		ResultsArray.Init(FVector(0), 10);
-		TArrayView<FVector> ResultsView(&(ResultsArray.operator[](0)), ResultsArray.Num());
+		TFieldArrayView<FVector> ResultsView(ResultsArray, 0, ResultsArray.Num());
 		RadialVector->Evaluate(Context, ResultsView);
 
 		for (int32 Index = 0; Index < 10; Index++)
@@ -317,7 +317,7 @@ namespace GeometryCollectionTest
 
 		TArray<FVector> ResultsArray;
 		ResultsArray.Init(FVector(0), 10);
-		TArrayView<FVector> ResultsView(&(ResultsArray.operator[](0)), ResultsArray.Num());
+		TFieldArrayView<FVector> ResultsView(ResultsArray, 0, ResultsArray.Num());
 		SumVector->Evaluate(Context, ResultsView);
 
 		Chaos::FReal RadialFalloffSize = RadialFalloff->Radius;
@@ -387,7 +387,7 @@ namespace GeometryCollectionTest
 
 		TArray<FVector> ResultsArray;
 		ResultsArray.Init(FVector(0), 10);
-		TArrayView<FVector> ResultsView(&(ResultsArray.operator[](0)), ResultsArray.Num());
+		TFieldArrayView<FVector> ResultsView(ResultsArray, 0, ResultsArray.Num());
 		SumVector->Evaluate(Context, ResultsView);
 
 		Chaos::FReal RadialFalloffSize = RadialFalloff->Radius;
@@ -454,7 +454,7 @@ namespace GeometryCollectionTest
 
 		TArray<FVector> ResultsArray;
 		ResultsArray.Init(FVector(0), 10);
-		TArrayView<FVector> ResultsView(&(ResultsArray.operator[](0)), ResultsArray.Num());
+		TFieldArrayView<FVector> ResultsView(ResultsArray, 0, ResultsArray.Num());
 		SumVector->Evaluate(Context, ResultsView);
 
 		Chaos::FReal RadialFalloffSize = RadialFalloff->Radius;
@@ -521,7 +521,7 @@ namespace GeometryCollectionTest
 
 		TArray<FVector> ResultsArray;
 		ResultsArray.Init(FVector(0), 10);
-		TArrayView<FVector> ResultsView(&(ResultsArray.operator[](0)), ResultsArray.Num());
+		TFieldArrayView<FVector> ResultsView(ResultsArray, 0, ResultsArray.Num());
 		SumVector->Evaluate(Context, ResultsView);
 
 		Chaos::FReal RadialFalloffSize = RadialFalloff->Radius;
@@ -588,7 +588,7 @@ namespace GeometryCollectionTest
 
 		TArray<FVector> ResultsArray;
 		ResultsArray.Init(FVector(0), 10);
-		TArrayView<FVector> ResultsView(&(ResultsArray.operator[](0)), ResultsArray.Num());
+		TFieldArrayView<FVector> ResultsView(ResultsArray, 0, ResultsArray.Num());
 		SumVector->Evaluate(Context, ResultsView);
 
 		Chaos::FReal RadialFalloffSize = RadialFalloff->Radius;
@@ -656,7 +656,7 @@ namespace GeometryCollectionTest
 
 		TArray<FVector> ResultsArray;
 		ResultsArray.Init(FVector(0), 10);
-		TArrayView<FVector> ResultsView(&(ResultsArray.operator[](0)), ResultsArray.Num());
+		TFieldArrayView<FVector> ResultsView(ResultsArray, 0, ResultsArray.Num());
 		SumVector->Evaluate(Context, ResultsView);
 
 		Chaos::FReal RadialFalloffSize = RadialFalloff->Radius;
@@ -717,7 +717,7 @@ namespace GeometryCollectionTest
 
 		TArray<float> ResultsArray;
 		ResultsArray.Init(0.f, ExecutionDatas.SamplePositions.Num());
-		TArrayView<float> ResultsView(&(ResultsArray.operator[](0)), ResultsArray.Num());
+		TFieldArrayView<float> ResultsView(ResultsArray, 0, ResultsArray.Num());
 		SumScalar->Evaluate(Context, ResultsView);
 
 		for (int32 Index = 0; Index < NumPoints; Index++)
@@ -780,7 +780,7 @@ namespace GeometryCollectionTest
 
 		TArray<float> ResultsArray;
 		ResultsArray.Init(0.f, ExecutionDatas.SamplePositions.Num());
-		TArrayView<float> ResultsView(&(ResultsArray.operator[](0)), ResultsArray.Num());
+		TFieldArrayView<float> ResultsView(ResultsArray, 0, ResultsArray.Num());
 		SumScalar->Evaluate(Context, ResultsView);
 
 		for (int32 Index = 0; Index < NumPoints; Index++)
@@ -835,7 +835,7 @@ namespace GeometryCollectionTest
 
 		TArray<float> ResultsArray;
 		ResultsArray.Init(0.f, ExecutionDatas.SamplePositions.Num());
-		TArrayView<float> ResultsView(&(ResultsArray.operator[](0)), ResultsArray.Num());
+		TFieldArrayView<float> ResultsView(ResultsArray, 0, ResultsArray.Num());
 		SumScalar->Evaluate(Context, ResultsView);
 
 		for (int32 Index = 0; Index < NumPoints; Index++)
@@ -897,7 +897,7 @@ namespace GeometryCollectionTest
 
 		TArray<float> ResultsArray;
 		ResultsArray.Init(0.f, ExecutionDatas.SamplePositions.Num());
-		TArrayView<float> ResultsView(&(ResultsArray.operator[](0)), ResultsArray.Num());
+		TFieldArrayView<float> ResultsView(ResultsArray, 0, ResultsArray.Num());
 		CullingField->Evaluate(Context, ResultsView);
 
 		for (int32 Index = 0; Index < NumPoints; Index++)
