fix for unit test CIS

diff --git a/Engine/Source/Programs/HeadlessChaos/Private/HeadlessChaosTestPhysicsScene.cpp b/Engine/Source/Programs/HeadlessChaos/Private/HeadlessChaosTestPhysicsScene.cpp
--- a/Engine/Source/Programs/HeadlessChaos/Private/HeadlessChaosTestPhysicsScene.cpp
+++ b/Engine/Source/Programs/HeadlessChaos/Private/HeadlessChaosTestPhysicsScene.cpp
@@ -1507,7 +1507,7 @@ namespace ChaosTest {
 		Params.InitialTM = FTransform::Identity;
 		Params.Scene = &Scene;
 
-		TGeometryParticle<FReal, 3>* StaticCube = nullptr;
+		FPhysicsActorHandle StaticCube = nullptr;
 
 		FChaosEngineInterface::CreateActor(Params, StaticCube);
 		ASSERT_NE(StaticCube, nullptr);
@@ -1523,10 +1523,13 @@ namespace ChaosTest {
 		Geoms.Emplace(MakeUnique<TBox<FReal, 3>>(-HalfBoxExtent, HalfBoxExtent));
 		Geoms.Emplace(MakeUnique<TBox<FReal, 3>>(-HalfBoxExtent, HalfBoxExtent));
 
-		TUniquePtr<FImplicitObjectUnion> GeomUnion = MakeUnique<FImplicitObjectUnion>(MoveTemp(Geoms));
-		StaticCube->SetGeometry(MoveTemp(GeomUnion));
-
-		TArray<TGeometryParticle<FReal, 3>*> Particles{ StaticCube };
+		auto& Particle = StaticCube->GetGameThreadAPI();
+		{
+			TUniquePtr<FImplicitObjectUnion> GeomUnion = MakeUnique<FImplicitObjectUnion>(MoveTemp(Geoms));
+			Particle.SetGeometry(MoveTemp(GeomUnion));
+		}
+		
+		TArray<FPhysicsActorHandle> Particles{ StaticCube };
 		Scene.AddActorsToScene_AssumesLocked(Particles);
 
 		FChaosSQAccelerator SQ{ *Scene.GetSpacialAcceleration() };
