headless chaos : fix errors in geom collection mass properties

diff --git a/Engine/Source/Programs/HeadlessChaos/Private/GeometryCollection/GeometryCollectionTestMassProperties.cpp b/Engine/Source/Programs/HeadlessChaos/Private/GeometryCollection/GeometryCollectionTestMassProperties.cpp
--- a/Engine/Source/Programs/HeadlessChaos/Private/GeometryCollection/GeometryCollectionTestMassProperties.cpp
+++ b/Engine/Source/Programs/HeadlessChaos/Private/GeometryCollection/GeometryCollectionTestMassProperties.cpp
@@ -393,9 +393,9 @@ namespace GeometryCollectionTest
 		CalculateInertiaAndRotationOfMass(MassSpaceParticles, TriMesh->GetSurfaceElements(), Density, ZeroVec, MassProperties.InertiaTensor, MassProperties.RotationOfMass);
 
 		// rotational alignment.
-		EXPECT_NEAR(MassProperties.RotationOfMass.Euler()[0], 135.f, KINDA_SMALL_NUMBER);
-		EXPECT_NEAR(MassProperties.RotationOfMass.Euler()[1], 0, KINDA_SMALL_NUMBER);
-		EXPECT_NEAR(MassProperties.RotationOfMass.Euler()[2], 0, KINDA_SMALL_NUMBER);
+		EXPECT_NEAR(MassProperties.RotationOfMass.Euler()[0], -45., KINDA_SMALL_NUMBER);
+		EXPECT_NEAR(MassProperties.RotationOfMass.Euler()[1], 0., KINDA_SMALL_NUMBER);
+		EXPECT_NEAR(MassProperties.RotationOfMass.Euler()[2], 0., KINDA_SMALL_NUMBER);
 		// X dominate inertia tensor
 		EXPECT_GT(MassProperties.InertiaTensor.M[0][0], MassProperties.InertiaTensor.M[2][2]);
 		EXPECT_GT(MassProperties.InertiaTensor.M[0][0], MassProperties.InertiaTensor.M[1][1]);
