Fix compilation issue with headless chaos in LWC mode

diff --git a/Engine/Source/Programs/HeadlessChaos/Private/GeometryCollection/GeometryCollectionTestClustering.cpp b/Engine/Source/Programs/HeadlessChaos/Private/GeometryCollection/GeometryCollectionTestClustering.cpp
--- a/Engine/Source/Programs/HeadlessChaos/Private/GeometryCollection/GeometryCollectionTestClustering.cpp
+++ b/Engine/Source/Programs/HeadlessChaos/Private/GeometryCollection/GeometryCollectionTestClustering.cpp
@@ -317,7 +317,7 @@ namespace GeometryCollectionTest
 
 			if (Frame == 2)
 			{
-				TMap<FGeometryParticleHandle*, float> ExternalStrains = { {ParticleHandles[0], 50.0f} };
+				TMap<FGeometryParticleHandle*, FReal> ExternalStrains = { {ParticleHandles[0], 50.0f} };
 				Clustering.BreakingModel(&ExternalStrains);
 			}
 
diff --git a/Engine/Source/Programs/HeadlessChaos/Private/GeometryCollection/GeometryCollectionTestFields.cpp b/Engine/Source/Programs/HeadlessChaos/Private/GeometryCollection/GeometryCollectionTestFields.cpp
--- a/Engine/Source/Programs/HeadlessChaos/Private/GeometryCollection/GeometryCollectionTestFields.cpp
+++ b/Engine/Source/Programs/HeadlessChaos/Private/GeometryCollection/GeometryCollectionTestFields.cpp
@@ -48,9 +48,9 @@ namespace GeometryCollectionTest
 			0.0
 		};
 
-		TArray<Chaos::FReal> ResultsArray;
+		TArray<float> ResultsArray;
 		ResultsArray.Init(false, Bounds*Bounds);
-		TFieldArrayView<Chaos::FReal> ResultsView(ResultsArray, 0, ResultsArray.Num());
+		TFieldArrayView<float> ResultsView(ResultsArray, 0, ResultsArray.Num());
 		NoiseField->Evaluate(Context, ResultsView);
 
 		Chaos::FReal min = FLT_MAX, max = -FLT_MAX;
@@ -133,9 +133,9 @@ namespace GeometryCollectionTest
 			0.0
 		};
 
-		TArray<Chaos::FReal> ResultsArray;
+		TArray<float> ResultsArray;
 		ResultsArray.Init(false, 10);
-		TFieldArrayView<Chaos::FReal> ResultsView(ResultsArray, 0, ResultsArray.Num());
+		TFieldArrayView<float> ResultsView(ResultsArray, 0, ResultsArray.Num());
 		RadialFalloff->Evaluate(Context, ResultsView);
 
 		for (int32 Index = 0; Index < 10; Index++)
@@ -181,9 +181,9 @@ namespace GeometryCollectionTest
 			0.0
 		};
 
-		TArray<Chaos::FReal> ResultsArray;
+		TArray<float> ResultsArray;
 		ResultsArray.Init(false, 10);
-		TFieldArrayView<Chaos::FReal> ResultsView(ResultsArray, 0, ResultsArray.Num());
+		TFieldArrayView<float> ResultsView(ResultsArray, 0, ResultsArray.Num());
 		PlaneFalloff->Evaluate(Context, ResultsView);
 
 		FPlane Plane(PlaneFalloff->Position, PlaneFalloff->Normal);
@@ -715,9 +715,9 @@ namespace GeometryCollectionTest
 			0.0
 		};
 
-		TArray<Chaos::FReal> ResultsArray;
+		TArray<float> ResultsArray;
 		ResultsArray.Init(0.f, ExecutionDatas.SamplePositions.Num());
-		TFieldArrayView<Chaos::FReal> ResultsView(ResultsArray, 0, ResultsArray.Num());
+		TFieldArrayView<float> ResultsView(ResultsArray, 0, ResultsArray.Num());
 		SumScalar->Evaluate(Context, ResultsView);
 
 		for (int32 Index = 0; Index < NumPoints; Index++)
@@ -778,9 +778,9 @@ namespace GeometryCollectionTest
 			0.0
 		};
 
-		TArray<Chaos::FReal> ResultsArray;
+		TArray<float> ResultsArray;
 		ResultsArray.Init(0.f, ExecutionDatas.SamplePositions.Num());
-		TFieldArrayView<Chaos::FReal> ResultsView(ResultsArray, 0, ResultsArray.Num());
+		TFieldArrayView<float> ResultsView(ResultsArray, 0, ResultsArray.Num());
 		SumScalar->Evaluate(Context, ResultsView);
 
 		for (int32 Index = 0; Index < NumPoints; Index++)
@@ -833,9 +833,9 @@ namespace GeometryCollectionTest
 			0.0
 		};
 
-		TArray<Chaos::FReal> ResultsArray;
+		TArray<float> ResultsArray;
 		ResultsArray.Init(0.f, ExecutionDatas.SamplePositions.Num());
-		TFieldArrayView<Chaos::FReal> ResultsView(ResultsArray, 0, ResultsArray.Num());
+		TFieldArrayView<float> ResultsView(ResultsArray, 0, ResultsArray.Num());
 		SumScalar->Evaluate(Context, ResultsView);
 
 		for (int32 Index = 0; Index < NumPoints; Index++)
@@ -895,9 +895,9 @@ namespace GeometryCollectionTest
 			0.0
 		};
 
-		TArray<Chaos::FReal> ResultsArray;
+		TArray<float> ResultsArray;
 		ResultsArray.Init(0.f, ExecutionDatas.SamplePositions.Num());
-		TFieldArrayView<Chaos::FReal> ResultsView(ResultsArray, 0, ResultsArray.Num());
+		TFieldArrayView<float> ResultsView(ResultsArray, 0, ResultsArray.Num());
 		CullingField->Evaluate(Context, ResultsView);
 
 		for (int32 Index = 0; Index < NumPoints; Index++)
