Chaos: CIS: Fixed HeadlessChaos build - new tests post merge couldn't compile against API changes

diff --git a/Engine/Source/Programs/HeadlessChaos/Private/HeadlessChaosTestPhysicsScene.cpp b/Engine/Source/Programs/HeadlessChaos/Private/HeadlessChaosTestPhysicsScene.cpp
--- a/Engine/Source/Programs/HeadlessChaos/Private/HeadlessChaosTestPhysicsScene.cpp
+++ b/Engine/Source/Programs/HeadlessChaos/Private/HeadlessChaosTestPhysicsScene.cpp
@@ -446,9 +446,7 @@ namespace ChaosTest {
 		const float PhysicsTimestep = 1; // 1 second
 		float DeltaSeconds = PhysicsTimestep;
 		FVec3 Grav(0, 0, -1);
-		Scene.SetUpForFrame(&Grav, DeltaSeconds, 9999, 9999, 9999, false);
-		Scene.GetSolver()->EnableAsyncMode(PhysicsTimestep);
-
+		Scene.SetUpForFrame(&Grav, DeltaSeconds, 9999, 9999, 9999, false, PhysicsTimestep);
 
 		// Raycast params, aimed to hit our particle at (0,0,0)
 		const FVector Start(0, 0, -5);
@@ -567,10 +565,8 @@ namespace ChaosTest {
 		float DeltaSeconds = PhysicsTimestep * PhysicsStepsInFrame;
 		FVec3 Grav(0, 0, -1);
 
-		Scene.SetUpForFrame(&Grav, DeltaSeconds, 9999, 9999, 9999, false);
-		Scene.GetSolver()->EnableAsyncMode(PhysicsTimestep);
-
-
+		Scene.SetUpForFrame(&Grav, DeltaSeconds, 9999, 9999, 9999, false, PhysicsTimestep);
+		
 		// Raycast params, aimed to hit our kinematic target (10,0,0)
 		const FVector Start(10, 0, -5);
 		const FVector Dir(0, 0,1);
