Fix headless chaos build, specify AsyncDt in solver/scene constructors.

diff --git a/Engine/Source/Programs/HeadlessChaos/Private/GeometryCollection/GeometryCollectionTestEvents.cpp b/Engine/Source/Programs/HeadlessChaos/Private/GeometryCollection/GeometryCollectionTestEvents.cpp
--- a/Engine/Source/Programs/HeadlessChaos/Private/GeometryCollection/GeometryCollectionTestEvents.cpp
+++ b/Engine/Source/Programs/HeadlessChaos/Private/GeometryCollection/GeometryCollectionTestEvents.cpp
@@ -86,7 +86,7 @@ namespace GeometryCollectionTest
 	GTEST_TEST(AllTraits, GeometryCollection_EventBufferTest_Event_Handler)
 	{
 		Chaos::FEventManager EventManager(Chaos::EMultiBufferMode::Single);
-		Chaos::FPBDRigidsSolver* Solver = FChaosSolversModule::GetModule()->CreateSolver(nullptr);
+		Chaos::FPBDRigidsSolver* Solver = FChaosSolversModule::GetModule()->CreateSolver(nullptr, /*AsyncDt=*/-1);
 
 		MyEventHandler HandlerTest(EventManager);
 		MyEventHandler AnotherHandlerTest(EventManager);
diff --git a/Engine/Source/Programs/HeadlessChaos/Private/GeometryCollection/GeometryCollectionTestFramework.cpp b/Engine/Source/Programs/HeadlessChaos/Private/GeometryCollection/GeometryCollectionTestFramework.cpp
--- a/Engine/Source/Programs/HeadlessChaos/Private/GeometryCollection/GeometryCollectionTestFramework.cpp
+++ b/Engine/Source/Programs/HeadlessChaos/Private/GeometryCollection/GeometryCollectionTestFramework.cpp
@@ -266,7 +266,7 @@ namespace GeometryCollectionTest
 	, Module(FChaosSolversModule::GetModule())
 	, Solver(nullptr)
 	{
-		Solver = Module->CreateSolver(nullptr,Parameters.ThreadingMode);	//until refactor is done, solver must be created after thread change
+		Solver = Module->CreateSolver(nullptr, /*AsyncDt=*/-1, Parameters.ThreadingMode);	//until refactor is done, solver must be created after thread change
 		ChaosTest::InitSolverSettings(Solver);
 	}
 
diff --git a/Engine/Source/Programs/HeadlessChaos/Private/HeadlessChaosTestDataMarshalling.cpp b/Engine/Source/Programs/HeadlessChaos/Private/HeadlessChaosTestDataMarshalling.cpp
--- a/Engine/Source/Programs/HeadlessChaos/Private/HeadlessChaosTestDataMarshalling.cpp
+++ b/Engine/Source/Programs/HeadlessChaos/Private/HeadlessChaosTestDataMarshalling.cpp
@@ -101,7 +101,7 @@ namespace ChaosTest
 
 	GTEST_TEST(AllTraits, DataMarshalling_Callbacks)
 	{
-		auto* Solver = FChaosSolversModule::GetModule()->CreateSolver(nullptr, EThreadingMode::SingleThread);
+		auto* Solver = FChaosSolversModule::GetModule()->CreateSolver(nullptr, /*AsyncDt=*/-1, EThreadingMode::SingleThread);
 		
 		int Count = 0;
 		FReal Time = 0;
@@ -167,7 +167,7 @@ namespace ChaosTest
 
 	GTEST_TEST(AllTraits,DataMarshalling_OneShotCallbacks)
 	{
-		auto* Solver = FChaosSolversModule::GetModule()->CreateSolver(nullptr,EThreadingMode::SingleThread);
+		auto* Solver = FChaosSolversModule::GetModule()->CreateSolver(nullptr, /*AsyncDt=*/-1, EThreadingMode::SingleThread);
 		
 		int Count = 0;
 		Solver->RegisterSimOneShotCallback([&Count]()
diff --git a/Engine/Source/Programs/HeadlessChaos/Private/HeadlessChaosTestPhysicsScene.cpp b/Engine/Source/Programs/HeadlessChaos/Private/HeadlessChaosTestPhysicsScene.cpp
--- a/Engine/Source/Programs/HeadlessChaos/Private/HeadlessChaosTestPhysicsScene.cpp
+++ b/Engine/Source/Programs/HeadlessChaos/Private/HeadlessChaosTestPhysicsScene.cpp
@@ -145,7 +145,7 @@ namespace ChaosTest {
 
 	GTEST_TEST(EngineInterface, CreateAndReleaseActor)
 	{
-		FChaosScene Scene(nullptr);
+		FChaosScene Scene(nullptr, /*AsyncDt=*/-1);
 
 		FActorCreationParams Params;
 		Params.Scene = &Scene;
@@ -167,7 +167,7 @@ namespace ChaosTest {
 
 	GTEST_TEST(EngineInterface, CreateMoveAndReleaseInScene)
 	{
-		FChaosScene Scene(nullptr);
+		FChaosScene Scene(nullptr, /*AsyncDt=*/-1);
 
 		FActorCreationParams Params;
 		Params.Scene = &Scene;
@@ -231,7 +231,7 @@ namespace ChaosTest {
 	{
 		//make sure acceleration structure has appropriate sync time
 
-		FChaosScene Scene(nullptr);
+		FChaosScene Scene(nullptr, /*AsyncDt=*/-1);
 		Scene.GetSolver()->SetThreadingMode_External(EThreadingModeTemp::SingleThread);
 
 		EXPECT_EQ(Scene.GetSpacialAcceleration()->GetSyncTimestamp(), 0);	//timestamp of 0 because we flush when scene is created
@@ -240,7 +240,7 @@ namespace ChaosTest {
 		for (int Step = 1; Step < 10; ++Step)
 		{
 			FVec3 Grav(0,0,-1);
-			Scene.SetUpForFrame(&Grav, 1,99999,99999,10,false, -1);
+			Scene.SetUpForFrame(&Grav, 1,99999,99999,10,false);
 			Scene.StartFrame();
 			Scene.GetSolver()->GetEvolution()->FlushSpatialAcceleration();	//make sure we get a new tree every step
 			Scene.EndFrame();
@@ -253,14 +253,14 @@ namespace ChaosTest {
 	{
 		//make sure acceleration structure has appropriate sync time when PT falls behind GT
 
-		FChaosScene Scene(nullptr);
+		FChaosScene Scene(nullptr, /*AsyncDt=*/-1);
 		Scene.GetSolver()->SetThreadingMode_External(EThreadingModeTemp::SingleThread);
 		Scene.GetSolver()->SetStealAdvanceTasks_ForTesting(true); // prevents execution on StartFrame so we can execute task manually.
 
 		EXPECT_EQ(Scene.GetSpacialAcceleration()->GetSyncTimestamp(), 0);	//timestamp of 0 because we flush when scene is created
 
 		FVec3 Grav(0, 0, -1);
-		Scene.SetUpForFrame(&Grav, 1, 99999, 99999, 10, false, -1);
+		Scene.SetUpForFrame(&Grav, 1, 99999, 99999, 10, false);
 
 		// Game thread enqueues second solver task before first completes (we did not execute advance task)
 		Scene.StartFrame();
@@ -293,14 +293,14 @@ namespace ChaosTest {
 	{
 		//make sure acceleration structure has appropriate sync time when PT falls behind GT
 
-		FChaosScene Scene(nullptr);
+		FChaosScene Scene(nullptr, /*AsyncDt=*/-1);
 		Scene.GetSolver()->SetThreadingMode_External(EThreadingModeTemp::SingleThread);
 		Scene.GetSolver()->SetStealAdvanceTasks_ForTesting(true); // prevents execution on StartFrame so we can execute task manually.
 
 		EXPECT_EQ(Scene.GetSpacialAcceleration()->GetSyncTimestamp(), 0);	//timestamp of 0 because we flush when scene is created
 
 		FVec3 Grav(0,0,-1);
-		Scene.SetUpForFrame(&Grav, 1,99999,99999,10,false, -1);
+		Scene.SetUpForFrame(&Grav, 1,99999,99999,10,false);
 
 		// PT not finished yet (we didn't execute solver task), should still be 0.
 		Scene.StartFrame();
@@ -342,12 +342,12 @@ namespace ChaosTest {
 		// This test is designed to verify pulldata is being timestamped correctly, and that we will not write to a deleted GT Proxy 
 		// in this case. 
 
-		FChaosScene Scene(nullptr);
+		FChaosScene Scene(nullptr, /*AsyncDt=*/-1);
 		Scene.GetSolver()->SetThreadingMode_External(EThreadingModeTemp::SingleThread);
 		Scene.GetSolver()->SetStealAdvanceTasks_ForTesting(true); // prevents execution on StartFrame so we can execute task manually.
 
 		FVec3 Grav(0,0,-1);
-		Scene.SetUpForFrame(&Grav, 1,99999,99999,10,false, -1);
+		Scene.SetUpForFrame(&Grav, 1,99999,99999,10,false);
 
 		FActorCreationParams Params;
 		Params.Scene = &Scene;
@@ -440,13 +440,13 @@ namespace ChaosTest {
 
 	GTEST_TEST(EngineInterface, UpdatingAccelerationStructurePrePreFilterOnShapeFilterChange)
 	{
-		FChaosScene Scene(nullptr);
+		const float PhysicsTimestep = 1; // 1 second
+		FChaosScene Scene(nullptr, PhysicsTimestep);
 		Scene.GetSolver()->SetThreadingMode_External(EThreadingModeTemp::SingleThread);
 
-		const float PhysicsTimestep = 1; // 1 second
 		float DeltaSeconds = PhysicsTimestep;
 		FVec3 Grav(0, 0, -1);
-		Scene.SetUpForFrame(&Grav, DeltaSeconds, 9999, 9999, 9999, false, PhysicsTimestep);
+		Scene.SetUpForFrame(&Grav, DeltaSeconds, 9999, 9999, 9999, false);
 
 		// Raycast params, aimed to hit our particle at (0,0,0)
 		const FVector Start(0, 0, -5);
@@ -553,19 +553,19 @@ namespace ChaosTest {
 		// kinematic has not reached target yet, that the bounds in structure representing interpolated position do not make it to game thread, otherwise game thread has
 		// position at target, but bounds that don't match.
 
+		const float PhysicsTimestep = 1; // 1 second
 
 		// Setup solver so we can manually execute each physics step.
-		FChaosScene Scene(nullptr);
+		FChaosScene Scene(nullptr, PhysicsTimestep);
 		Scene.GetSolver()->SetThreadingMode_External(EThreadingModeTemp::SingleThread);
 		Scene.GetSolver()->SetStealAdvanceTasks_ForTesting(true);
 
 		// In this test we have a 10s Dt, split into 10 physics steps of 1s.
-		const float PhysicsTimestep = 1; // 1 second
 		const int32 PhysicsStepsInFrame = 10; 
 		float DeltaSeconds = PhysicsTimestep * PhysicsStepsInFrame;
 		FVec3 Grav(0, 0, -1);
 
-		Scene.SetUpForFrame(&Grav, DeltaSeconds, 9999, 9999, 9999, false, PhysicsTimestep);
+		Scene.SetUpForFrame(&Grav, DeltaSeconds, 9999, 9999, 9999, false);
 		
 		// Raycast params, aimed to hit our kinematic target (10,0,0)
 		const FVector Start(10, 0, -5);
@@ -653,7 +653,7 @@ namespace ChaosTest {
 
 	GTEST_TEST(EngineInterface, CreateActorPostFlush)
 	{
-		FChaosScene Scene(nullptr);
+		FChaosScene Scene(nullptr, /*AsyncDt=*/-1);
 		Scene.GetSolver()->SetThreadingMode_External(EThreadingModeTemp::SingleThread);
 
 		FActorCreationParams Params;
@@ -673,7 +673,7 @@ namespace ChaosTest {
 		//tick solver but don't call EndFrame (want to flush and swap manually)
 		{
 			FVec3 Grav(0,0,-1);
-			Scene.SetUpForFrame(&Grav,1,99999,99999,10,false, -1);
+			Scene.SetUpForFrame(&Grav,1,99999,99999,10,false);
 			Scene.StartFrame();
 		}
 
@@ -693,7 +693,7 @@ namespace ChaosTest {
 
 	GTEST_TEST(EngineInterface, MoveActorPostFlush)
 	{
-		FChaosScene Scene(nullptr);
+		FChaosScene Scene(nullptr, /*AsyncDt=*/-1);
 		Scene.GetSolver()->SetThreadingMode_External(EThreadingModeTemp::SingleThread);
 
 		FActorCreationParams Params;
@@ -717,7 +717,7 @@ namespace ChaosTest {
 		//tick solver so that Proxy is created, but don't call EndFrame (want to flush and swap manually)
 		{
 			FVec3 Grav(0,0,-1);
-			Scene.SetUpForFrame(&Grav,1,99999,99999,10,false, -1);
+			Scene.SetUpForFrame(&Grav,1,99999,99999,10,false);
 			Scene.StartFrame();
 		}
 
@@ -737,7 +737,7 @@ namespace ChaosTest {
 
 	GTEST_TEST(EngineInterface, RemoveActorPostFlush)
 	{
-		FChaosScene Scene(nullptr);
+		FChaosScene Scene(nullptr, /*AsyncDt=*/-1);
 		Scene.GetSolver()->SetThreadingMode_External(EThreadingModeTemp::SingleThread);
 
 		FActorCreationParams Params;
@@ -761,7 +761,7 @@ namespace ChaosTest {
 		//tick solver so that Proxy is created, but don't call EndFrame (want to flush and swap manually)
 		{
 			FVec3 Grav(0,0,-1);
-			Scene.SetUpForFrame(&Grav,1,99999,99999,10,false, -1);
+			Scene.SetUpForFrame(&Grav,1,99999,99999,10,false);
 			Scene.StartFrame();
 		}
 
@@ -780,7 +780,7 @@ namespace ChaosTest {
 
 	GTEST_TEST(EngineInterface, RemoveActorPostFlush0Dt)
 	{
-		FChaosScene Scene(nullptr);
+		FChaosScene Scene(nullptr, /*AsyncDt=*/-1);
 		Scene.GetSolver()->SetThreadingMode_External(EThreadingModeTemp::SingleThread);
 
 		FActorCreationParams Params;
@@ -805,7 +805,7 @@ namespace ChaosTest {
 		{
 			//use 0 dt to make sure pending operations are not sensitive to 0 dt
 			FVec3 Grav(0,0,-1);
-			Scene.SetUpForFrame(&Grav,0,99999,99999,10,false, -1);
+			Scene.SetUpForFrame(&Grav,0,99999,99999,10,false);
 			Scene.StartFrame();
 		}
 
@@ -824,7 +824,7 @@ namespace ChaosTest {
 
 	GTEST_TEST(EngineInterface, CreateAndRemoveActorPostFlush)
 	{
-		FChaosScene Scene(nullptr);
+		FChaosScene Scene(nullptr, /*AsyncDt=*/-1);
 		Scene.GetSolver()->SetThreadingMode_External(EThreadingModeTemp::SingleThread);
 
 		FActorCreationParams Params;
@@ -835,7 +835,7 @@ namespace ChaosTest {
 		//tick solver, but don't call EndFrame (want to flush and swap manually)
 		{
 			FVec3 Grav(0,0,-1);
-			Scene.SetUpForFrame(&Grav,1,99999,99999,10,false, -1);
+			Scene.SetUpForFrame(&Grav,1,99999,99999,10,false);
 			Scene.StartFrame();
 		}
 
@@ -869,7 +869,7 @@ namespace ChaosTest {
 	{
 		for (int Delay = 0; Delay < 4; ++Delay)
 		{
-			FChaosScene Scene(nullptr);
+			FChaosScene Scene(nullptr, /*AsyncDt=*/-1);
 			Scene.GetSolver()->SetThreadingMode_External(EThreadingModeTemp::SingleThread);
 			Scene.GetSolver()->GetMarshallingManager().SetTickDelay_External(Delay);
 
@@ -896,7 +896,7 @@ namespace ChaosTest {
 				//tick solver
 				{
 					FVec3 Grav(0,0,-1);
-					Scene.SetUpForFrame(&Grav,1,99999,99999,1,false, -1);
+					Scene.SetUpForFrame(&Grav,1,99999,99999,1,false);
 					Scene.StartFrame();
 					Scene.EndFrame();
 				}
@@ -918,7 +918,7 @@ namespace ChaosTest {
 			//tick solver one last time
 			{
 				FVec3 Grav(0,0,-1);
-				Scene.SetUpForFrame(&Grav,1,99999,99999,1,false, -1);
+				Scene.SetUpForFrame(&Grav,1,99999,99999,1,false);
 				Scene.StartFrame();
 				Scene.EndFrame();
 			}
@@ -937,7 +937,7 @@ namespace ChaosTest {
 				//tick solver
 				{
 					FVec3 Grav(0,0,-1);
-					Scene.SetUpForFrame(&Grav,1,99999,99999,1,false, -1);
+					Scene.SetUpForFrame(&Grav,1,99999,99999,1,false);
 					Scene.StartFrame();
 					Scene.EndFrame();
 				}
@@ -954,7 +954,7 @@ namespace ChaosTest {
 			//tick solver one last time
 			{
 				FVec3 Grav(0,0,-1);
-				Scene.SetUpForFrame(&Grav,1,99999,99999,1,false, -1);
+				Scene.SetUpForFrame(&Grav,1,99999,99999,1,false);
 				Scene.StartFrame();
 				Scene.EndFrame();
 			}
@@ -984,7 +984,7 @@ namespace ChaosTest {
 			{
 				//tick solver
 				FVec3 Grav(0,0,-1);
-				Scene.SetUpForFrame(&Grav,1,99999,99999,1,false, -1);
+				Scene.SetUpForFrame(&Grav,1,99999,99999,1,false);
 				Scene.StartFrame();
 				Scene.EndFrame();
 
@@ -999,7 +999,7 @@ namespace ChaosTest {
 	{
 		for (int Delay = 0; Delay < 4; ++Delay)
 		{
-			FChaosScene Scene(nullptr);
+			FChaosScene Scene(nullptr, /*AsyncDt=*/-1);
 			Scene.GetSolver()->SetThreadingMode_External(EThreadingModeTemp::SingleThread);
 			Scene.GetSolver()->GetMarshallingManager().SetTickDelay_External(Delay);
 
@@ -1042,7 +1042,7 @@ namespace ChaosTest {
 			{
 				{
 					FVec3 Grav(0,0,0);
-					Scene.SetUpForFrame(&Grav,1,99999,99999,10,false, -1);
+					Scene.SetUpForFrame(&Grav,1,99999,99999,10,false);
 					Scene.StartFrame();
 					Scene.EndFrame();
 				}
@@ -1055,7 +1055,7 @@ namespace ChaosTest {
 			//tick solver and see new position synced from sim
 			{
 				FVec3 Grav(0,0,0);
-				Scene.SetUpForFrame(&Grav,1,99999,99999,10,false, -1);
+				Scene.SetUpForFrame(&Grav,1,99999,99999,10,false);
 				Scene.StartFrame();
 				Scene.EndFrame();
 				EXPECT_NEAR(Particle.X()[2], -1, 1e-4);
@@ -1065,7 +1065,7 @@ namespace ChaosTest {
 			//tick solver and delete in between solver finishing and sync
 			{
 				FVec3 Grav(0,0,0);
-				Scene.SetUpForFrame(&Grav,1,99999,99999,10,false, -1);
+				Scene.SetUpForFrame(&Grav,1,99999,99999,10,false);
 				Scene.StartFrame();
 
 				//delete Proxy
@@ -1081,7 +1081,7 @@ namespace ChaosTest {
 			{
 				{
 					FVec3 Grav(0,0,0);
-					Scene.SetUpForFrame(&Grav,1,99999,99999,10,false, -1);
+					Scene.SetUpForFrame(&Grav,1,99999,99999,10,false);
 					Scene.StartFrame();
 					Scene.EndFrame();
 					EXPECT_NEAR(Particle2.X()[1], -3 - Repeat, 1e-4);	//other Proxy keeps moving
@@ -1094,7 +1094,7 @@ namespace ChaosTest {
 	{
 		for (int Delay = 0; Delay < 4; ++Delay)
 		{
-			FChaosScene Scene(nullptr);
+			FChaosScene Scene(nullptr, /*AsyncDt=*/-1);
 			Scene.GetSolver()->SetThreadingMode_External(EThreadingModeTemp::SingleThread);
 			Scene.GetSolver()->GetMarshallingManager().SetTickDelay_External(Delay);
 
@@ -1124,7 +1124,7 @@ namespace ChaosTest {
 			{
 				{
 					FVec3 Grav(0,0,0);
-					Scene.SetUpForFrame(&Grav,1,99999,99999,10,false, -1);
+					Scene.SetUpForFrame(&Grav,1,99999,99999,10,false);
 					Scene.StartFrame();
 					Scene.EndFrame();
 				}
@@ -1136,7 +1136,7 @@ namespace ChaosTest {
 			//tick solver and see new position synced from sim
 			{
 				FVec3 Grav(0,0,0);
-				Scene.SetUpForFrame(&Grav,1,99999,99999,10,false, -1);
+				Scene.SetUpForFrame(&Grav,1,99999,99999,10,false);
 				Scene.StartFrame();
 				Scene.EndFrame();
 				EXPECT_NEAR(Particle.X()[2], -1, 1e-4);
@@ -1149,7 +1149,7 @@ namespace ChaosTest {
 			{
 				{
 					FVec3 Grav(0,0,0);
-					Scene.SetUpForFrame(&Grav,1,99999,99999,10,false, -1);
+					Scene.SetUpForFrame(&Grav,1,99999,99999,10,false);
 					Scene.StartFrame();
 					Scene.EndFrame();
 
@@ -1160,7 +1160,7 @@ namespace ChaosTest {
 			//tick solver one last time, should see sim results from the place we teleported to
 			{
 				FVec3 Grav(0,0,0);
-				Scene.SetUpForFrame(&Grav,1,99999,99999,10,false, -1);
+				Scene.SetUpForFrame(&Grav,1,99999,99999,10,false);
 				Scene.StartFrame();
 				Scene.EndFrame();
 				EXPECT_NEAR(Particle.X()[2], 9, 1e-4);
@@ -1169,7 +1169,7 @@ namespace ChaosTest {
 			//set x after sim but before EndFrame, make sure to see gt position since it was written after
 			{
 				FVec3 Grav(0,0,0);
-				Scene.SetUpForFrame(&Grav,1,99999,99999,10,false,-1);
+				Scene.SetUpForFrame(&Grav,1,99999,99999,10,false);
 				Scene.StartFrame();
 				FChaosEngineInterface::SetGlobalPose_AssumesLocked(Proxy, FTransform(FQuat::Identity, FVec3(0, 0, 100)));
 				Scene.EndFrame();
@@ -1180,7 +1180,7 @@ namespace ChaosTest {
 			{
 				{
 					FVec3 Grav(0,0,0);
-					Scene.SetUpForFrame(&Grav,1,99999,99999,10,false, -1);
+					Scene.SetUpForFrame(&Grav,1,99999,99999,10,false);
 					Scene.StartFrame();
 					Scene.EndFrame();
 
@@ -1191,7 +1191,7 @@ namespace ChaosTest {
 			//tick solver one last time, should see sim results from the place we teleported to
 			{
 				FVec3 Grav(0,0,0);
-				Scene.SetUpForFrame(&Grav,1,99999,99999,10,false, -1);
+				Scene.SetUpForFrame(&Grav,1,99999,99999,10,false);
 				Scene.StartFrame();
 				Scene.EndFrame();
 				EXPECT_NEAR(Particle.X()[2], 99, 1e-4);
@@ -1201,7 +1201,7 @@ namespace ChaosTest {
 
 	GTEST_TEST(EngineInterface, SimRoundTrip)
 	{
-		FChaosScene Scene(nullptr);
+		FChaosScene Scene(nullptr, /*AsyncDt=*/-1);
 		Scene.GetSolver()->SetThreadingMode_External(EThreadingModeTemp::SingleThread);
 
 		FActorCreationParams Params;
@@ -1222,7 +1222,7 @@ namespace ChaosTest {
 		Particle.AddForce(FVec3(0, 0, 10) * Particle.M());
 
 		FVec3 Grav(0,0,0);
-		Scene.SetUpForFrame(&Grav,1,99999,99999,10,false,-1);
+		Scene.SetUpForFrame(&Grav,1,99999,99999,10,false);
 		Scene.StartFrame();
 		Scene.EndFrame();
 
@@ -1243,9 +1243,10 @@ namespace ChaosTest {
 		//wake events must be collapsed (sleep awake sleep becomes sleep)
 		//collision events must be collapsed
 		//forces are averaged
-		FChaosScene Scene(nullptr);
-		Scene.GetSolver()->SetThreadingMode_External(EThreadingModeTemp::SingleThread);
 		const FReal FixedDT = 1;
+		FChaosScene Scene(nullptr, FixedDT);
+		Scene.GetSolver()->SetThreadingMode_External(EThreadingModeTemp::SingleThread);
+
 
 		FActorCreationParams Params;
 		Params.Scene = &Scene;
@@ -1304,7 +1305,7 @@ namespace ChaosTest {
 			//set force every external frame
 			Particle2.AddForce(ConstantForce);
 			FVec3 Grav(0, 0, 0);
-			Scene.SetUpForFrame(&Grav, GTDt, 99999, 99999, 1, false,FixedDT);
+			Scene.SetUpForFrame(&Grav, GTDt, 99999, 99999, 1, false);
 			Scene.StartFrame();
 			Scene.EndFrame();
 
@@ -1342,7 +1343,7 @@ namespace ChaosTest {
 		// Velocity is zero if only SetGlobalPose_AssumesLocked is called (Teleport)
 		// Particle positions and velocities are correct after SetGlobalPose_AssumesLocked, SetKinematicTarget_AssumesLocked (Teleport)
 		// Particle positions and velocities are correct after SetKinematicTarget_AssumesLocked, SetGlobalPose_AssumesLocked (Teleport, KinematicTarget is cleared)
-		FChaosScene Scene(nullptr);
+		FChaosScene Scene(nullptr, /*AsyncDt=*/-1);
 		Scene.GetSolver()->SetThreadingMode_External(EThreadingModeTemp::SingleThread);
 
 		FActorCreationParams Params;
@@ -1388,7 +1389,7 @@ namespace ChaosTest {
 
 		auto AdvanceFrameAndRunTest = [&](const FVec3 &CorrectX, const FVec3 &CorrectV)
 		{
-			Scene.SetUpForFrame(&Grav, Dt, 99999, 99999, 10, false, -1);
+			Scene.SetUpForFrame(&Grav, Dt, 99999, 99999, 10, false);
 			Scene.StartFrame();
 			Scene.EndFrame();
 			// Test X and V on GT
@@ -1469,9 +1470,9 @@ namespace ChaosTest {
 		//Need to test:
 		//setting transform, velocities, wake state, on external thread means we overwrite results until sim catches up
 		//deleted proxy does not incorrectly update after it's deleted on gt
-		FChaosScene Scene(nullptr);
-		Scene.GetSolver()->SetThreadingMode_External(EThreadingModeTemp::SingleThread);
 		const FReal FixedDT = 1;
+		FChaosScene Scene(nullptr, FixedDT);
+		Scene.GetSolver()->SetThreadingMode_External(EThreadingModeTemp::SingleThread);
 		Scene.GetSolver()->EnableAsyncMode(1);	//tick 1 dt at a time
 
 		FActorCreationParams Params;
@@ -1542,7 +1543,7 @@ namespace ChaosTest {
 			}
 
 			FVec3 Grav(0, 0, 0);
-			Scene.SetUpForFrame(&Grav, GTDt, 99999, 99999, 10, false, FixedDT);
+			Scene.SetUpForFrame(&Grav, GTDt, 99999, 99999, 10, false);
 			Scene.StartFrame();
 			Scene.EndFrame();
 
@@ -1640,7 +1641,7 @@ namespace ChaosTest {
 
 		bool bHitOnShutDown = false;
 		{
-			FChaosScene Scene(nullptr);
+			FChaosScene Scene(nullptr, /*AsyncDt=*/-1);
 			Scene.GetSolver()->SetThreadingMode_External(EThreadingModeTemp::SingleThread);
 			Scene.GetSolver()->EnableAsyncMode(1);	//tick 1 dt at a time
 
@@ -1677,7 +1678,7 @@ namespace ChaosTest {
 			auto Callback = Scene.GetSolver()->CreateAndRegisterSimCallbackObject_External<FCallback>();
 
 			FVec3 Grav(0, 0, 0);
-			Scene.SetUpForFrame(&Grav, 0, 99999, 99999, 10, false,false);	//flush with dt 0
+			Scene.SetUpForFrame(&Grav, 0, 99999, 99999, 10, false);	//flush with dt 0
 			Scene.StartFrame();
 			Scene.EndFrame();
 
@@ -1698,9 +1699,9 @@ namespace ChaosTest {
 		//kinematic targets are interpolated over the sub-step
 		//identical inputs are given to sub-steps
 
-		FChaosScene Scene(nullptr);
-		Scene.GetSolver()->SetThreadingMode_External(EThreadingModeTemp::SingleThread);
 		const FReal FixedDT = 1;
+		FChaosScene Scene(nullptr, FixedDT);
+		Scene.GetSolver()->SetThreadingMode_External(EThreadingModeTemp::SingleThread);
 
 		FActorCreationParams Params;
 		Params.Scene = &Scene;
@@ -1750,7 +1751,7 @@ namespace ChaosTest {
 			//set force every external frame
 			Particle.AddForce(FVec3(0, 0, 1 * Particle.M()));	//should counteract gravity
 			FVec3 Grav(0, 0, -1);
-			Scene.SetUpForFrame(&Grav, GTDt, 99999, 99999, 10, false, FixedDT);
+			Scene.SetUpForFrame(&Grav, GTDt, 99999, 99999, 10, false);
 			Scene.StartFrame();
 			Scene.EndFrame();
 
@@ -1768,7 +1769,7 @@ namespace ChaosTest {
 		//destroyed proxy still valid in callback, but Proxy is nulled out
 		//valid for multiple sub-steps
 
-		FChaosScene Scene(nullptr);
+		FChaosScene Scene(nullptr, /*AsyncDt=*/-1);
 		Scene.GetSolver()->SetThreadingMode_External(EThreadingModeTemp::SingleThread);
 		const FReal FixedDT = 1;
 		Scene.GetSolver()->EnableAsyncMode(FixedDT);	//tick 1 dt at a time
@@ -1808,14 +1809,14 @@ namespace ChaosTest {
 		Scene.GetSolver()->UnregisterObject(Proxy);
 
 		FVec3 Grav(0, 0, -1);
-		Scene.SetUpForFrame(&Grav, FixedDT * 3, 99999, 99999, 10, false, false);
+		Scene.SetUpForFrame(&Grav, FixedDT * 3, 99999, 99999, 10, false);
 		Scene.StartFrame();
 		Scene.EndFrame();
 	}
 	
 	GTEST_TEST(EngineInterface, OverlapOffsetActor)
 	{
-		FChaosScene Scene(nullptr);
+		FChaosScene Scene(nullptr, /*AsyncDt=*/-1);
 
 		FActorCreationParams Params;
 		Params.Scene = &Scene;
diff --git a/Engine/Source/Programs/HeadlessChaos/Private/HeadlessChaosTestRewind.cpp b/Engine/Source/Programs/HeadlessChaos/Private/HeadlessChaosTestRewind.cpp
--- a/Engine/Source/Programs/HeadlessChaos/Private/HeadlessChaosTestRewind.cpp
+++ b/Engine/Source/Programs/HeadlessChaos/Private/HeadlessChaosTestRewind.cpp
@@ -38,7 +38,7 @@ namespace ChaosTest {
 
 		// Make a solver
 		FChaosSolversModule* Module = FChaosSolversModule::GetModule();
-		auto* Solver = Module->CreateSolver(nullptr);
+		auto* Solver = Module->CreateSolver(nullptr,/*AsyncDt=*/-1);
                 InitSolverSettings(Solver);
 
 		Solver->EnableRewindCapture(RewindHistorySize, !!Optimization);
@@ -2340,7 +2340,7 @@ namespace ChaosTest {
 			FChaosSolversModule* Module = FChaosSolversModule::GetModule();
 
 			// Make a solver
-			auto* Solver = Module->CreateSolver(nullptr);
+			auto* Solver = Module->CreateSolver(nullptr, /*AsyncDt=*/-1);
 			InitSolverSettings(Solver);
 
 			//note: this 5 is just a suggestion, there could be more frames saved than that
@@ -2406,7 +2406,7 @@ namespace ChaosTest {
 			FChaosSolversModule* Module = FChaosSolversModule::GetModule();
 
 			// Make a solver
-			auto* Solver = Module->CreateSolver(nullptr);
+			auto* Solver = Module->CreateSolver(nullptr, /*AsyncDt=*/-1);
 
 
 			InitSolverSettings(Solver);
@@ -2565,7 +2565,7 @@ namespace ChaosTest {
 			FChaosSolversModule* Module = FChaosSolversModule::GetModule();
 
 			// Make a solver
-			auto* Solver = Module->CreateSolver(nullptr);
+			auto* Solver = Module->CreateSolver(nullptr, /*AsyncDt=*/-1);
 			InitSolverSettings(Solver);
 
 			Solver->EnableRewindCapture(7, !!Optimization);
@@ -2661,7 +2661,7 @@ namespace ChaosTest {
 			FChaosSolversModule* Module = FChaosSolversModule::GetModule();
 
 			// Make a solver
-			auto* Solver = Module->CreateSolver(nullptr);
+			auto* Solver = Module->CreateSolver(nullptr, /*AsyncDt=*/-1);
 			InitSolverSettings(Solver);
 
 			Solver->EnableRewindCapture(7, !!Optimization);
@@ -2754,7 +2754,7 @@ namespace ChaosTest {
 			FChaosSolversModule* Module = FChaosSolversModule::GetModule();
 
 			// Make a solver
-			auto* Solver = Module->CreateSolver(nullptr);
+			auto* Solver = Module->CreateSolver(nullptr, /*AsyncDt=*/-1);
 			InitSolverSettings(Solver);
 
 			Solver->EnableRewindCapture(7, !!Optimization);
@@ -2852,7 +2852,7 @@ namespace ChaosTest {
 			FChaosSolversModule* Module = FChaosSolversModule::GetModule();
 
 			// Make a solver
-			auto* Solver = Module->CreateSolver(nullptr);
+			auto* Solver = Module->CreateSolver(nullptr, /*AsyncDt=*/-1);
 			InitSolverSettings(Solver);
 			
 			Solver->EnableRewindCapture(7, !!Optimization);
@@ -2902,7 +2902,7 @@ namespace ChaosTest {
 			FChaosSolversModule* Module = FChaosSolversModule::GetModule();
 
 			// Make a solver
-			auto* Solver = Module->CreateSolver(nullptr);
+			auto* Solver = Module->CreateSolver(nullptr, /*AsyncDt=*/-1);
 			InitSolverSettings(Solver);
 
 			Solver->EnableRewindCapture(7, !!Optimization);
@@ -2994,7 +2994,7 @@ namespace ChaosTest {
 			FChaosSolversModule* Module = FChaosSolversModule::GetModule();
 
 			// Make a solver
-			auto* Solver = Module->CreateSolver(nullptr);
+			auto* Solver = Module->CreateSolver(nullptr, /*AsyncDt=*/-1);
 			InitSolverSettings(Solver);
 
 			Solver->EnableRewindCapture(7, !!Optimization);
@@ -3085,7 +3085,7 @@ namespace ChaosTest {
 			FChaosSolversModule* Module = FChaosSolversModule::GetModule();
 
 			// Make a solver
-			auto* Solver = Module->CreateSolver(nullptr);
+			auto* Solver = Module->CreateSolver(nullptr, /*AsyncDt=*/-1);
 			InitSolverSettings(Solver);
 
 			Solver->EnableRewindCapture(100, !!Optimization);
@@ -3177,7 +3177,7 @@ namespace ChaosTest {
 			FChaosSolversModule* Module = FChaosSolversModule::GetModule();
 
 			// Make a solver
-			auto* Solver = Module->CreateSolver(nullptr);
+			auto* Solver = Module->CreateSolver(nullptr, /*AsyncDt=*/-1);
 			InitSolverSettings(Solver);
 
 			Solver->EnableRewindCapture(100, !!Optimization);
@@ -3265,7 +3265,7 @@ namespace ChaosTest {
 			FChaosSolversModule* Module = FChaosSolversModule::GetModule();
 
 			// Make a solver
-			auto* Solver = Module->CreateSolver(nullptr);
+			auto* Solver = Module->CreateSolver(nullptr, /*AsyncDt=*/-1);
 			InitSolverSettings(Solver);
 
 			Solver->EnableRewindCapture(7, !!Optimization);
@@ -3345,7 +3345,7 @@ namespace ChaosTest {
 			FChaosSolversModule* Module = FChaosSolversModule::GetModule();
 
 			// Make a solver
-			auto* Solver = Module->CreateSolver(nullptr);
+			auto* Solver = Module->CreateSolver(nullptr, /*AsyncDt=*/-1);
 			InitSolverSettings(Solver);
 
 			Solver->EnableRewindCapture(7, !!Optimization);
@@ -3435,7 +3435,7 @@ namespace ChaosTest {
 			FChaosSolversModule* Module = FChaosSolversModule::GetModule();
 
 			// Make a solver
-			auto* Solver = Module->CreateSolver(nullptr);
+			auto* Solver = Module->CreateSolver(nullptr, /*AsyncDt=*/-1);
 			InitSolverSettings(Solver);
 
 			Solver->EnableRewindCapture(7, !!Optimization);
@@ -3526,7 +3526,7 @@ namespace ChaosTest {
 			FChaosSolversModule* Module = FChaosSolversModule::GetModule();
 
 			// Make a solver
-			auto* Solver = Module->CreateSolver(nullptr);
+			auto* Solver = Module->CreateSolver(nullptr, /*AsyncDt=*/-1);
 			InitSolverSettings(Solver);
 
 			Solver->EnableRewindCapture(100, !!Optimization);
@@ -3633,7 +3633,7 @@ namespace ChaosTest {
 		FChaosSolversModule* Module = FChaosSolversModule::GetModule();
 
 		// Make a solver
-		auto* Solver = Module->CreateSolver(nullptr);
+		auto* Solver = Module->CreateSolver(nullptr, /*AsyncDt=*/-1);
 		InitSolverSettings(Solver);
 
 		Solver->EnableRewindCapture(100,true);	//soft desync only exists when resim optimization is on
@@ -3744,7 +3744,7 @@ namespace ChaosTest {
 		FChaosSolversModule* Module = FChaosSolversModule::GetModule();
 
 		// Make a solver
-		auto* Solver = Module->CreateSolver(nullptr);
+		auto* Solver = Module->CreateSolver(nullptr, /*AsyncDt=*/-1);
 		InitSolverSettings(Solver);
 
 		Solver->EnableRewindCapture(100,true);	//soft desync only exists when resim optimization is on
@@ -3981,7 +3981,7 @@ namespace ChaosTest {
 		FChaosSolversModule* Module = FChaosSolversModule::GetModule();
 
 		// Make a solver
-		auto* Solver = Module->CreateSolver(nullptr);
+		auto* Solver = Module->CreateSolver(nullptr, /*AsyncDt=*/-1);
 		InitSolverSettings(Solver);
 
 		TArray<FPhysicsActorHandle> Storage = InitFunc(Solver);
diff --git a/Engine/Source/Programs/HeadlessChaos/Private/HeadlessChaosTestSim.cpp b/Engine/Source/Programs/HeadlessChaos/Private/HeadlessChaosTestSim.cpp
--- a/Engine/Source/Programs/HeadlessChaos/Private/HeadlessChaosTestSim.cpp
+++ b/Engine/Source/Programs/HeadlessChaos/Private/HeadlessChaosTestSim.cpp
@@ -246,7 +246,7 @@ namespace ChaosTest {
 		TSharedPtr<FImplicitObject, ESPMode::ThreadSafe> Sphere{new TSphere<FReal, 3>(FVec3(0), 10)};
 
 		FChaosSolversModule* Module = FChaosSolversModule::GetModule();
-		auto Solver = Module->CreateSolver(nullptr);
+		auto Solver = Module->CreateSolver(nullptr, /*AsyncDt=*/-1);
 
 		FSingleParticlePhysicsProxy* Proxy = FSingleParticlePhysicsProxy::Create(Chaos::TPBDRigidParticle<FReal, 3>::CreateParticle());
 		Chaos::FRigidBodyHandle_External& Particle = Proxy->GetGameThreadAPI();
diff --git a/Engine/Source/Programs/HeadlessChaos/Private/HeadlessChaosTestSolverProxies.cpp b/Engine/Source/Programs/HeadlessChaos/Private/HeadlessChaosTestSolverProxies.cpp
--- a/Engine/Source/Programs/HeadlessChaos/Private/HeadlessChaosTestSolverProxies.cpp
+++ b/Engine/Source/Programs/HeadlessChaos/Private/HeadlessChaosTestSolverProxies.cpp
@@ -27,7 +27,7 @@ namespace ChaosTest {
 		FChaosSolversModule* Module = FChaosSolversModule::GetModule();
 
 		// Make a solver
-		auto* Solver = Module->CreateSolver(nullptr);
+		auto* Solver = Module->CreateSolver(nullptr, /*AsyncDt=*/-1);
 
 		// Make a particle
 
@@ -69,7 +69,7 @@ namespace ChaosTest {
 		FChaosSolversModule* Module = FChaosSolversModule::GetModule();
 
 		// Make a solver
-		auto* Solver = Module->CreateSolver(nullptr);
+		auto* Solver = Module->CreateSolver(nullptr, /*AsyncDt=*/-1);
 
 		// Make a particle
 
