headlesschaos
-- fixed Initilization_TransformedGeometryCollectionGeometryVertices

diff --git a/Engine/Source/Programs/HeadlessChaos/Private/GeometryCollection/GeometryCollectionTestInitilization.cpp b/Engine/Source/Programs/HeadlessChaos/Private/GeometryCollection/GeometryCollectionTestInitilization.cpp
--- a/Engine/Source/Programs/HeadlessChaos/Private/GeometryCollection/GeometryCollectionTestInitilization.cpp
+++ b/Engine/Source/Programs/HeadlessChaos/Private/GeometryCollection/GeometryCollectionTestInitilization.cpp
@@ -103,7 +103,9 @@ GTEST_TEST(AllTraits,GeometryCollection_Initilization_TransformedGeometryCollect
 	}
 }
 
-
+	
+	// We expect that if we actually apply the global matrices that it does not care about MassToLocal. To make this accurate, 
+	// X and Y should still be at origin, and check that Z is < 0. IT shouldnt make a difference if we move the geometry...
 GTEST_TEST(AllTraits,GeometryCollection_Initilization_TransformedGeometryCollectionGeometryVertices)
 {
 	FVector GlobalTranslation(10); FQuat GlobalRotation = FQuat::MakeFromEuler(FVector(0));
@@ -130,8 +132,6 @@ GTEST_TEST(AllTraits,GeometryCollection_Initilization_TransformedGeometryCollect
 	EXPECT_NEAR(CenterOfMass.Y - GlobalTranslation[1],0.0f,KINDA_SMALL_NUMBER);
 	EXPECT_NEAR(CenterOfMass.Z - GlobalTranslation[2],0.0f,KINDA_SMALL_NUMBER);
 
-
-
 	FFramework UnitTest;
 	UnitTest.AddSimulationObject(Collection);
 	UnitTest.Initialize();
@@ -141,20 +141,19 @@ GTEST_TEST(AllTraits,GeometryCollection_Initilization_TransformedGeometryCollect
 		TManagedArray<FTransform>& MassToLocal = Collection->RestCollection->template GetAttribute<FTransform>("MassToLocal",FGeometryCollection::TransformGroup);
 
 
-		//EXPECT_EQ(UnitTest.Solver->GetParticles().GetGeometryCollectionParticles().Size(), 1);
-		//FVector X = UnitTest.Solver->GetParticles().GetGeometryCollectionParticles().X(0);
-		//FQuat R = UnitTest.Solver->GetParticles().GetGeometryCollectionParticles().R(0);
-		//EXPECT_TRUE((R * GlobalRotation.Inverse()).IsIdentity(KINDA_SMALL_NUMBER));
-		//EXPECT_NEAR(X.X - GlobalTranslation[0], 0.0f, KINDA_SMALL_NUMBER);
-		//EXPECT_NEAR(X.Y - GlobalTranslation[1], 0.0f, KINDA_SMALL_NUMBER);
-		//EXPECT_LT(X.Z, GlobalTranslation[2]);
-
-		//TArray<FTransform> Transform;
-		//GeometryCollectionAlgo::GlobalMatrices(Collection->DynamicCollection->Transform, Collection->DynamicCollection->Parent, Transform);
-		//EXPECT_EQ(Collection->DynamicCollection->Parent[0], FGeometryCollection::Invalid); // is not a child
-		//EXPECT_NEAR(Transform[0].GetTranslation().X - GlobalTranslation[0], 0.0f, KINDA_SMALL_NUMBER);
-		//EXPECT_NEAR(Transform[0].GetTranslation().Y - GlobalTranslation[1], 0.0f, KINDA_SMALL_NUMBER);
-		//EXPECT_LT(Transform[0].GetTranslation().Z, GlobalTranslation[2]);
+		EXPECT_EQ(UnitTest.Solver->GetParticles().GetGeometryCollectionParticles().Size(), 1);
+		FVector ParticlePos = UnitTest.Solver->GetParticles().GetGeometryCollectionParticles().X(0);
+		FQuat ParticleR = UnitTest.Solver->GetParticles().GetGeometryCollectionParticles().R(0);
+
+		TArray<FTransform> Transform;
+		GeometryCollectionAlgo::GlobalMatrices(Collection->DynamicCollection->Transform, Collection->DynamicCollection->Parent, Transform);
+		EXPECT_EQ(Collection->DynamicCollection->Parent[0], FGeometryCollection::Invalid); // is not a child
+		EXPECT_NEAR(ParticlePos.X, GlobalTranslation[0], KINDA_SMALL_NUMBER); // Check particle position from solver 
+		EXPECT_NEAR(ParticlePos.Y, GlobalTranslation[1], KINDA_SMALL_NUMBER);
+		EXPECT_NEAR(Transform[0].GetTranslation().X, 0.0f, KINDA_SMALL_NUMBER);
+		EXPECT_NEAR(Transform[0].GetTranslation().Y, 0.0f, KINDA_SMALL_NUMBER);
+		EXPECT_LT(Transform[0].GetTranslation().Z, 0.0f); // Check collection fell
+		EXPECT_GT(Transform[0].GetTranslation().Z, -1.0f);
 	}
 }
 
