headlesschaos
-- exercising basic field algebraic properties

diff --git a/Engine/Source/Programs/HeadlessChaos/Private/GeometryCollection/GeometryCollectionTestSimulationField.cpp b/Engine/Source/Programs/HeadlessChaos/Private/GeometryCollection/GeometryCollectionTestSimulationField.cpp
--- a/Engine/Source/Programs/HeadlessChaos/Private/GeometryCollection/GeometryCollectionTestSimulationField.cpp
+++ b/Engine/Source/Programs/HeadlessChaos/Private/GeometryCollection/GeometryCollectionTestSimulationField.cpp
@@ -826,6 +826,64 @@ namespace GeometryCollectionTest
 		delete FalloffField;
 	}
 
+	GTEST_TEST(AllTraits, GeometryCollection_RigidBodies_Field_Algebra1)
+	{
+		FFramework UnitTest;
+		FVector ExpectedLocation = FVector(0.0f, 0.0f, 0.0f);
+		FReal LastZ = ExpectedLocation.Z;
+
+		// Physics Object Setup
+		CreationParameters Params;
+		Params.DynamicState = EObjectStateTypeEnum::Chaos_Object_Dynamic;
+		Params.RootTransform.SetLocation(ExpectedLocation);
+		FGeometryCollectionWrapper* Collection = TNewSimulationObject<GeometryType::GeometryCollectionWithSingleRigid>::Init(Params)->template As<FGeometryCollectionWrapper>();
+		UnitTest.AddSimulationObject(Collection);
+
+		// Opposing Fields setup (velocity)
+		FUniformVector* VectorFieldRight = new FUniformVector();
+		VectorFieldRight->Magnitude = 100.0;
+		VectorFieldRight->Direction = FVector(1.0, 0.0, 0.0);
+		FUniformVector* VectorFieldLeft = new FUniformVector();
+		VectorFieldLeft->Magnitude = 100.0;
+		VectorFieldLeft->Direction = FVector(-1.0f, 0.0, 0.0);
+
+		// Opposing Fields setup (force)
+		FUniformVector* VectorFieldForward = new FUniformVector();
+		VectorFieldRight->Magnitude = 100.0;
+		VectorFieldRight->Direction = FVector(0.0, 100.0, 0.0);
+		FUniformVector* VectorFieldBackward = new FUniformVector();
+		VectorFieldLeft->Magnitude = 100.0;
+		VectorFieldLeft->Direction = FVector(0.0f, -100.0, 0.0);
+
+		UnitTest.Initialize();
+
+		FName TargetNameVel = GetFieldPhysicsName(EFieldPhysicsType::Field_LinearVelocity);
+		UnitTest.Solver->GetPerSolverField().AddTransientCommand({ TargetNameVel, VectorFieldRight->NewCopy() });
+		UnitTest.Solver->GetPerSolverField().AddTransientCommand({ TargetNameVel, VectorFieldLeft->NewCopy() });
+		FName TargetNameForce = GetFieldPhysicsName(EFieldPhysicsType::Field_LinearForce);
+		UnitTest.Solver->GetPerSolverField().AddTransientCommand({ TargetNameForce, VectorFieldForward->NewCopy() });
+		UnitTest.Solver->GetPerSolverField().AddTransientCommand({ TargetNameForce, VectorFieldBackward->NewCopy() });
+		UnitTest.Advance();
+
+		TManagedArray<FTransform>& Transform = Collection->DynamicCollection->Transform;
+		for (int Frame = 1; Frame < 10; Frame++)
+		{
+			UnitTest.Solver->GetPerSolverField().AddTransientCommand({ GetFieldPhysicsName(EFieldPhysicsType::Field_LinearVelocity), VectorFieldRight->NewCopy() });
+			UnitTest.Solver->GetPerSolverField().AddTransientCommand({ GetFieldPhysicsName(EFieldPhysicsType::Field_LinearVelocity), VectorFieldLeft->NewCopy() });
+			UnitTest.Solver->GetPerSolverField().AddTransientCommand({ GetFieldPhysicsName(EFieldPhysicsType::Field_LinearForce), VectorFieldForward->NewCopy() });
+			UnitTest.Solver->GetPerSolverField().AddTransientCommand({ GetFieldPhysicsName(EFieldPhysicsType::Field_LinearForce), VectorFieldBackward->NewCopy() });
+
+			UnitTest.Advance();
+
+			EXPECT_NEAR(Transform[0].GetTranslation().X, ExpectedLocation.X, KINDA_SMALL_NUMBER);
+			EXPECT_NEAR(Transform[0].GetTranslation().Y, ExpectedLocation.Y, KINDA_SMALL_NUMBER);
+			EXPECT_LT(Transform[0].GetTranslation().Z, ExpectedLocation.Z);
+			EXPECT_LT(Transform[0].GetTranslation().Z, LastZ);
+			LastZ = Transform[0].GetTranslation().Z;
+		}
+	}
+
+
 }
 
 
