headlesschaos
-- initial velocity field algebraic properties

diff --git a/Engine/Source/Programs/HeadlessChaos/Private/GeometryCollection/GeometryCollectionTestSimulationField.cpp b/Engine/Source/Programs/HeadlessChaos/Private/GeometryCollection/GeometryCollectionTestSimulationField.cpp
--- a/Engine/Source/Programs/HeadlessChaos/Private/GeometryCollection/GeometryCollectionTestSimulationField.cpp
+++ b/Engine/Source/Programs/HeadlessChaos/Private/GeometryCollection/GeometryCollectionTestSimulationField.cpp
@@ -884,6 +884,54 @@ namespace GeometryCollectionTest
 	}
 
 
+	GTEST_TEST(AllTraits, GeometryCollection_RigidBodies_Field_Algebra2)
+	{
+		FFramework UnitTest;
+		FVector ExpectedLocation = FVector(0.0f, 0.0f, 0.0f);
+		FVector LastLocation = ExpectedLocation;
+
+		// Physics Object Setup
+		CreationParameters Params;
+		Params.DynamicState = EObjectStateTypeEnum::Chaos_Object_Dynamic;
+		Params.InitialVelocityType = EInitialVelocityTypeEnum::Chaos_Initial_Velocity_User_Defined;
+		Params.InitialLinearVelocity = FVector(100.0, 0.0, 0.0);
+		Params.RootTransform.SetLocation(ExpectedLocation);
+		FGeometryCollectionWrapper* Collection = TNewSimulationObject<GeometryType::GeometryCollectionWithSingleRigid>::Init(Params)->template As<FGeometryCollectionWrapper>();
+		UnitTest.AddSimulationObject(Collection);
+
+		// Opposing Fields setup (force)
+		FUniformVector* VectorFieldRightForce = new FUniformVector();
+		VectorFieldRightForce->Magnitude = 100.0;
+		VectorFieldRightForce->Direction = FVector(100.0, 0.0, 0.0);
+		FUniformVector* VectorFieldLeftForce = new FUniformVector();
+		VectorFieldLeftForce->Magnitude = 100.0;
+		VectorFieldLeftForce->Direction = FVector(-100.0f, 0.0, 0.0);
+
+		UnitTest.Initialize();
+		UnitTest.Advance();
+
+		TArray<Chaos::TPBDRigidClusteredParticleHandle<FReal, 3>*>& ParticleHandles = Collection->PhysObject->GetSolverParticleHandles();
+		Chaos::TVector<float,3> CurrV = ParticleHandles[0]->V();
+		TManagedArray<FTransform>& Transform = Collection->DynamicCollection->Transform;
+
+		for (int Frame = 1; Frame < 10; Frame++)
+		{
+			UnitTest.Solver->GetPerSolverField().AddTransientCommand({ GetFieldPhysicsName(EFieldPhysicsType::Field_LinearForce), VectorFieldRightForce->NewCopy() });
+			UnitTest.Solver->GetPerSolverField().AddTransientCommand({ GetFieldPhysicsName(EFieldPhysicsType::Field_LinearForce), VectorFieldLeftForce->NewCopy() });
+
+			UnitTest.Advance();
+
+			CurrV = ParticleHandles[0]->V();
+			EXPECT_NEAR(CurrV.X, Params.InitialLinearVelocity.X, KINDA_SMALL_NUMBER); // Velocity in +x
+			EXPECT_GT(Transform[0].GetTranslation().X, LastLocation.X); // Pos in +x
+
+			// Still falling?
+			EXPECT_LT(Transform[0].GetTranslation().Z, LastLocation.Z);
+
+			LastLocation = Transform[0].GetTranslation();
+			
+		}
+	}
 }
 
 
