Revived GeometryCollection_RigidBodies_CollisionGroup

diff --git a/Engine/Source/Programs/HeadlessChaos/Private/GeometryCollection/GeometryCollectionTestSimulation.cpp b/Engine/Source/Programs/HeadlessChaos/Private/GeometryCollection/GeometryCollectionTestSimulation.cpp
--- a/Engine/Source/Programs/HeadlessChaos/Private/GeometryCollection/GeometryCollectionTestSimulation.cpp
+++ b/Engine/Source/Programs/HeadlessChaos/Private/GeometryCollection/GeometryCollectionTestSimulation.cpp
@@ -33,7 +33,7 @@ namespace GeometryCollectionTest
 	{
 		FGeometryCollectionWrapper* Collection = TNewSimulationObject<GeometryType::GeometryCollectionWithSingleRigid>::Init()->template As<FGeometryCollectionWrapper>();
 
-		FFramework UnitTest;
+		FFramework UnitTest; 
 		UnitTest.AddSimulationObject(Collection);
 		UnitTest.Initialize();
 		UnitTest.Advance();
@@ -220,97 +220,79 @@ namespace GeometryCollectionTest
 
 	}
 
-	
+	// CollisionGroup == 0 : Collide_With_Everything
+	// CollisionGroup == INDEX_NONE : Disabled collisions
+	// CollisionGroup_A != CollisionGroup_B : Skip Check
 	GTEST_TEST(AllTraits, GeometryCollection_RigidBodies_CollisionGroup)
 	{
-		/*
-		TUniquePtr<Chaos::FChaosPhysicsMaterial> PhysicalMaterial = nullptr;
-		TSharedPtr<FGeometryCollection> RestCollection = nullptr;
-		TSharedPtr<FGeometryDynamicCollection> DynamicCollection = nullptr;
-
-		//
-		//  Rigid Body Setup
-		//
-		auto RestInitFunc = [](TSharedPtr<FGeometryCollection>& RestCollection)
-		{
-			RestCollection->AppendGeometry(*GeometryCollection::MakeCubeElement(FTransform(FVector(0, 0, 210.0)), FVector(100.0)));
-			RestCollection->AppendGeometry(*GeometryCollection::MakeCubeElement(FTransform(FVector(0, 0, 320.0)), FVector(100.0)));
-			RestCollection->AppendGeometry(*GeometryCollection::MakeCubeElement(FTransform(FVector(0, 0, 430.0)), FVector(100.0)));
-		};
-
-		////InitCollectionsParameters InitParams = { FTransform(FVector(0, 0, 100.0)), FVector(100.0), RestInitFunc, (int32)EObjectStateTypeEnum::Chaos_Object_Kinematic };
-		//InitCollections(PhysicalMaterial, RestCollection, DynamicCollection, InitParams);
-
-		//
-		// Solver setup
-		//
-		auto CustomFunc = [&RestCollection, &DynamicCollection, &PhysicalMaterial](FSimulationParameters& InParams)
-		{
-			InParams.Shared.SizeSpecificData[0].ImplicitType = EImplicitTypeEnum::Chaos_Implicit_Box;
-		};
 
-		FGeometryCollectionPhysicsProxy* PhysObject = RigidBodySetup(PhysicalMaterial, RestCollection, DynamicCollection, CustomFunc);
-		PhysObject->SetCollisionParticlesPerObjectFraction( 1.0 );
+		FFramework UnitTest;
+
+		TSharedPtr<FGeometryCollection> RestCollection;
+		RestCollection = GeometryCollection::MakeCubeElement(FTransform(FVector(0.f, 0.f, 210.f)), FVector(100.0));
+		RestCollection->AppendGeometry(*GeometryCollection::MakeCubeElement(FTransform(FVector(0.f, 0.f, 320.f)), FVector(100.0)));
+		RestCollection->AppendGeometry(*GeometryCollection::MakeCubeElement(FTransform(FVector(0.f, 0.f, 430.f)), FVector(100.0)));
+		RestCollection->AppendGeometry(*GeometryCollection::MakeCubeElement(FTransform(FVector(0.f, 0.f, 540.f)), FVector(100.0)));
+		RestCollection->AppendGeometry(*GeometryCollection::MakeCubeElement(FTransform(FVector(0.f, 0.f, 650.f)), FVector(100.0)));
 
-		Chaos::FPBDRigidsSolver* Solver = FChaosSolversModule::GetModule()->CreateSolver(nullptr, ESolverFlags::Standalone);
-#if CHAOS_PARTICLEHANDLE_TODO
-		Solver->RegisterObject(PhysObject);
-#endif
-		//Solver->SetHasFloor(true);
-		Solver->SetEnabled(true);
-		//PhysObject->ActivateBodies();
+		CreationParameters Params;
+		Params.RestCollection = RestCollection;
+		Params.ImplicitType = EImplicitTypeEnum::Chaos_Implicit_Box;
+		// I think there is suppose to be one more input param, but not sure what it is for...
 
-		Solver->AdvanceSolverBy(1 / 24.);
-#if TODO_REIMPLEMENT_GET_RIGID_PARTICLES
-		Chaos::TPBDRigidParticles<FReal, 3>& Particles = Solver->GetRigidParticles();
+		FGeometryCollectionWrapper* Collection = TNewSimulationObject<GeometryType::GeometryCollectionWithSuppliedRestCollection>::Init(Params)->template As<FGeometryCollectionWrapper>();
+		RigidBodyWrapper* Floor = TNewSimulationObject<GeometryType::RigidFloor>::Init()->template As<RigidBodyWrapper>();
+		UnitTest.AddSimulationObject(Collection);
+		UnitTest.AddSimulationObject(Floor);
+		UnitTest.Initialize();
+		UnitTest.Advance();
 
+		// testing...
+		auto GCParticles = Collection->PhysObject->GetSolverParticleHandles();
+		TManagedArray<FTransform>& Transform = Collection->DynamicCollection->Transform;
 		for (int Frame = 1; Frame < 200; Frame++)
 		{
-			Solver->AdvanceSolverBy(1 / 24.);
-			//FinalizeSolver(*Solver);
-			
 			if (Frame == 1)
 			{
-				Particles.CollisionGroup(0)=  0;
-				Particles.CollisionGroup(1)=  1;
-				Particles.CollisionGroup(2)=  1;
-				Particles.CollisionGroup(3)=  3;
-				Particles.CollisionGroup(4)= -1;
-			}
-			if (Frame == 13)
-			{
-				EXPECT_LT(FMath::Abs(Particles.X(0).Z), SMALL_NUMBER);
-				EXPECT_LT(FMath::Abs(Particles.X(1).Z - 50.f), 10.f);
-				EXPECT_LT(FMath::Abs(Particles.X(2).Z - 150.f), 10.f);
+				GCParticles[0]->SetCollisionGroup(0);
+				GCParticles[1]->SetCollisionGroup(1);
+				GCParticles[2]->SetCollisionGroup(1);
+				GCParticles[3]->SetCollisionGroup(3);
+				GCParticles[4]->SetCollisionGroup(-1);
+
+				EXPECT_TRUE(Transform[0].GetRotation() == FQuat::Identity); // Can use defaulted zero rotation to indicate that the
+				EXPECT_TRUE(Transform[1].GetRotation() == FQuat::Identity); // rigid has not been affected.
+				EXPECT_TRUE(Transform[2].GetRotation() == FQuat::Identity);
+				EXPECT_TRUE(Transform[3].GetRotation() == FQuat::Identity);
+				EXPECT_TRUE(Transform[4].GetRotation() == FQuat::Identity);
 			}
-			if( Frame == 30 )
-			{
-				EXPECT_LT(FMath::Abs(Particles.X(0).Z), SMALL_NUMBER);
-				EXPECT_LT(FMath::Abs(Particles.X(1).Z - 50.f), 10.f);
-				EXPECT_LT(FMath::Abs(Particles.X(2).Z - 150.f), 10.f);
-				EXPECT_GT(Particles.X(3).Z, 50.f);
-				EXPECT_LT(Particles.X(4).Z, -100);
-			}
-			if (Frame == 31)
+
+			if (Frame == 55)
 			{
-				Particles.CollisionGroup(0) = 0;
-				Particles.CollisionGroup(1) = -1;
-				Particles.CollisionGroup(2) = 1;
-				Particles.CollisionGroup(3) = -1;
-				Particles.CollisionGroup(4) = -1;
+				EXPECT_NEAR(Transform[0].GetTranslation().Z, 50.0f, 1.0f);
+				EXPECT_NEAR(Transform[1].GetTranslation().Z, 50.0f, 1.0f);
+				EXPECT_NEAR(Transform[2].GetTranslation().Z, 50.0f, 1.0f);
+				EXPECT_FALSE(Transform[0].GetRotation() == FQuat::Identity); 
+				EXPECT_FALSE(Transform[1].GetRotation() == FQuat::Identity);
+				EXPECT_FALSE(Transform[2].GetRotation() == FQuat::Identity);
+				EXPECT_TRUE(Transform[3].GetRotation() == FQuat::Identity);
+				EXPECT_TRUE(Transform[4].GetRotation() == FQuat::Identity);
 			}
+			UnitTest.Advance();
 		}
 
-		EXPECT_LT(FMath::Abs(Particles.X(0).Z), SMALL_NUMBER);
-		EXPECT_LT(Particles.X(1).Z, -10000);
-		EXPECT_GT(Particles.X(2).Z, 50.0);
-		EXPECT_LT(Particles.X(3).Z, -10000);
-		EXPECT_LT(Particles.X(4).Z, -10000);
-#endif
+		EXPECT_NEAR(Transform[0].GetTranslation().Z, 50.0f, 1.0f);
+		EXPECT_NEAR(Transform[1].GetTranslation().Z, 50.0f, 1.0f);
+		EXPECT_NEAR(Transform[2].GetTranslation().Z, 50.0f, 1.0f);
+		EXPECT_NEAR(Transform[3].GetTranslation().Z, 50.0f, 1.0f);
+		EXPECT_FALSE(Transform[3].GetRotation() == FQuat::Identity);
+		EXPECT_TRUE(Transform[4].GetRotation() == FQuat::Identity); // Phased through everything, good.
+		EXPECT_LT(Transform[4].GetTranslation().Z, -100.0f);
+
+		GCParticles[0]->SetCollisionGroup(-1);
+		for (int i = 0; i < 50; i++) { UnitTest.Advance(); }
+		EXPECT_LT(Transform[0].GetTranslation().Z, -100.0f);
 
-		FChaosSolversModule::GetModule()->DestroySolver(Solver);
-		delete PhysObject;
-		*/
 	}
 
 
