Fix FJointConstraintPhysicsProxy adding to internal proxy array from game thread instead of physics thread.

diff --git a/Engine/Source/Programs/HeadlessChaos/Private/HeadlessChaosTestSolverProxies.cpp b/Engine/Source/Programs/HeadlessChaos/Private/HeadlessChaosTestSolverProxies.cpp
--- a/Engine/Source/Programs/HeadlessChaos/Private/HeadlessChaosTestSolverProxies.cpp
+++ b/Engine/Source/Programs/HeadlessChaos/Private/HeadlessChaosTestSolverProxies.cpp
@@ -4,7 +4,7 @@
 #include "HeadlessChaos.h"
 #include "HeadlessChaosTestUtility.h"
 
-
+#include "Chaos/ChaosEngineInterface.h"
 #include "Chaos/ParticleHandle.h"
 #include "Chaos/ErrorReporter.h"
 #include "PhysicsProxy/SingleParticlePhysicsProxy.h"
@@ -117,9 +117,66 @@ namespace ChaosTest {
 		Module->DestroySolver(Solver);
 	}
 
+	void JointProxySolverInitTest()
+	{
+		// Test that when we create a joint constraint, that the proxy is not added to solver proxy array by game thread.
+		// This array is physics thread only, it should only be added after ticking solver.
+
+		using namespace Chaos;
+		auto Sphere = TSharedPtr<FImplicitObject, ESPMode::ThreadSafe>(new TSphere<FReal, 3>(FVec3(0), 10));
+		FChaosSolversModule* Module = FChaosSolversModule::GetModule();
+
+		// Make a solver
+		auto* Solver = Module->CreateSolver(nullptr, /*AsyncDt=*/-1);
+
+		// Make 2 particles for our constraint
+		auto Proxy1 = FSingleParticlePhysicsProxy::Create(Chaos::FPBDRigidParticle::CreateParticle());
+		auto& Particle1 = Proxy1->GetGameThreadAPI();
+		Particle1.SetGeometry(Sphere);
+		Solver->RegisterObject(Proxy1);
+
+		auto Proxy2 = FSingleParticlePhysicsProxy::Create(Chaos::FPBDRigidParticle::CreateParticle());
+		auto& Particle2 = Proxy2->GetGameThreadAPI();
+		Particle2.SetGeometry(Sphere);
+		Solver->RegisterObject(Proxy2);
+
+		// Confirm solver doesn't have any joint constraint proxies yet.
+		EXPECT_EQ(Solver->GetJointConstraintPhysicsProxies_Internal().Num(), 0);
+
+		FTransform IdentityTransform(FTransform::Identity);
+		FPhysicsConstraintHandle Constraint = FChaosEngineInterface::CreateConstraint(Proxy1, Proxy2, IdentityTransform, IdentityTransform);
+
+		// joint Proxy has been created, should be in dirty proxy list to be sent to physics thread.
+		// Note that we expect three dirty proxies, 2 for particles, 1 for joint.
+		EXPECT_EQ(Solver->NumDirtyProxies_External(), 3);
+
+		// Although proxy exists, it should not have been added to physics thread only proxy array yet, as doing so from game thread code is wrong.
+		// Confirm game thread did not add it.
+		EXPECT_EQ(Solver->GetJointConstraintPhysicsProxies_Internal().Num(), 0);
+
+		// Tick physics thread, this should add proxy to solver proxy array.
+		Solver->AdvanceAndDispatch_External(/*Dt=*/1.0f);
+		EXPECT_EQ(Solver->GetJointConstraintPhysicsProxies_Internal().Num(), 1);
+
+		// Release constraint to test removal
+		FChaosEngineInterface::ReleaseConstraint(Constraint);
+
+		// Have only released constraint on game thread, proxy should still be in array.
+		EXPECT_EQ(Solver->GetJointConstraintPhysicsProxies_Internal().Num(), 1);
+
+		// Tick physics thread, this should remove proxy from array.
+		Solver->AdvanceAndDispatch_External(/*Dt=*/1.0f);
+		EXPECT_EQ(Solver->GetJointConstraintPhysicsProxies_Internal().Num(), 0);
+
+		Solver->UnregisterObject(Proxy1);
+		Solver->UnregisterObject(Proxy2);
+		Module->DestroySolver(Solver);
+	}
+
 	GTEST_TEST(AllTraits, SingleParticleProxyTests)
 	{
 		ChaosTest::SingleParticleProxySingleThreadTest();
 		ChaosTest::SingleParticleProxyWakeEventPropagationTest();
+		ChaosTest::JointProxySolverInitTest();
 	}
 }
diff --git a/Engine/Source/Runtime/Experimental/Chaos/Public/Chaos/Framework/PhysicsSolverBase.h b/Engine/Source/Runtime/Experimental/Chaos/Public/Chaos/Framework/PhysicsSolverBase.h
--- a/Engine/Source/Runtime/Experimental/Chaos/Public/Chaos/Framework/PhysicsSolverBase.h
+++ b/Engine/Source/Runtime/Experimental/Chaos/Public/Chaos/Framework/PhysicsSolverBase.h
@@ -126,6 +126,10 @@ namespace Chaos
 		{
 			MarshallingManager.GetProducerData_External()->DirtyProxiesDataBuffer.Remove(ProxyBaseIn);
 		}
+		int32 NumDirtyProxies_External()
+		{
+			return MarshallingManager.GetProducerData_External()->DirtyProxiesDataBuffer.NumDirtyProxies();
+		}
 
 		// Batch dirty proxies without checking DirtyIdx.
 		template <typename TProxiesArray>
diff --git a/Engine/Source/Runtime/Experimental/Chaos/Public/PBDRigidsSolver.h b/Engine/Source/Runtime/Experimental/Chaos/Public/PBDRigidsSolver.h
--- a/Engine/Source/Runtime/Experimental/Chaos/Public/PBDRigidsSolver.h
+++ b/Engine/Source/Runtime/Experimental/Chaos/Public/PBDRigidsSolver.h
@@ -239,6 +239,11 @@ namespace Chaos
 			return GeometryCollectionPhysicsProxies_Internal;
 		}
 
+		const TArray<FJointConstraintPhysicsProxy*>& GetJointConstraintPhysicsProxies_Internal() const
+		{
+			return JointConstraintPhysicsProxies_Internal;
+		}
+
 		/** Events hooked up to the Chaos material manager */
 		void UpdateMaterial(Chaos::FMaterialHandle InHandle, const Chaos::FChaosPhysicsMaterial& InNewData);
 		void CreateMaterial(Chaos::FMaterialHandle InHandle, const Chaos::FChaosPhysicsMaterial& InNewData);
