Added unit tests for changing dynamic to kinematic and vice versa

diff --git a/Engine/Source/Programs/HeadlessChaos/Private/HeadlessChaosTestRewind.cpp b/Engine/Source/Programs/HeadlessChaos/Private/HeadlessChaosTestRewind.cpp
--- a/Engine/Source/Programs/HeadlessChaos/Private/HeadlessChaosTestRewind.cpp
+++ b/Engine/Source/Programs/HeadlessChaos/Private/HeadlessChaosTestRewind.cpp
@@ -1298,6 +1298,84 @@ namespace ChaosTest {
 			});
 	}
 
+	GTEST_TEST(AllTraits, RewindTest_RewindBeforeMadeKinematic)
+	{
+		// Rewind before a dynamic was made kinematic and check the view is properly handling the rewind
+		TRewindHelper::TestDynamicSphere([](auto* Solver, FReal SimDt, int32 Optimization, auto Proxy, auto Sphere)
+			{
+				int32 ExpectedSleepFrame = INDEX_NONE;
+				const int32 ResimStartFrame = 3;
+				FRewindCallbackTestHelper* Helper = RegisterCallbackHelper(Solver, 13, ResimStartFrame);
+				Helper->ProcessInputsFunc = [&ExpectedSleepFrame, ResimStartFrame, Proxy, RewindData = Solver->GetRewindData(), Solver](const int32 PhysicsStep, bool bIsResimming)
+				{
+					if (PhysicsStep == 12)
+					{
+						Proxy->GetPhysicsThreadAPI()->SetObjectState(EObjectStateType::Kinematic);
+					}
+
+					if(PhysicsStep < 12)
+					{
+						EXPECT_EQ(Solver->GetParticles().GetNonDisabledDynamicView().Num(), 1);
+						EXPECT_EQ(Solver->GetParticles().GetActiveKinematicParticlesView().Num(), 0);
+					}
+					else
+					{
+						EXPECT_EQ(Solver->GetParticles().GetNonDisabledDynamicView().Num(), 0);
+						EXPECT_EQ(Solver->GetParticles().GetActiveKinematicParticlesView().Num(), 1);
+					}
+				};
+
+				auto& Particle = Proxy->GetGameThreadAPI();
+				Particle.SetGravityEnabled(false);
+				Particle.SetV(FVec3(0, 0, 1));
+
+				for (int Step = 0; Step <= 32; ++Step)
+				{
+					TickSolverHelper(Solver);
+				}
+
+			});
+	}
+
+	GTEST_TEST(AllTraits, RewindTest_RewindBeforeMadeDynamic)
+	{
+		// Rewind before a kinematic was made dynamic and check the view is properly handling the rewind
+		TRewindHelper::TestDynamicSphere([](auto* Solver, FReal SimDt, int32 Optimization, auto Proxy, auto Sphere)
+			{
+				int32 ExpectedSleepFrame = INDEX_NONE;
+				const int32 ResimStartFrame = 3;
+				FRewindCallbackTestHelper* Helper = RegisterCallbackHelper(Solver, 13, ResimStartFrame);
+				Helper->ProcessInputsFunc = [&ExpectedSleepFrame, ResimStartFrame, Proxy, RewindData = Solver->GetRewindData(), Solver](const int32 PhysicsStep, bool bIsResimming)
+				{
+					if (PhysicsStep == 12)
+					{
+						Proxy->GetPhysicsThreadAPI()->SetObjectState(EObjectStateType::Dynamic);
+					}
+
+					if (PhysicsStep < 12)
+					{
+						EXPECT_EQ(Solver->GetParticles().GetNonDisabledDynamicView().Num(), 0);
+						EXPECT_EQ(Solver->GetParticles().GetActiveKinematicParticlesView().Num(), 1);
+					}
+					else
+					{
+						EXPECT_EQ(Solver->GetParticles().GetNonDisabledDynamicView().Num(), 1);
+						EXPECT_EQ(Solver->GetParticles().GetActiveKinematicParticlesView().Num(), 0);
+					}
+				};
+
+				auto& Particle = Proxy->GetGameThreadAPI();
+				Particle.SetGravityEnabled(false);
+				Proxy->GetGameThreadAPI().SetObjectState(EObjectStateType::Kinematic);
+
+				for (int Step = 0; Step <= 32; ++Step)
+				{
+					TickSolverHelper(Solver);
+				}
+
+			});
+	}
+
 	GTEST_TEST(AllTraits, RewindTest_RecordForcesInSimCallback)
 	{
 		//Makes sure that we record the forces applied during sim callback
