headlesschaos
-- ExpectedNumCollisionParticles comparison fixed

diff --git a/Engine/Source/Programs/HeadlessChaos/Private/GeometryCollection/GeometryCollectionTestClustering.cpp b/Engine/Source/Programs/HeadlessChaos/Private/GeometryCollection/GeometryCollectionTestClustering.cpp
--- a/Engine/Source/Programs/HeadlessChaos/Private/GeometryCollection/GeometryCollectionTestClustering.cpp
+++ b/Engine/Source/Programs/HeadlessChaos/Private/GeometryCollection/GeometryCollectionTestClustering.cpp
@@ -2189,37 +2189,34 @@ namespace GeometryCollectionTest
 		Params.Simulating = true;
 		Params.EnableClustering = true;				
 		Params.CollisionGroup = -1;
+		Params.CollisionParticleFraction = 0.70f;
 		Params.MinLevelSetResolution = 15;
 		Params.MaxLevelSetResolution = 20;
 		FGeometryCollectionWrapper* Collection = TNewSimulationObject<GeometryType::GeometryCollectionWithSuppliedRestCollection>::Init(Params)->template As<FGeometryCollectionWrapper>();
 
 		UnitTest.AddSimulationObject(Collection);
+		UnitTest.Initialize();
 		
-		
-		// Todo: these aren't used anywhere in the test?
-		//typedef TUniquePtr<Chaos::FImplicitObject> FImplicitPointer;
-		//const TManagedArray<FImplicitPointer> & Implicits = RestCollection->template GetAttribute<FImplicitPointer>(FGeometryCollectionPhysicsProxy::ImplicitsAttribute, FTransformCollection::TransformGroup);
-
+		// We can get the size of an original CollisionParticles by storing the Simplicials before any internal resizing. 
 		typedef TUniquePtr< FCollisionStructureManager::FSimplicial > FSimplicialPointer;
 		const TManagedArray<FSimplicialPointer> & Simplicials = RestCollection->template GetAttribute<FSimplicialPointer>(FGeometryDynamicCollection::SimplicialsAttribute, FTransformCollection::TransformGroup);
 
-		UnitTest.Advance();		
+		UnitTest.Advance();		// this call triggers the array resize based on the fraction
 
+		// Test non-clustered bodies
 		TArray<Chaos::TPBDRigidClusteredParticleHandle<FReal, 3>*>& ParticleHandles = Collection->PhysObject->GetSolverParticleHandles();
-
-		FReal CollisionParticlesPerObjectFractionDefault = (FReal)0.5;
-		IConsoleVariable*  CVarCollisionParticlesPerObjectFractionDefault = IConsoleManager::Get().FindConsoleVariable(TEXT("p.CollisionParticlesPerObjectFractionDefault"));
-		EXPECT_NE(CVarCollisionParticlesPerObjectFractionDefault, nullptr);
-		if (CVarCollisionParticlesPerObjectFractionDefault != nullptr)
+		int32 NumCollisionParticles, ExpectedNumCollisionParticles;
+		for (int i = 0; i < ParticleHandles.Num(); i++)
 		{
-			CollisionParticlesPerObjectFractionDefault = (FReal)CVarCollisionParticlesPerObjectFractionDefault->GetFloat();
+			if (Collection->RestCollection->SimulationType[i] == FGeometryCollection::ESimulationTypes::FST_Rigid)
+			{
+				NumCollisionParticles = ParticleHandles[i]->CollisionParticlesSize(); // resized
+				ExpectedNumCollisionParticles = (int32)Simplicials[i]->Size() * Params.CollisionParticleFraction;
+				EXPECT_EQ(ExpectedNumCollisionParticles, NumCollisionParticles);
+				EXPECT_FALSE(NumCollisionParticles == 0.0f);
+				EXPECT_FALSE(Params.CollisionParticleFraction == 1.0f); // not defaulted
+			}
 		}
-		
-/*
-		todo: what is the replacement here?
-		EXPECT_EQ(Particles.CollisionParticles(Object->PhysicsProxy->RigidBodyIDArray_TestingAccess()[10])->Size(), (int)(Simplicials[10]->Size() * CollisionParticlesPerObjectFractionDefault));
-		EXPECT_EQ(Particles.CollisionParticles(Object->PhysicsProxy->RigidBodyIDArray_TestingAccess()[11])->Size(), (int)(Simplicials[11]->Size() * CollisionParticlesPerObjectFractionDefault));
-		EXPECT_EQ(Particles.CollisionParticles(Object->PhysicsProxy->RigidBodyIDArray_TestingAccess()[12])->Size(), (int)(Simplicials[12]->Size() * CollisionParticlesPerObjectFractionDefault));
-*/
 	}
+
 }
diff --git a/Engine/Source/Programs/HeadlessChaos/Private/GeometryCollection/GeometryCollectionTestFramework.cpp b/Engine/Source/Programs/HeadlessChaos/Private/GeometryCollection/GeometryCollectionTestFramework.cpp
--- a/Engine/Source/Programs/HeadlessChaos/Private/GeometryCollection/GeometryCollectionTestFramework.cpp
+++ b/Engine/Source/Programs/HeadlessChaos/Private/GeometryCollection/GeometryCollectionTestFramework.cpp
@@ -158,6 +158,11 @@ namespace GeometryCollectionTest
 			SimulationParams.Shared.bMassAsDensity = Params.bMassAsDensity;
 			SimulationParams.Shared.SizeSpecificData[0].CollisionShapesData[0].CollisionType = Params.CollisionType;
 			SimulationParams.Shared.SizeSpecificData[0].CollisionShapesData[0].ImplicitType = Params.ImplicitType;
+			SimulationParams.Shared.SizeSpecificData[0].CollisionShapesData[0].LevelSetData.MinLevelSetResolution = Params.MinLevelSetResolution;
+			SimulationParams.Shared.SizeSpecificData[0].CollisionShapesData[0].LevelSetData.MaxLevelSetResolution = Params.MaxLevelSetResolution;
+			SimulationParams.Shared.SizeSpecificData[0].CollisionShapesData[0].ImplicitType = Params.ImplicitType;
+			SimulationParams.Shared.SizeSpecificData[0].CollisionShapesData[0].CollisionParticleData.CollisionParticlesFraction = Params.CollisionParticleFraction;
+			SimulationParams.CollisionSampleFraction = Params.CollisionParticleFraction; // this should not be here, must remove/replace SimulationParams.CollisionSampleFraction in other files.
 			SimulationParams.Simulating = Params.Simulating;
 			SimulationParams.EnableClustering = Params.EnableClustering;
 			SimulationParams.InitialLinearVelocity = Params.InitialLinearVelocity;
@@ -169,11 +174,6 @@ namespace GeometryCollectionTest
 			SimulationParams.CollisionGroup = Params.CollisionGroup;
 			SimulationParams.ClusterGroupIndex = Params.ClusterGroupIndex;
 
-			FSharedSimulationSizeSpecificData Tmp;
-			Tmp.CollisionShapesData[0].LevelSetData.MinLevelSetResolution = Params.MinLevelSetResolution;
-			Tmp.CollisionShapesData[0].LevelSetData.MaxLevelSetResolution = Params.MaxLevelSetResolution;
-			SimulationParams.Shared.SizeSpecificData.Add(Tmp);
-
 			Chaos::FErrorReporter ErrorReporter;
 			BuildSimulationData(ErrorReporter, *RestCollection.Get(), SimulationParams.Shared);
 
diff --git a/Engine/Source/Programs/HeadlessChaos/Public/GeometryCollection/GeometryCollectionTestFramework.h b/Engine/Source/Programs/HeadlessChaos/Public/GeometryCollection/GeometryCollectionTestFramework.h
--- a/Engine/Source/Programs/HeadlessChaos/Public/GeometryCollection/GeometryCollectionTestFramework.h
+++ b/Engine/Source/Programs/HeadlessChaos/Public/GeometryCollection/GeometryCollectionTestFramework.h
@@ -84,6 +84,7 @@ namespace GeometryCollectionTest
 		ECollisionTypeEnum CollisionType = ECollisionTypeEnum::Chaos_Surface_Volumetric;
 		ESimplicialType SimplicialType = ESimplicialType::Chaos_Simplicial_Sphere;
 		EImplicitTypeEnum ImplicitType = EImplicitTypeEnum::Chaos_Implicit_Sphere;
+		float CollisionParticleFraction = 1.0f;
 		EInitialVelocityTypeEnum InitialVelocityType = EInitialVelocityTypeEnum::Chaos_Initial_Velocity_None;
 		TArray<FTransform> NestedTransforms;
 		bool EnableClustering = true;
