trivial fixes in headless / Clustering API fixes

diff --git a/Engine/Source/Programs/HeadlessChaos/Private/GeometryCollection/GeometryCollectionTestClustering.cpp b/Engine/Source/Programs/HeadlessChaos/Private/GeometryCollection/GeometryCollectionTestClustering.cpp
--- a/Engine/Source/Programs/HeadlessChaos/Private/GeometryCollection/GeometryCollectionTestClustering.cpp
+++ b/Engine/Source/Programs/HeadlessChaos/Private/GeometryCollection/GeometryCollectionTestClustering.cpp
@@ -143,18 +143,14 @@ namespace GeometryCollectionTest
 		RestCollection->AppendGeometry(*GeometryCollection::MakeCubeElement(FTransform(FQuat::MakeFromEuler(FVector(0.f)), FVector(60.f)), FVector(1.0)));
 
 		// 4 mid-level cluster parents
-		RestCollection->AddElements(4, FGeometryCollection::TransformGroup);
-		// @todo(ClusteringUtils) This is a bad assumption, the state flags should be initialized to zero.
-		(RestCollection->SimulationType)[5] = FGeometryCollection::ESimulationTypes::FST_Clustered;
-		(RestCollection->SimulationType)[6] = FGeometryCollection::ESimulationTypes::FST_Clustered;
-		(RestCollection->SimulationType)[7] = FGeometryCollection::ESimulationTypes::FST_Clustered;
-		(RestCollection->SimulationType)[8] = FGeometryCollection::ESimulationTypes::FST_Clustered;
-
-		// Build a binary tree cluster parent hierarchy
-		GeometryCollectionAlgo::ParentTransforms(RestCollection.Get(), 5, { 4,3 }); // Transform index 5 is parent to 4 and 3
-		GeometryCollectionAlgo::ParentTransforms(RestCollection.Get(), 6, { 5,2 });
-		GeometryCollectionAlgo::ParentTransforms(RestCollection.Get(), 7, { 6,1 });
-		GeometryCollectionAlgo::ParentTransforms(RestCollection.Get(), 8, { 7,0 });
+		FGeometryCollectionClusteringUtility::ClusterBonesUnderNewNode(RestCollection.Get(), 5, { 4,3 }, true, false); // just validate at end of construction
+		FGeometryCollectionClusteringUtility::ClusterBonesUnderNewNode(RestCollection.Get(), 6, { 5,2 }, true, false);
+		FGeometryCollectionClusteringUtility::ClusterBonesUnderNewNode(RestCollection.Get(), 7, { 6,1 }, true, false);
+		FGeometryCollectionClusteringUtility::ClusterAllBonesUnderNewRoot(RestCollection.Get());
+		FGeometryCollectionClusteringUtility::ValidateResults(RestCollection.Get());		
+
+		//GeometryCollectionAlgo::PrintParentHierarchy(RestCollection.Get());
+
 		
 		CreationParameters Params;
 		Params.RestCollection = RestCollection;
@@ -264,18 +260,12 @@ namespace GeometryCollectionTest
 		RestCollection->AppendGeometry(*GeometryCollection::MakeCubeElement(FTransform(FQuat::MakeFromEuler(FVector(0.f)), FVector(60.f)), FVector(1.0)));
 
 		// 4 mid-level cluster parents
-		RestCollection->AddElements(4, FGeometryCollection::TransformGroup);
-		// @todo(ClusteringUtils) This is a bad assumption, the state flags should be initialized to zero.
-		(RestCollection->SimulationType)[5] = FGeometryCollection::ESimulationTypes::FST_Clustered;
-		(RestCollection->SimulationType)[6] = FGeometryCollection::ESimulationTypes::FST_Clustered;
-		(RestCollection->SimulationType)[7] = FGeometryCollection::ESimulationTypes::FST_Clustered;
-		(RestCollection->SimulationType)[8] = FGeometryCollection::ESimulationTypes::FST_Clustered;
-
-		// Build a binary tree cluster parent hierarchy
-		GeometryCollectionAlgo::ParentTransforms(RestCollection.Get(), 5, { 4,3 }); // Transform index 5 is parent to 4 and 3
-		GeometryCollectionAlgo::ParentTransforms(RestCollection.Get(), 6, { 5,2 });
-		GeometryCollectionAlgo::ParentTransforms(RestCollection.Get(), 7, { 6,1 });
-		GeometryCollectionAlgo::ParentTransforms(RestCollection.Get(), 8, { 7,0 });
+		FGeometryCollectionClusteringUtility::ClusterBonesUnderNewNode(RestCollection.Get(), 5, { 4,3 }, true, false); // just validate at end of construction
+		FGeometryCollectionClusteringUtility::ClusterBonesUnderNewNode(RestCollection.Get(), 6, { 5,2 }, true, false);
+		FGeometryCollectionClusteringUtility::ClusterBonesUnderNewNode(RestCollection.Get(), 7, { 6,1 }, true, false);
+		FGeometryCollectionClusteringUtility::ClusterAllBonesUnderNewRoot(RestCollection.Get());
+		FGeometryCollectionClusteringUtility::ValidateResults(RestCollection.Get());
+
 
 		CreationParameters Params;
 		Params.RestCollection = RestCollection;
@@ -673,22 +663,11 @@ namespace GeometryCollectionTest
 		RestCollection->AppendGeometry(*GeometryCollection::MakeCubeElement(FTransform(FQuat::MakeFromEuler(FVector(0.f)), FVector(50.f)), FVector(1.0)));
 		RestCollection->AppendGeometry(*GeometryCollection::MakeCubeElement(FTransform(FQuat::MakeFromEuler(FVector(0.f)), FVector(60.f)), FVector(1.0)));
 
-		RestCollection->AddElements(4, FGeometryCollection::TransformGroup);
-		// @todo(ClusteringUtils) This is a bad assumption, the state flags should be initialized to zero.
-		(RestCollection->SimulationType)[5] = FGeometryCollection::ESimulationTypes::FST_Clustered;
-		(RestCollection->SimulationType)[6] = FGeometryCollection::ESimulationTypes::FST_Clustered;
-		(RestCollection->SimulationType)[7] = FGeometryCollection::ESimulationTypes::FST_Clustered;
-		(RestCollection->SimulationType)[8] = FGeometryCollection::ESimulationTypes::FST_Clustered;
-
-		GeometryCollectionAlgo::ParentTransforms(RestCollection.Get(), 5, { 4,3 });
-		GeometryCollectionAlgo::ParentTransforms(RestCollection.Get(), 6, { 5,2 });
-		GeometryCollectionAlgo::ParentTransforms(RestCollection.Get(), 7, { 6,1 });
-		GeometryCollectionAlgo::ParentTransforms(RestCollection.Get(), 8, { 7,0 });
-
-		// @todo(brice->Bill.Henderson) Why did this not work? I needed to build my own parenting and level initilization. 
-		//FGeometryCollectionClusteringUtility::ClusterAllBonesUnderNewRoot(RestCollection.Get());
-		//FGeometryCollectionClusteringUtility::ClusterBonesUnderNewNode(RestCollection.Get(), 4, { 0, 1 }, true);
-		//FGeometryCollectionClusteringUtility::ClusterBonesUnderNewNode(RestCollection.Get(), 4, { 2, 3 }, true);
+		FGeometryCollectionClusteringUtility::ClusterAllBonesUnderNewRoot(RestCollection.Get()); 
+		FGeometryCollectionClusteringUtility::ClusterBonesUnderNewNode(RestCollection.Get(), 4, { 2, 3 }, true, true);
+		FGeometryCollectionClusteringUtility::ClusterBonesUnderNewNode(RestCollection.Get(), 2, { 1, 0 }, true, true);
+		FGeometryCollectionClusteringUtility::ClusterBonesUnderNewNode(RestCollection.Get(), 1, { 3, 0 }, true, true);
+		
 
 		CreationParameters Params;
 		Params.RestCollection = RestCollection;
@@ -698,6 +677,8 @@ namespace GeometryCollectionTest
 		Params.Simulating = true;
 		Params.EnableClustering = true;
 		Params.DamageThreshold = { 30.0, 30.0, 30.0, FLT_MAX };
+
+		// basically a stand-in for a 'component'
 		FGeometryCollectionWrapper* Collection = TNewSimulationObject<GeometryType::GeometryCollectionWithSuppliedRestCollection>::Init(Params)->template As<FGeometryCollectionWrapper>();
 
 		UnitTest.AddSimulationObject(Collection);
@@ -715,10 +696,10 @@ namespace GeometryCollectionTest
 		auto& Clustering = UnitTest.Solver->GetEvolution()->GetRigidClustering();
 		const auto& ClusterMap = Clustering.GetChildrenMap();
 
-		EXPECT_TRUE(ClusterMapContains(ClusterMap, ParticleHandles[5], { ParticleHandles[4],ParticleHandles[3] }));
-		EXPECT_TRUE(ClusterMapContains(ClusterMap, ParticleHandles[6], { ParticleHandles[5],ParticleHandles[2] }));
-		EXPECT_TRUE(ClusterMapContains(ClusterMap, ParticleHandles[7], { ParticleHandles[6],ParticleHandles[1] }));
-		EXPECT_TRUE(ClusterMapContains(ClusterMap, ParticleHandles[8], { ParticleHandles[7],ParticleHandles[0] }));
+		EXPECT_TRUE(ClusterMapContains(ClusterMap, ParticleHandles[5], { ParticleHandles[6],ParticleHandles[4] }));
+		EXPECT_TRUE(ClusterMapContains(ClusterMap, ParticleHandles[6], { ParticleHandles[7],ParticleHandles[2] }));
+		EXPECT_TRUE(ClusterMapContains(ClusterMap, ParticleHandles[7], { ParticleHandles[8],ParticleHandles[1] }));
+		EXPECT_TRUE(ClusterMapContains(ClusterMap, ParticleHandles[8], { ParticleHandles[3],ParticleHandles[0] }));
 
 
 		for (int Frame = 1; Frame < 40; Frame++)
@@ -734,88 +715,89 @@ namespace GeometryCollectionTest
 					ParticleHandles[2]->Disabled() == true &&
 					ParticleHandles[3]->Disabled() == true &&
 					ParticleHandles[4]->Disabled() == true &&
-					ParticleHandles[5]->Disabled() == true &&
+					ParticleHandles[5]->Disabled() == false && // root
 					ParticleHandles[6]->Disabled() == true &&
 					ParticleHandles[7]->Disabled() == true &&
-					ParticleHandles[8]->Disabled() == false)
+					ParticleHandles[8]->Disabled() == true)
 				{
 					Conditions[0] = true;
 				}
 			}
+			// Root cluster broken, check for activated children
 			else if (Conditions[0] == true && Conditions[1] == false)
 			{
 				if (
-					ParticleHandles[0]->Disabled() == false &&
+					ParticleHandles[0]->Disabled() == true &&
 					ParticleHandles[1]->Disabled() == true &&
 					ParticleHandles[2]->Disabled() == true &&
 					ParticleHandles[3]->Disabled() == true &&
-					ParticleHandles[4]->Disabled() == true &&
-					ParticleHandles[5]->Disabled() == true &&
-					ParticleHandles[6]->Disabled() == true &&
-					ParticleHandles[7]->Disabled() == false &&
+					ParticleHandles[4]->Disabled() == false &&
+					ParticleHandles[5]->Disabled() == true && // root, now disabled
+					ParticleHandles[6]->Disabled() == false &&
+					ParticleHandles[7]->Disabled() == true &&
 					ParticleHandles[8]->Disabled() == true)
 				{
 					Conditions[1] = true;
 
 					EXPECT_EQ(ClusterMap.Num(), 3);
-					EXPECT_TRUE(ClusterMapContains(ClusterMap, ParticleHandles[5], { ParticleHandles[4],ParticleHandles[3] }));
-					EXPECT_TRUE(ClusterMapContains(ClusterMap, ParticleHandles[6], { ParticleHandles[5],ParticleHandles[2] }));
-					EXPECT_TRUE(ClusterMapContains(ClusterMap, ParticleHandles[7], { ParticleHandles[6],ParticleHandles[1] }));
+					EXPECT_TRUE(ClusterMapContains(ClusterMap, ParticleHandles[8], { ParticleHandles[3],ParticleHandles[0] }));
+					EXPECT_TRUE(ClusterMapContains(ClusterMap, ParticleHandles[7], { ParticleHandles[1],ParticleHandles[8] }));
+					EXPECT_TRUE(ClusterMapContains(ClusterMap, ParticleHandles[6], { ParticleHandles[7],ParticleHandles[2] }));
 				}
 			}
 			else if (Conditions[1] == true && Conditions[2] == false)
 			{
 				if (
-					ParticleHandles[0]->Disabled() == false &&
-					ParticleHandles[1]->Disabled() == false &&
-					ParticleHandles[2]->Disabled() == true &&
+					ParticleHandles[0]->Disabled() == true &&
+					ParticleHandles[1]->Disabled() == true &&
+					ParticleHandles[2]->Disabled() == false &&
 					ParticleHandles[3]->Disabled() == true &&
-					ParticleHandles[4]->Disabled() == true &&
+					ParticleHandles[4]->Disabled() == false &&
 					ParticleHandles[5]->Disabled() == true &&
-					ParticleHandles[6]->Disabled() == false &&
-					ParticleHandles[7]->Disabled() == true &&
+					ParticleHandles[6]->Disabled() == true &&
+					ParticleHandles[7]->Disabled() == false &&
 					ParticleHandles[8]->Disabled() == true)
 				{
 					Conditions[2] = true;
 
 					EXPECT_EQ(ClusterMap.Num(), 2);
-					EXPECT_TRUE(ClusterMapContains(ClusterMap, ParticleHandles[5], { ParticleHandles[4],ParticleHandles[3] }));
-					EXPECT_TRUE(ClusterMapContains(ClusterMap, ParticleHandles[6], { ParticleHandles[5],ParticleHandles[2] }));
+					EXPECT_TRUE(ClusterMapContains(ClusterMap, ParticleHandles[8], { ParticleHandles[3],ParticleHandles[0] }));
+					EXPECT_TRUE(ClusterMapContains(ClusterMap, ParticleHandles[7], { ParticleHandles[1],ParticleHandles[8] }));
 				}
 			}
 			else if (Conditions[2] == true && Conditions[3] == false)
 			{
 				if (
-					ParticleHandles[0]->Disabled() == false &&
+					ParticleHandles[0]->Disabled() == true &&
 					ParticleHandles[1]->Disabled() == false &&
 					ParticleHandles[2]->Disabled() == false &&
 					ParticleHandles[3]->Disabled() == true &&
-					ParticleHandles[4]->Disabled() == true &&
-					ParticleHandles[5]->Disabled() == false &&
+					ParticleHandles[4]->Disabled() == false &&
+					ParticleHandles[5]->Disabled() == true &&
 					ParticleHandles[6]->Disabled() == true &&
 					ParticleHandles[7]->Disabled() == true &&
-					ParticleHandles[8]->Disabled() == true)
+					ParticleHandles[8]->Disabled() == false)
 				{
 					Conditions[3] = true;
 
 					EXPECT_EQ(ClusterMap.Num(), 1);
-					EXPECT_TRUE(ClusterMapContains(ClusterMap, ParticleHandles[5], { ParticleHandles[4],ParticleHandles[3] }));
+					EXPECT_TRUE(ClusterMapContains(ClusterMap, ParticleHandles[8], { ParticleHandles[3],ParticleHandles[0] }));
 				}
 			}
 			else if (Conditions[3] == true)
 			{
 				// fLT_MAX strain so last cluster should never break. 
-				EXPECT_TRUE(ParticleHandles[0]->Disabled() == false);
+				EXPECT_TRUE(ParticleHandles[0]->Disabled() == true);
 				EXPECT_TRUE(ParticleHandles[1]->Disabled() == false);
 				EXPECT_TRUE(ParticleHandles[2]->Disabled() == false);
 				EXPECT_TRUE(ParticleHandles[3]->Disabled() == true);
-				EXPECT_TRUE(ParticleHandles[4]->Disabled() == true);
-				EXPECT_TRUE(ParticleHandles[5]->Disabled() == false);
+				EXPECT_TRUE(ParticleHandles[4]->Disabled() == false);
+				EXPECT_TRUE(ParticleHandles[5]->Disabled() == true);
 				EXPECT_TRUE(ParticleHandles[6]->Disabled() == true);
 				EXPECT_TRUE(ParticleHandles[7]->Disabled() == true);
-				EXPECT_TRUE(ParticleHandles[8]->Disabled() == true);
+				EXPECT_TRUE(ParticleHandles[8]->Disabled() == false);
 				EXPECT_EQ(ClusterMap.Num(), 1);
-				EXPECT_TRUE(ClusterMapContains(ClusterMap, ParticleHandles[5], { ParticleHandles[4],ParticleHandles[3] }));
+				EXPECT_TRUE(ClusterMapContains(ClusterMap, ParticleHandles[8], { ParticleHandles[3],ParticleHandles[0] }));
 			}
 		}
 		for (int i = 0; i < Conditions.Num(); i++)
@@ -839,17 +821,13 @@ namespace GeometryCollectionTest
 		RestCollection->AppendGeometry(*GeometryCollection::MakeCubeElement(FTransform(FQuat::MakeFromEuler(FVector(0.f)), FVector(50.f)), FVector(1.0)));
 		RestCollection->AppendGeometry(*GeometryCollection::MakeCubeElement(FTransform(FQuat::MakeFromEuler(FVector(0.f)), FVector(60.f)), FVector(1.0)));
 
-		RestCollection->AddElements(4, FGeometryCollection::TransformGroup);
-		// @todo(ClusteringUtils) This is a bad assumption, the state flags should be initialized to zero.
-		(RestCollection->SimulationType)[5] = FGeometryCollection::ESimulationTypes::FST_Clustered;
-		(RestCollection->SimulationType)[6] = FGeometryCollection::ESimulationTypes::FST_Clustered;
-		(RestCollection->SimulationType)[7] = FGeometryCollection::ESimulationTypes::FST_Clustered;
-		(RestCollection->SimulationType)[8] = FGeometryCollection::ESimulationTypes::FST_Clustered;
+		FGeometryCollectionClusteringUtility::ClusterBonesUnderNewNode(RestCollection.Get(), 5, { 4,3 }, true, false); // just validate at end of construction
+		FGeometryCollectionClusteringUtility::ClusterBonesUnderNewNode(RestCollection.Get(), 6, { 5,2 }, true, false);
+		FGeometryCollectionClusteringUtility::ClusterBonesUnderNewNode(RestCollection.Get(), 7, { 6,1 }, true, false);
+		FGeometryCollectionClusteringUtility::ClusterAllBonesUnderNewRoot(RestCollection.Get());
+		FGeometryCollectionClusteringUtility::ValidateResults(RestCollection.Get());
 
-		GeometryCollectionAlgo::ParentTransforms(RestCollection.Get(), 5, { 4,3 });
-		GeometryCollectionAlgo::ParentTransforms(RestCollection.Get(), 6, { 5,2 });
-		GeometryCollectionAlgo::ParentTransforms(RestCollection.Get(), 7, { 6,1 });
-		GeometryCollectionAlgo::ParentTransforms(RestCollection.Get(), 8, { 7,0 });
+		//GeometryCollectionAlgo::PrintParentHierarchy(RestCollection.Get());
 
 		CreationParameters Params;
 		Params.RestCollection = RestCollection;
@@ -988,11 +966,11 @@ namespace GeometryCollectionTest
 					FVector X66 = Ref6;
 
 
-					//check(!X0.ContainsNaN());
-					//check(!Ref0.ContainsNaN());
-					//check(FMath::IsFinite(X0.Size()));
-					//check(FMath::IsFinite(Ref0.Size()));
-					//check(FMath::IsFinite(X0.Size() - Ref0.Size()));
+					check(!X0.ContainsNaN());
+					check(!Ref0.ContainsNaN());
+					check(FMath::IsFinite(X0.Size()));
+					check(FMath::IsFinite(Ref0.Size()));
+					check(FMath::IsFinite(X0.Size() - Ref0.Size()));
 					check(FMath::IsFinite(FMath::Abs(X0.Size() - X00.Size())));
 					EXPECT_NEAR(FMath::Abs(X0.Size() - Ref0.Size()), 0, KINDA_SMALL_NUMBER);// << *FString("Kinematic body1 moved");
 					check(FMath::IsFinite(FMath::Abs(X1.Size() - X11.Size())));
@@ -1143,17 +1121,13 @@ namespace GeometryCollectionTest
 		RestCollection->AppendGeometry(*GeometryCollection::MakeCubeElement(FTransform(FQuat::MakeFromEuler(FVector(0.f)), FVector(50.f)), FVector(1.0)));
 		RestCollection->AppendGeometry(*GeometryCollection::MakeCubeElement(FTransform(FQuat::MakeFromEuler(FVector(0.f)), FVector(60.f)), FVector(1.0)));
 
-		RestCollection->AddElements(4, FGeometryCollection::TransformGroup);
-		// @todo(ClusteringUtils) This is a bad assumption, the state flags should be initialized to zero.
-		(RestCollection->SimulationType)[5] = FGeometryCollection::ESimulationTypes::FST_Clustered;
-		(RestCollection->SimulationType)[6] = FGeometryCollection::ESimulationTypes::FST_Clustered;
-		(RestCollection->SimulationType)[7] = FGeometryCollection::ESimulationTypes::FST_Clustered;
-		(RestCollection->SimulationType)[8] = FGeometryCollection::ESimulationTypes::FST_Clustered;
+		FGeometryCollectionClusteringUtility::ClusterBonesUnderNewNode(RestCollection.Get(), 5, { 4,3 }, true, false); // just validate at end of construction
+		FGeometryCollectionClusteringUtility::ClusterBonesUnderNewNode(RestCollection.Get(), 6, { 5,2 }, true, false);
+		FGeometryCollectionClusteringUtility::ClusterBonesUnderNewNode(RestCollection.Get(), 7, { 6,1 }, true, false);
+		FGeometryCollectionClusteringUtility::ClusterAllBonesUnderNewRoot(RestCollection.Get());
+		FGeometryCollectionClusteringUtility::ValidateResults(RestCollection.Get());
 
-		GeometryCollectionAlgo::ParentTransforms(RestCollection.Get(), 5, { 4,3 });
-		GeometryCollectionAlgo::ParentTransforms(RestCollection.Get(), 6, { 5,2 });
-		GeometryCollectionAlgo::ParentTransforms(RestCollection.Get(), 7, { 6,1 });
-		GeometryCollectionAlgo::ParentTransforms(RestCollection.Get(), 8, { 7,0 });
+		//GeometryCollectionAlgo::PrintParentHierarchy(RestCollection.Get());
 
 		CreationParameters Params;
 		Params.RestCollection = RestCollection;
@@ -2217,6 +2191,7 @@ namespace GeometryCollectionTest
 				EXPECT_FALSE(Params.CollisionParticleFraction == 1.0f); // not defaulted
 			}
 		}
+
 	}
 
 }
diff --git a/Engine/Source/Programs/HeadlessChaos/Private/GeometryCollection/GeometryCollectionTestCreation.cpp b/Engine/Source/Programs/HeadlessChaos/Private/GeometryCollection/GeometryCollectionTestCreation.cpp
--- a/Engine/Source/Programs/HeadlessChaos/Private/GeometryCollection/GeometryCollectionTestCreation.cpp
+++ b/Engine/Source/Programs/HeadlessChaos/Private/GeometryCollection/GeometryCollectionTestCreation.cpp
@@ -255,10 +255,10 @@ namespace GeometryCollectionTest
 
 			EXPECT_EQ((Collection->FaceCount)[0], 12);
 			EXPECT_EQ((Collection->FaceCount)[1], 12);
-			EXPECT_EQ((Collection->FaceCount)[1], 12);
-			EXPECT_EQ((Collection->FaceCount)[1], 12);
-			EXPECT_EQ((Collection->FaceCount)[1], 12);
-			EXPECT_EQ((Collection->FaceCount)[1], 12);
+			EXPECT_EQ((Collection->FaceCount)[2], 12);
+			EXPECT_EQ((Collection->FaceCount)[3], 12);
+			EXPECT_EQ((Collection->FaceCount)[4], 12);
+			EXPECT_EQ((Collection->FaceCount)[5], 12);
 			EXPECT_EQ((Collection->Indices).Num(), 72);
 
 			EXPECT_EQ((Collection->VertexStart)[0], 0);
