remove GeometryCollectionSkeletalMeshPhysicsProxy from headless

diff --git a/Engine/Source/Programs/HeadlessChaos/Private/GeometryCollection/GeometryCollectionTestSkeletalMeshPhysicsProxy.cpp b/Engine/Source/Programs/HeadlessChaos/Private/GeometryCollection/GeometryCollectionTestSkeletalMeshPhysicsProxy.cpp
deleted file mode 100644
diff --git a/Engine/Source/Programs/HeadlessChaos/Private/HeadlessChaosApp.cpp b/Engine/Source/Programs/HeadlessChaos/Private/HeadlessChaosApp.cpp
--- a/Engine/Source/Programs/HeadlessChaos/Private/HeadlessChaosApp.cpp
+++ b/Engine/Source/Programs/HeadlessChaos/Private/HeadlessChaosApp.cpp
@@ -59,7 +59,6 @@
 #include "GeometryCollection/GeometryCollectionTestSpatialHash.h"
 #include "GeometryCollection/GeometryCollectionTestVisibility.h"
 #include "GeometryCollection/GeometryCollectionTestEvents.h"
-#include "GeometryCollection/GeometryCollectionTestSkeletalMeshPhysicsProxy.h"
 #include "GeometryCollection/GeometryCollectionTestSerialization.h"
 
 #include "CompGeom/ExactPredicates.h"
diff --git a/Engine/Source/Programs/HeadlessChaos/Public/GeometryCollection/GeometryCollectionTestSkeletalMeshPhysicsProxy.h b/Engine/Source/Programs/HeadlessChaos/Public/GeometryCollection/GeometryCollectionTestSkeletalMeshPhysicsProxy.h
deleted file mode 100644
