field algebra #3

diff --git a/Engine/Source/Programs/HeadlessChaos/Private/GeometryCollection/GeometryCollectionTestSimulationField.cpp b/Engine/Source/Programs/HeadlessChaos/Private/GeometryCollection/GeometryCollectionTestSimulationField.cpp
--- a/Engine/Source/Programs/HeadlessChaos/Private/GeometryCollection/GeometryCollectionTestSimulationField.cpp
+++ b/Engine/Source/Programs/HeadlessChaos/Private/GeometryCollection/GeometryCollectionTestSimulationField.cpp
@@ -41,7 +41,7 @@ namespace GeometryCollectionTest
 		Params.DynamicState = EObjectStateTypeEnum::Chaos_Object_Kinematic;
 		Params.RootTransform.SetLocation(Translation0);
 		FGeometryCollectionWrapper* Collection = TNewSimulationObject<GeometryType::GeometryCollectionWithSingleRigid>::Init(Params)->template As<FGeometryCollectionWrapper>();
-		Params.RootTransform.SetLocation(FVector(100,0,0));
+		Params.RootTransform.SetLocation(FVector(100, 0, 0));
 		FGeometryCollectionWrapper* CollectionOther = TNewSimulationObject<GeometryType::GeometryCollectionWithSingleRigid>::Init(Params)->template As<FGeometryCollectionWrapper>();
 
 		FFramework UnitTest;
@@ -61,10 +61,10 @@ namespace GeometryCollectionTest
 		UnitTest.Advance();
 
 		UnitTest.Solver->RegisterSimOneShotCallback([&]()
-		{
-			EXPECT_EQ(CollectionOther->DynamicCollection->DynamicState[0],(int32)EObjectStateTypeEnum::Chaos_Object_Kinematic);
-			EXPECT_EQ(Collection->DynamicCollection->DynamicState[0],(int32)EObjectStateTypeEnum::Chaos_Object_Dynamic);
-		});
+			{
+				EXPECT_EQ(CollectionOther->DynamicCollection->DynamicState[0], (int32)EObjectStateTypeEnum::Chaos_Object_Kinematic);
+				EXPECT_EQ(Collection->DynamicCollection->DynamicState[0], (int32)EObjectStateTypeEnum::Chaos_Object_Dynamic);
+			});
 	}
 
 	GTEST_TEST(AllTraits, GeometryCollection_RigidBodies_Field_KinematicActivationOnProxyDuringUpdate)
@@ -76,7 +76,7 @@ namespace GeometryCollectionTest
 		Params.RootTransform.SetLocation(Translation0);
 		FGeometryCollectionWrapper* Collection = TNewSimulationObject<GeometryType::GeometryCollectionWithSingleRigid>::Init(Params)->template As<FGeometryCollectionWrapper>();
 
-		FFramework UnitTest; UnitTest.Dt = 1/24.0;
+		FFramework UnitTest; UnitTest.Dt = 1 / 24.0;
 		UnitTest.AddSimulationObject(Collection);
 
 		UnitTest.Initialize();
@@ -92,7 +92,7 @@ namespace GeometryCollectionTest
 		const TManagedArray<FTransform>& Transform = Collection->DynamicCollection->Transform;
 		EXPECT_EQ(Transform.Num(), 1);
 		const FVector Translation1 = Transform[0].GetTranslation();
-		EXPECT_NEAR( (Translation0-Translation1).Size(), 0.f, KINDA_SMALL_NUMBER);
+		EXPECT_NEAR((Translation0 - Translation1).Size(), 0.f, KINDA_SMALL_NUMBER);
 		EXPECT_NEAR(Transform[0].GetTranslation().Z, 1.f, KINDA_SMALL_NUMBER);
 
 		FRadialIntMask* RadialMask = new FRadialIntMask();
@@ -118,7 +118,7 @@ namespace GeometryCollectionTest
 	{
 		const FVector Translation0(0, 0, 1);
 
-		CreationParameters Params; 
+		CreationParameters Params;
 		Params.DynamicState = EObjectStateTypeEnum::Chaos_Object_Kinematic;
 		Params.RootTransform.SetLocation(Translation0);
 		FGeometryCollectionWrapper* Collection = TNewSimulationObject<GeometryType::GeometryCollectionWithSingleRigid>::Init(Params)->template As<FGeometryCollectionWrapper>();
@@ -137,7 +137,7 @@ namespace GeometryCollectionTest
 		const TManagedArray<FTransform>& Transform = Collection->DynamicCollection->Transform;
 		EXPECT_EQ(Transform.Num(), 1);
 		const FVector Translation1 = Transform[0].GetTranslation();
-		EXPECT_NEAR( (Translation0-Translation1).Size(), 0.f, KINDA_SMALL_NUMBER);
+		EXPECT_NEAR((Translation0 - Translation1).Size(), 0.f, KINDA_SMALL_NUMBER);
 		EXPECT_NEAR(Transform[0].GetTranslation().Z, 1.f, KINDA_SMALL_NUMBER);
 
 		FRadialIntMask* RadialMask = new FRadialIntMask();
@@ -147,7 +147,7 @@ namespace GeometryCollectionTest
 		RadialMask->ExteriorValue = (int32)EObjectStateTypeEnum::Chaos_Object_Kinematic;
 		RadialMask->SetMaskCondition = ESetMaskConditionType::Field_Set_IFF_NOT_Interior;
 		FName TargetName = GetFieldPhysicsName(EFieldPhysicsType::Field_DynamicState);
-		UnitTest.Solver->GetPerSolverField().AddTransientCommand( { TargetName, RadialMask });
+		UnitTest.Solver->GetPerSolverField().AddTransientCommand({ TargetName, RadialMask });
 
 		for (int i = 0; i < 100; i++)
 		{
@@ -231,7 +231,7 @@ namespace GeometryCollectionTest
 		UnitTest.AddSimulationObject(Collection);
 
 		// Field setup
-		FRadialIntMask * RadialMask = new FRadialIntMask();
+		FRadialIntMask* RadialMask = new FRadialIntMask();
 		RadialMask->Position = FVector(0.0, 0.0, PreviousHeight);
 		RadialMask->Radius = 5.0;
 		RadialMask->InteriorValue = (int32)EObjectStateTypeEnum::Chaos_Object_Dynamic;
@@ -247,7 +247,7 @@ namespace GeometryCollectionTest
 			if (Frame == 5)
 			{
 				FName TargetName = GetFieldPhysicsName(EFieldPhysicsType::Field_DynamicState);
-				UnitTest.Solver->GetPerSolverField().AddTransientCommand( { TargetName, RadialMask });
+				UnitTest.Solver->GetPerSolverField().AddTransientCommand({ TargetName, RadialMask });
 			}
 
 			UnitTest.Advance();
@@ -281,7 +281,7 @@ namespace GeometryCollectionTest
 		UnitTest.AddSimulationObject(Collection);
 
 		// Field setup
-		FUniformVector * UniformVector = new FUniformVector();
+		FUniformVector* UniformVector = new FUniformVector();
 		UniformVector->Direction = FVector(0.0, 1.0, 0.0);
 		UniformVector->Magnitude = 1000.0;
 
@@ -294,7 +294,7 @@ namespace GeometryCollectionTest
 			if (Frame >= 5)
 			{
 				FName TargetName = GetFieldPhysicsName(EFieldPhysicsType::Field_LinearForce);
-				UnitTest.Solver->GetPerSolverField().AddTransientCommand( { TargetName, UniformVector->NewCopy() });
+				UnitTest.Solver->GetPerSolverField().AddTransientCommand({ TargetName, UniformVector->NewCopy() });
 			}
 
 			UnitTest.Advance();
@@ -329,7 +329,7 @@ namespace GeometryCollectionTest
 		UnitTest.AddSimulationObject(Collection);
 
 		// Field setup
-		FUniformVector * UniformVector = new FUniformVector();
+		FUniformVector* UniformVector = new FUniformVector();
 		UniformVector->Direction = FVector(0.0, 1.0, 0.0);
 		UniformVector->Magnitude = 100.0;
 
@@ -342,7 +342,7 @@ namespace GeometryCollectionTest
 			if (Frame >= 5)
 			{
 				FName TargetName = GetFieldPhysicsName(EFieldPhysicsType::Field_AngularTorque);
-				UnitTest.Solver->GetPerSolverField().AddTransientCommand( { TargetName, UniformVector->NewCopy() });
+				UnitTest.Solver->GetPerSolverField().AddTransientCommand({ TargetName, UniformVector->NewCopy() });
 			}
 
 			UnitTest.Advance();
@@ -375,7 +375,7 @@ namespace GeometryCollectionTest
 		UnitTest.AddSimulationObject(Collection);
 
 		// Field setup
-		FPlaneFalloff * FalloffField = new FPlaneFalloff();
+		FPlaneFalloff* FalloffField = new FPlaneFalloff();
 		FalloffField->Magnitude = 1.0;
 		FalloffField->Distance = 10.0f;
 		FalloffField->Position = FVector(0.0, 0.0, 5.0);
@@ -387,14 +387,14 @@ namespace GeometryCollectionTest
 		TManagedArray<FTransform>& Transform = Collection->DynamicCollection->Transform;
 		TManagedArray<bool>& Active = Collection->DynamicCollection->Active;
 		auto& Particles = UnitTest.Solver->GetParticles().GetGeometryCollectionParticles();
-		for(int Frame = 0; Frame < 20; Frame++)
+		for (int Frame = 0; Frame < 20; Frame++)
 		{
 			FName TargetName = GetFieldPhysicsName(EFieldPhysicsType::Field_Kill);
-			UnitTest.Solver->GetPerSolverField().AddTransientCommand( { TargetName, FalloffField->NewCopy() });
+			UnitTest.Solver->GetPerSolverField().AddTransientCommand({ TargetName, FalloffField->NewCopy() });
 
 			UnitTest.Advance();
 
-			if(Particles.Disabled(0))
+			if (Particles.Disabled(0))
 			{
 				break;
 			}
@@ -419,21 +419,21 @@ namespace GeometryCollectionTest
 		UnitTest.AddSimulationObject(Collection);
 
 		// Field setup
-		FUniformVector * VectorField = new FUniformVector();
+		FUniformVector* VectorField = new FUniformVector();
 		VectorField->Magnitude = 100.0;
 		VectorField->Direction = FVector(1.0, 0.0, 0.0);
 
 		UnitTest.Initialize();
 
 		FName TargetName = GetFieldPhysicsName(EFieldPhysicsType::Field_LinearVelocity);
-		UnitTest.Solver->GetPerSolverField().AddTransientCommand( { TargetName, VectorField->NewCopy() });
+		UnitTest.Solver->GetPerSolverField().AddTransientCommand({ TargetName, VectorField->NewCopy() });
 		UnitTest.Advance();
 
 		FReal PreviousX = 0.0;
 		TManagedArray<FTransform>& Transform = Collection->DynamicCollection->Transform;
 		for (int Frame = 1; Frame < 10; Frame++)
 		{
-			UnitTest.Solver->GetPerSolverField().AddTransientCommand( { GetFieldPhysicsName(EFieldPhysicsType::Field_LinearVelocity), VectorField->NewCopy() });
+			UnitTest.Solver->GetPerSolverField().AddTransientCommand({ GetFieldPhysicsName(EFieldPhysicsType::Field_LinearVelocity), VectorField->NewCopy() });
 
 			UnitTest.Advance();
 
@@ -449,7 +449,7 @@ namespace GeometryCollectionTest
 		 * group, they drop through the ground.
 		 */
 
-		FFramework UnitTest; UnitTest.Dt = 1/24.0;
+		FFramework UnitTest; UnitTest.Dt = 1 / 24.0;
 
 		RigidBodyWrapper* Floor = TNewSimulationObject<GeometryType::RigidFloor>::Init()->template As<RigidBodyWrapper>();
 		UnitTest.AddSimulationObject(Floor);
@@ -462,16 +462,16 @@ namespace GeometryCollectionTest
 		Params.GeomTransform.SetScale3D(Scale);
 
 		FGeometryCollectionWrapper* Collection[4];
-		for (int n=0; n<3;n++)
+		for (int n = 0; n < 3; n++)
 		{
-			Params.RootTransform.SetLocation(FVector(0.f, 0.f, n*200.0f + 100.0f));
+			Params.RootTransform.SetLocation(FVector(0.f, 0.f, n * 200.0f + 100.0f));
 			Params.CollisionType = ECollisionTypeEnum::Chaos_Volumetric;
-			Collection[n+1] = TNewSimulationObject<GeometryType::GeometryCollectionWithSingleRigid>::Init(Params)->template As<FGeometryCollectionWrapper>();
-			UnitTest.AddSimulationObject(Collection[n+1]);
+			Collection[n + 1] = TNewSimulationObject<GeometryType::GeometryCollectionWithSingleRigid>::Init(Params)->template As<FGeometryCollectionWrapper>();
+			UnitTest.AddSimulationObject(Collection[n + 1]);
 		}
 
 		// Field setup
-		FRadialIntMask * RadialMask = new FRadialIntMask();
+		FRadialIntMask* RadialMask = new FRadialIntMask();
 		RadialMask->Position = FVector(0.0, 0.0, 0.0);
 		RadialMask->Radius = 0;
 		RadialMask->InteriorValue = -1;
@@ -528,23 +528,23 @@ namespace GeometryCollectionTest
 
 		UnitTest.AddSimulationObject(Collection);
 
-		FRadialFalloff * FalloffField = new FRadialFalloff();
+		FRadialFalloff* FalloffField = new FRadialFalloff();
 		FalloffField->Magnitude = 1.5;
 		FalloffField->Radius = 100.0;
 		FalloffField->Position = FVector(0.0, 0.0, 0.0);
 		FalloffField->Falloff = EFieldFalloffType::Field_FallOff_None;
-		
-		UnitTest.Initialize();		
+
+		UnitTest.Initialize();
 
 		auto& Clustering = UnitTest.Solver->GetEvolution()->GetRigidClustering();
 		const auto& ClusterMap = Clustering.GetChildrenMap();
 
 		UnitTest.Advance();
-		
+
 		TArray<Chaos::TPBDRigidClusteredParticleHandle<FReal, 3>*>& ParticleHandles = Collection->PhysObject->GetSolverParticleHandles();
 		{
 			FName TargetName = GetFieldPhysicsName(EFieldPhysicsType::Field_ExternalClusterStrain);
-			UnitTest.Solver->GetPerSolverField().AddTransientCommand( { TargetName, FalloffField->NewCopy() });
+			UnitTest.Solver->GetPerSolverField().AddTransientCommand({ TargetName, FalloffField->NewCopy() });
 
 			EXPECT_EQ(ClusterMap.Num(), 3);
 			EXPECT_EQ(ClusterMap[ParticleHandles[4]].Num(), 2);
@@ -563,9 +563,9 @@ namespace GeometryCollectionTest
 			EXPECT_TRUE(ParticleHandles[3]->Disabled());
 			EXPECT_TRUE(ParticleHandles[4]->Disabled());
 			EXPECT_TRUE(ParticleHandles[5]->Disabled());
-			EXPECT_FALSE(ParticleHandles[6]->Disabled());			
+			EXPECT_FALSE(ParticleHandles[6]->Disabled());
 
-			UnitTest.Advance();			
+			UnitTest.Advance();
 
 			// todo: indices here might seem odd, particles 4 & 5 are swapped
 			EXPECT_EQ(ClusterMap.Num(), 2);
@@ -602,32 +602,32 @@ namespace GeometryCollectionTest
 		Params.Simulating = true;
 		Params.EnableClustering = true;
 		Params.DamageThreshold = { 1.0 };
-		Params.MaxClusterLevel = 1000;		
+		Params.MaxClusterLevel = 1000;
 		Params.ClusterGroupIndex = 0;
 
 		FGeometryCollectionWrapper* Collection = TNewSimulationObject<GeometryType::GeometryCollectionWithSuppliedRestCollection>::Init(Params)->template As<FGeometryCollectionWrapper>();
 
 		UnitTest.AddSimulationObject(Collection);
 
-		FRadialFalloff * FalloffField = new FRadialFalloff();
+		FRadialFalloff* FalloffField = new FRadialFalloff();
 		FalloffField->Magnitude = 1.5;
 		FalloffField->Radius = 200.0;
 		FalloffField->Position = FVector(0.0, 0.0, 0.0);
 		FalloffField->Falloff = EFieldFalloffType::Field_FallOff_None;
 
-		UnitTest.Initialize();	
+		UnitTest.Initialize();
 		UnitTest.Advance();
 
 		TArray<Chaos::TPBDRigidClusteredParticleHandle<FReal, 3>*>& ParticleHandles = Collection->PhysObject->GetSolverParticleHandles();
 		auto& Clustering = UnitTest.Solver->GetEvolution()->GetRigidClustering();
 		const auto& ClusterMap = Clustering.GetChildrenMap();
 
-		{	
+		{
 			FName TargetName = GetFieldPhysicsName(EFieldPhysicsType::Field_ExternalClusterStrain);
 			FFieldSystemCommand Command(TargetName, FalloffField->NewCopy());
 			FFieldSystemMetaDataProcessingResolution* ResolutionData = new FFieldSystemMetaDataProcessingResolution(EFieldResolutionType::Field_Resolution_Maximum);
 			Command.MetaData.Add(FFieldSystemMetaData::EMetaType::ECommandData_ProcessingResolution, TUniquePtr< FFieldSystemMetaDataProcessingResolution >(ResolutionData));
-			UnitTest.Solver->GetPerSolverField().AddTransientCommand( Command);
+			UnitTest.Solver->GetPerSolverField().AddTransientCommand(Command);
 
 			EXPECT_EQ(ClusterMap.Num(), 3);
 			EXPECT_EQ(ClusterMap[ParticleHandles[6]].Num(), 3);
@@ -653,7 +653,7 @@ namespace GeometryCollectionTest
 			EXPECT_FALSE(ParticleHandles[8]->Disabled());
 
 			UnitTest.Advance();
-			UnitTest.Solver->GetPerSolverField().AddTransientCommand( Command);
+			UnitTest.Solver->GetPerSolverField().AddTransientCommand(Command);
 			UnitTest.Advance();
 
 			EXPECT_EQ(ClusterMap.Num(), 1);
@@ -690,14 +690,14 @@ namespace GeometryCollectionTest
 		Params.Simulating = true;
 		Params.EnableClustering = true;
 		Params.DamageThreshold = { 1.0 };
-		Params.MaxClusterLevel = 1000;		
+		Params.MaxClusterLevel = 1000;
 		Params.ClusterGroupIndex = 0;
 
 		FGeometryCollectionWrapper* Collection = TNewSimulationObject<GeometryType::GeometryCollectionWithSuppliedRestCollection>::Init(Params)->template As<FGeometryCollectionWrapper>();
 
 		UnitTest.AddSimulationObject(Collection);
 
-		FRadialFalloff * FalloffField = new FRadialFalloff();
+		FRadialFalloff* FalloffField = new FRadialFalloff();
 		FalloffField->Magnitude = 1.1;
 		FalloffField->Radius = 200.0;
 		FalloffField->Position = FVector(350.0, 0.0, 0.0);
@@ -718,8 +718,8 @@ namespace GeometryCollectionTest
 			FFieldSystemCommand Command(TargetName, FalloffField->NewCopy());
 			FFieldSystemMetaDataProcessingResolution* ResolutionData = new FFieldSystemMetaDataProcessingResolution(EFieldResolutionType::Field_Resolution_Maximum);
 			Command.MetaData.Add(FFieldSystemMetaData::EMetaType::ECommandData_ProcessingResolution, TUniquePtr< FFieldSystemMetaDataProcessingResolution >(ResolutionData));
-			UnitTest.Solver->GetPerSolverField().AddTransientCommand( Command);
-		
+			UnitTest.Solver->GetPerSolverField().AddTransientCommand(Command);
+
 			EXPECT_TRUE(ParticleHandles[0]->Disabled());
 			EXPECT_TRUE(ParticleHandles[1]->Disabled());
 			EXPECT_TRUE(ParticleHandles[2]->Disabled());
@@ -753,7 +753,7 @@ namespace GeometryCollectionTest
 			EXPECT_FALSE(ParticleHandles[7]->Disabled());
 			EXPECT_TRUE(ParticleHandles[8]->Disabled());
 		}
-			
+
 		delete FalloffField;
 
 	}
@@ -772,13 +772,13 @@ namespace GeometryCollectionTest
 		Params.Simulating = true;
 		Params.EnableClustering = true;
 		Params.DamageThreshold = { 1.0 };
-		Params.MaxClusterLevel = 1000;		
+		Params.MaxClusterLevel = 1000;
 		Params.ClusterGroupIndex = 0;
 
 		FGeometryCollectionWrapper* Collection = TNewSimulationObject<GeometryType::GeometryCollectionWithSuppliedRestCollection>::Init(Params)->template As<FGeometryCollectionWrapper>();
 
 		UnitTest.AddSimulationObject(Collection);
-		
+
 		FRadialFalloff* FalloffField = new FRadialFalloff();
 		FalloffField->Magnitude = 1.5;
 		FalloffField->Radius = 100.0;
@@ -794,7 +794,7 @@ namespace GeometryCollectionTest
 		TArray<Chaos::TPBDRigidClusteredParticleHandle<FReal, 3>*>& ParticleHandles = Collection->PhysObject->GetSolverParticleHandles();
 		{
 			FName TargetName = GetFieldPhysicsName(EFieldPhysicsType::Field_ExternalClusterStrain);
-			UnitTest.Solver->GetPerSolverField().AddTransientCommand( { TargetName, FalloffField->NewCopy() });
+			UnitTest.Solver->GetPerSolverField().AddTransientCommand({ TargetName, FalloffField->NewCopy() });
 
 			EXPECT_TRUE(ParticleHandles[0]->Disabled());
 			EXPECT_TRUE(ParticleHandles[1]->Disabled());
@@ -813,7 +813,7 @@ namespace GeometryCollectionTest
 			EXPECT_EQ(ClusterMap[ParticleHandles[5]].Num(), 2);
 			EXPECT_TRUE(ClusterMap[ParticleHandles[5]].Contains(ParticleHandles[2]));
 			EXPECT_TRUE(ClusterMap[ParticleHandles[5]].Contains(ParticleHandles[3]));
-			
+
 			EXPECT_TRUE(ParticleHandles[0]->Disabled());
 			EXPECT_TRUE(ParticleHandles[1]->Disabled());
 			EXPECT_TRUE(ParticleHandles[2]->Disabled());
@@ -822,7 +822,7 @@ namespace GeometryCollectionTest
 			EXPECT_FALSE(ParticleHandles[5]->Disabled());
 			EXPECT_TRUE(ParticleHandles[6]->Disabled());
 		}
-		
+
 		delete FalloffField;
 	}
 
@@ -911,7 +911,7 @@ namespace GeometryCollectionTest
 		UnitTest.Advance();
 
 		TArray<Chaos::TPBDRigidClusteredParticleHandle<FReal, 3>*>& ParticleHandles = Collection->PhysObject->GetSolverParticleHandles();
-		Chaos::TVector<float,3> CurrV = ParticleHandles[0]->V();
+		Chaos::TVector<float, 3> CurrV = ParticleHandles[0]->V();
 		TManagedArray<FTransform>& Transform = Collection->DynamicCollection->Transform;
 
 		for (int Frame = 1; Frame < 10; Frame++)
@@ -929,9 +929,57 @@ namespace GeometryCollectionTest
 			EXPECT_LT(Transform[0].GetTranslation().Z, LastLocation.Z);
 
 			LastLocation = Transform[0].GetTranslation();
+
+		}
+	}
+
+	GTEST_TEST(AllTraits, GeometryCollection_RigidBodies_Field_Algebra3)
+	{
+		FFramework UnitTest;
+
+		// Physics Object Setup
+		FVector Scale = FVector(10.0);
+		CreationParameters Params;
+		Params.DynamicState = EObjectStateTypeEnum::Chaos_Object_Dynamic;
+		Params.RootTransform.SetLocation(FVector(0.f, 0.f, 5.f));
+		Params.GeomTransform.SetScale3D(Scale);
+		FGeometryCollectionWrapper* Collection = TNewSimulationObject<GeometryType::GeometryCollectionWithSingleRigid>::Init(Params)->template As<FGeometryCollectionWrapper>();
+		UnitTest.AddSimulationObject(Collection);
+
+		// Fields setup... when both our applied, we expect Y rotation to cancel, while X rotation > 0.
+		FUniformVector* UniformVector1 = new FUniformVector();
+		UniformVector1->Direction = FVector(2.0, 1.0, 0.0);
+		UniformVector1->Magnitude = 100.0;
+		FUniformVector* UniformVector2 = new FUniformVector();
+		UniformVector2->Direction = FVector(-1.0, -1.0, 0.0);
+		UniformVector2->Magnitude = 100.0;
+
+		UnitTest.Initialize();
+
+		TManagedArray<FTransform>& Transform = Collection->DynamicCollection->Transform;
+		FReal PreviousHeight = Transform[0].GetTranslation().Z;
+		FReal PreviousX = Transform[0].GetRotation().Euler().X;
+		for (int Frame = 0; Frame < 10; Frame++)
+		{
+			
+			FName TargetName = GetFieldPhysicsName(EFieldPhysicsType::Field_AngularTorque);
+			UnitTest.Solver->GetPerSolverField().AddTransientCommand({ TargetName, UniformVector1->NewCopy() });
+			UnitTest.Solver->GetPerSolverField().AddTransientCommand({ TargetName, UniformVector2->NewCopy() });
 			
+			UnitTest.Advance();
+
+			Chaos::TPBDGeometryCollectionParticles<Chaos::FReal, 3>& Particles = UnitTest.Solver->GetParticles().GetGeometryCollectionParticles();
+			EXPECT_NE(FMath::Abs(Transform[0].GetRotation().Euler().Y), SMALL_THRESHOLD); // not rotating in Y?
+			EXPECT_GT(Particles.W(0).X, PreviousX); // rotating in X?
+			EXPECT_LT(Particles.X(0).Z, PreviousHeight); // still falling?
+
+			PreviousHeight = Particles.X(0).Z;
+			PreviousX = Particles.W(0).X;
 		}
+
+
 	}
 }
 
 
+
