max cluster level

diff --git a/Engine/Source/Programs/HeadlessChaos/Private/GeometryCollection/GeometryCollectionTestClustering.cpp b/Engine/Source/Programs/HeadlessChaos/Private/GeometryCollection/GeometryCollectionTestClustering.cpp
--- a/Engine/Source/Programs/HeadlessChaos/Private/GeometryCollection/GeometryCollectionTestClustering.cpp
+++ b/Engine/Source/Programs/HeadlessChaos/Private/GeometryCollection/GeometryCollectionTestClustering.cpp
@@ -2148,6 +2148,151 @@ namespace GeometryCollectionTest
 	
 	}
 
+	GTEST_TEST(AllTraits, GeometryCollection_RigidBodiess_ClusterTest_MaxClusterLevel_1)
+	{
+		FFramework UnitTest;
+		
+		// Create hierarchy 
+		TSharedPtr<FGeometryCollection> RestCollection;
+		RestCollection = GeometryCollection::MakeCubeElement(FTransform(FQuat::MakeFromEuler(FVector(0.f)), FVector(0, 0, 0)), FVector(1.0));
+		RestCollection->AppendGeometry(*GeometryCollection::MakeCubeElement(FTransform(FQuat::MakeFromEuler(FVector(0.f)), FVector(100, 0, 0)), FVector(1.0)));
+		RestCollection->AppendGeometry(*GeometryCollection::MakeCubeElement(FTransform(FQuat::MakeFromEuler(FVector(0.f)), FVector(200, 0, 0)), FVector(1.0)));
+		RestCollection->AppendGeometry(*GeometryCollection::MakeCubeElement(FTransform(FQuat::MakeFromEuler(FVector(0.f)), FVector(300, 0, 0)), FVector(1.0)));
+		RestCollection->AppendGeometry(*GeometryCollection::MakeCubeElement(FTransform(FQuat::MakeFromEuler(FVector(0.f)), FVector(400, 0, 0)), FVector(1.0)));
+		RestCollection->AppendGeometry(*GeometryCollection::MakeCubeElement(FTransform(FQuat::MakeFromEuler(FVector(0.f)), FVector(500, 0, 0)), FVector(1.0)));
+
+		RestCollection->AddElements(7, FGeometryCollection::TransformGroup);
+		RestCollection->Transform[10].SetTranslation(FVector(0,0,0));
+
+		RestCollection->SimulationType[0] = FGeometryCollection::ESimulationTypes::FST_Rigid;
+		RestCollection->SimulationType[1] = FGeometryCollection::ESimulationTypes::FST_Rigid;
+		RestCollection->SimulationType[2] = FGeometryCollection::ESimulationTypes::FST_Rigid;
+		RestCollection->SimulationType[3] = FGeometryCollection::ESimulationTypes::FST_Rigid;
+		RestCollection->SimulationType[4] = FGeometryCollection::ESimulationTypes::FST_Rigid;
+		RestCollection->SimulationType[5] = FGeometryCollection::ESimulationTypes::FST_Rigid;
+		RestCollection->SimulationType[6] = FGeometryCollection::ESimulationTypes::FST_Clustered;
+		RestCollection->SimulationType[7] = FGeometryCollection::ESimulationTypes::FST_Clustered;
+		RestCollection->SimulationType[8] = FGeometryCollection::ESimulationTypes::FST_Clustered;
+		RestCollection->SimulationType[9] = FGeometryCollection::ESimulationTypes::FST_Clustered;
+		RestCollection->SimulationType[10] = FGeometryCollection::ESimulationTypes::FST_Clustered;
+		RestCollection->SimulationType[11] = FGeometryCollection::ESimulationTypes::FST_Clustered;
+		RestCollection->SimulationType[12] = FGeometryCollection::ESimulationTypes::FST_Clustered;
+
+		GeometryCollectionAlgo::ParentTransforms(RestCollection.Get(), 6, { 11 });
+		GeometryCollectionAlgo::ParentTransforms(RestCollection.Get(), 7, { 3,4,5 });
+		GeometryCollectionAlgo::ParentTransforms(RestCollection.Get(), 8, { 6 });
+		GeometryCollectionAlgo::ParentTransforms(RestCollection.Get(), 9, { 7 });
+		GeometryCollectionAlgo::ParentTransforms(RestCollection.Get(), 10, { 9,8 });
+		GeometryCollectionAlgo::ParentTransforms(RestCollection.Get(), 11, { 12 });
+		GeometryCollectionAlgo::ParentTransforms(RestCollection.Get(), 12, { 0,1,2 });
+
+		CreationParameters Params;
+		Params.RestCollection = RestCollection;
+		Params.DynamicState = EObjectStateTypeEnum::Chaos_Object_Dynamic;
+		Params.CollisionType = ECollisionTypeEnum::Chaos_Surface_Volumetric;
+		Params.ImplicitType = EImplicitTypeEnum::Chaos_Implicit_Box;
+		Params.Simulating = true;
+		Params.EnableClustering = true;
+		Params.DamageThreshold = { 100.0f, 4.0f, 2.0f };
+		Params.MaxClusterLevel = 3;
+		Params.ClusterGroupIndex = 0;
+
+		FGeometryCollectionWrapper* Collection = TNewSimulationObject<GeometryType::GeometryCollectionWithSuppliedRestCollection>::Init(Params)->template As<FGeometryCollectionWrapper>();
+
+		UnitTest.AddSimulationObject(Collection);
+		UnitTest.Initialize();
+		UnitTest.Advance();
+
+		FRadialFalloff* FalloffField = new FRadialFalloff();
+		FalloffField->Magnitude = 104.0;
+		FalloffField->Radius = 10000.0;
+		FalloffField->Position = FVector(0.0, 0.0, 0.0);
+		FalloffField->Falloff = EFieldFalloffType::Field_FallOff_None;
+
+		TArray<Chaos::TPBDRigidClusteredParticleHandle<FReal, 3>*>& ParticleHandles = Collection->PhysObject->GetSolverParticleHandles();
+		Chaos::TPBDRigidClustering<Chaos::FPBDRigidsEvolutionGBF, Chaos::FPBDCollisionConstraints>& Clustering = UnitTest.Solver->GetEvolution()->GetRigidClustering();
+		const auto& ClusterMap = Clustering.GetChildrenMap();
+		auto& x = UnitTest.Solver->GetEvolution()->GetRigidClustering().GetStrainArray();
+
+
+		{
+			EXPECT_EQ(ClusterMap.Num(), 5);
+			EXPECT_EQ(ClusterMap[ParticleHandles[6]].Num(), 3);
+			EXPECT_TRUE(ClusterMap[ParticleHandles[6]].Contains(ParticleHandles[0]));
+			EXPECT_TRUE(ClusterMap[ParticleHandles[6]].Contains(ParticleHandles[1]));
+			EXPECT_TRUE(ClusterMap[ParticleHandles[6]].Contains(ParticleHandles[2]));
+			EXPECT_EQ(ClusterMap[ParticleHandles[7]].Num(), 3);
+			EXPECT_TRUE(ClusterMap[ParticleHandles[7]].Contains(ParticleHandles[3]));
+			EXPECT_TRUE(ClusterMap[ParticleHandles[7]].Contains(ParticleHandles[4]));
+			EXPECT_TRUE(ClusterMap[ParticleHandles[7]].Contains(ParticleHandles[5]));
+			EXPECT_EQ(ClusterMap[ParticleHandles[8]].Num(), 1);
+			EXPECT_TRUE(ClusterMap[ParticleHandles[8]].Contains(ParticleHandles[6]));
+			EXPECT_EQ(ClusterMap[ParticleHandles[9]].Num(), 1);
+			EXPECT_TRUE(ClusterMap[ParticleHandles[9]].Contains(ParticleHandles[7]));
+			EXPECT_EQ(ClusterMap[ParticleHandles[10]].Num(), 2);
+			EXPECT_TRUE(ClusterMap[ParticleHandles[10]].Contains(ParticleHandles[8]));
+			EXPECT_TRUE(ClusterMap[ParticleHandles[10]].Contains(ParticleHandles[9]));
+		}
+
+		EXPECT_TRUE(ParticleHandles[6]->Disabled());
+		EXPECT_TRUE(ParticleHandles[7]->Disabled());
+		EXPECT_TRUE(ParticleHandles[8]->Disabled());
+		EXPECT_TRUE(ParticleHandles[9]->Disabled());
+		EXPECT_FALSE(ParticleHandles[10]->Disabled());
+		// {100, 4, 2} -> {4, 2}
+		FName TargetName = GetFieldPhysicsName(EFieldPhysicsType::Field_ExternalClusterStrain);
+		UnitTest.Solver->GetPerSolverField().AddTransientCommand({ TargetName, FalloffField->NewCopy() });
+		UnitTest.Advance();
+		EXPECT_TRUE(ParticleHandles[6]->Disabled());
+		EXPECT_TRUE(ParticleHandles[7]->Disabled());
+		EXPECT_FALSE(ParticleHandles[8]->Disabled());
+		EXPECT_FALSE(ParticleHandles[9]->Disabled());
+		EXPECT_TRUE(ParticleHandles[10]->Disabled());
+
+		// {4, 2} -> {4, 2} (beyond max cluster level)
+		UnitTest.Solver->GetPerSolverField().AddTransientCommand({ TargetName, FalloffField->NewCopy() });
+		UnitTest.Advance();
+		EXPECT_TRUE(ParticleHandles[0]->Disabled());
+		EXPECT_TRUE(ParticleHandles[1]->Disabled());
+		EXPECT_TRUE(ParticleHandles[2]->Disabled());
+		EXPECT_TRUE(ParticleHandles[3]->Disabled());
+		EXPECT_TRUE(ParticleHandles[4]->Disabled());
+		EXPECT_TRUE(ParticleHandles[5]->Disabled());
+		EXPECT_TRUE(ParticleHandles[6]->Disabled());
+		EXPECT_TRUE(ParticleHandles[7]->Disabled());
+		EXPECT_FALSE(ParticleHandles[8]->Disabled());
+		EXPECT_FALSE(ParticleHandles[9]->Disabled());
+		EXPECT_TRUE(ParticleHandles[10]->Disabled());
+
+		// {4, 2} -> {4, 2} (beyond max cluster level)
+		// but this time, force internal reordering via internal strain.
+		TargetName = GetFieldPhysicsName(EFieldPhysicsType::Field_InternalClusterStrain); // Note: now internal strain
+		FalloffField->Magnitude = -3.0;
+		UnitTest.Solver->GetPerSolverField().AddTransientCommand({ TargetName, FalloffField->NewCopy() });
+		UnitTest.Advance();
+		EXPECT_TRUE(ParticleHandles[0]->Disabled());
+		EXPECT_TRUE(ParticleHandles[1]->Disabled());
+		EXPECT_TRUE(ParticleHandles[2]->Disabled());
+		EXPECT_TRUE(ParticleHandles[3]->Disabled());
+		EXPECT_TRUE(ParticleHandles[4]->Disabled());
+		EXPECT_TRUE(ParticleHandles[5]->Disabled());
+		EXPECT_TRUE(ParticleHandles[6]->Disabled());
+		EXPECT_TRUE(ParticleHandles[7]->Disabled());
+		EXPECT_FALSE(ParticleHandles[8]->Disabled());
+		EXPECT_FALSE(ParticleHandles[9]->Disabled());
+		EXPECT_TRUE(ParticleHandles[10]->Disabled());
+
+		{ // Check if level 3 rigids (children) are still clustered despite internal strain breaking level 2 clusters (parent)
+			EXPECT_EQ(ClusterMap[ParticleHandles[6]].Num(), 3);
+			EXPECT_TRUE(ClusterMap[ParticleHandles[6]].Contains(ParticleHandles[0]));
+			EXPECT_TRUE(ClusterMap[ParticleHandles[6]].Contains(ParticleHandles[1]));
+			EXPECT_TRUE(ClusterMap[ParticleHandles[6]].Contains(ParticleHandles[2]));
+			EXPECT_EQ(ClusterMap[ParticleHandles[7]].Num(), 3);
+			EXPECT_TRUE(ClusterMap[ParticleHandles[7]].Contains(ParticleHandles[3]));
+			EXPECT_TRUE(ClusterMap[ParticleHandles[7]].Contains(ParticleHandles[4]));
+			EXPECT_TRUE(ClusterMap[ParticleHandles[7]].Contains(ParticleHandles[5]));
+		}
+	}
 	
 	GTEST_TEST(AllTraits, GeometryCollection_RigidBodiess_ClusterTest_ParticleImplicitCollisionGeometry)
 	{
