Removing SimplicialCubeToAnalyticCube

diff --git a/Engine/Source/Programs/HeadlessChaos/Private/GeometryCollection/GeometryCollectionTestCollisionResolution.cpp b/Engine/Source/Programs/HeadlessChaos/Private/GeometryCollection/GeometryCollectionTestCollisionResolution.cpp
--- a/Engine/Source/Programs/HeadlessChaos/Private/GeometryCollection/GeometryCollectionTestCollisionResolution.cpp
+++ b/Engine/Source/Programs/HeadlessChaos/Private/GeometryCollection/GeometryCollectionTestCollisionResolution.cpp
@@ -262,11 +262,6 @@ namespace GeometryCollectionTest
 	}
 
 
-	void CollisionResolution_SimplicialCubeToAnalyticCube()
-	{
-
-	}
-
 
 	GTEST_TEST(AllTraits,GeometryCollection_CollisionResolution_SimplicialSphereToAnalyticSphere)
 	{
