min/max field resolution

diff --git a/Engine/Source/Programs/HeadlessChaos/Private/GeometryCollection/GeometryCollectionTestSimulationField.cpp b/Engine/Source/Programs/HeadlessChaos/Private/GeometryCollection/GeometryCollectionTestSimulationField.cpp
--- a/Engine/Source/Programs/HeadlessChaos/Private/GeometryCollection/GeometryCollectionTestSimulationField.cpp
+++ b/Engine/Source/Programs/HeadlessChaos/Private/GeometryCollection/GeometryCollectionTestSimulationField.cpp
@@ -6,6 +6,7 @@
 
 #include "GeometryCollection/GeometryCollection.h"
 #include "GeometryCollection/GeometryCollectionUtility.h"
+#include "GeometryCollection/GeometryCollectionAlgo.h"
 #include "GeometryCollection/TransformCollection.h"
 #include "UObject/Package.h"
 #include "UObject/UObjectGlobals.h"
@@ -763,6 +764,7 @@ namespace GeometryCollectionTest
 		FFramework UnitTest;
 
 		TSharedPtr<FGeometryCollection> RestCollection = CreateClusteredBody_TwoByTwo_ThreeTransform(FVector(0));
+		
 
 		CreationParameters Params;
 		Params.RestCollection = RestCollection;
@@ -979,6 +981,147 @@ namespace GeometryCollectionTest
 
 
 	}
+
+	GTEST_TEST(AllTraits, GeometryCollection_RigidBodies_Field_FieldResolutionMinimal)
+	{
+		FFramework UnitTest;
+
+		// Construct Rest Collection...
+		TSharedPtr<FGeometryCollection> RestCollection = GeometryCollection::MakeCubeElement(FTransform(FQuat::MakeFromEuler(FVector(0.f)), FVector(0, 0, 0)), FVector(1.0));
+		RestCollection->AppendGeometry(*GeometryCollection::MakeCubeElement(FTransform(FQuat::MakeFromEuler(FVector(0.f)), FVector(100, 0, 0)), FVector(1.0)));
+		RestCollection->AppendGeometry(*GeometryCollection::MakeCubeElement(FTransform(FQuat::MakeFromEuler(FVector(0.f)), FVector(200, 0, 0)), FVector(1.0)));
+		RestCollection->AppendGeometry(*GeometryCollection::MakeCubeElement(FTransform(FQuat::MakeFromEuler(FVector(0.f)), FVector(300, 0, 0)), FVector(1.0)));
+		RestCollection->AppendGeometry(*GeometryCollection::MakeCubeElement(FTransform(FQuat::MakeFromEuler(FVector(0.f)), FVector(400, 0, 0)), FVector(1.0)));
+		RestCollection->Transform[0].SetTranslation(FVector(0.0f));
+		RestCollection->SimulationType[0] = FGeometryCollection::ESimulationTypes::FST_Clustered;
+		RestCollection->SimulationType[1] = FGeometryCollection::ESimulationTypes::FST_Clustered;
+		RestCollection->SimulationType[2] = FGeometryCollection::ESimulationTypes::FST_Clustered;
+		RestCollection->SimulationType[3] = FGeometryCollection::ESimulationTypes::FST_Clustered;
+		RestCollection->SimulationType[4] = FGeometryCollection::ESimulationTypes::FST_Rigid;
+		GeometryCollectionAlgo::ParentTransforms(RestCollection.Get(), 0, { 1 });
+		GeometryCollectionAlgo::ParentTransforms(RestCollection.Get(), 1, { 2 });
+		GeometryCollectionAlgo::ParentTransforms(RestCollection.Get(), 2, { 3 });
+		GeometryCollectionAlgo::ParentTransforms(RestCollection.Get(), 3, { 4 });
+
+		CreationParameters Params;
+		Params.RestCollection = RestCollection;
+		Params.DynamicState = EObjectStateTypeEnum::Chaos_Object_Dynamic;
+		Params.ImplicitType = EImplicitTypeEnum::Chaos_Implicit_Box;
+		Params.CollisionType = ECollisionTypeEnum::Chaos_Surface_Volumetric;
+		Params.Simulating = true;
+		Params.EnableClustering = true;
+		Params.DamageThreshold = { 91.0f, 92.0f, 93.0f, 94.0f };
+		Params.RemoveOnFractureEnabled = true;
+		FGeometryCollectionWrapper* Collection = TNewSimulationObject<GeometryType::GeometryCollectionWithSuppliedRestCollection>::Init(Params)->template As<FGeometryCollectionWrapper>();
+		UnitTest.AddSimulationObject(Collection);
+
+		FRadialFalloff* FalloffField = new FRadialFalloff();
+		FalloffField->Magnitude = -100;
+		FalloffField->Radius = 1000.0;
+		FalloffField->Position = FVector(0.0, 0.0, 0.0);
+		FalloffField->Falloff = EFieldFalloffType::Field_FallOff_None;
+
+		UnitTest.Initialize();
+		UnitTest.Advance();
+
+		TArray<Chaos::TPBDRigidClusteredParticleHandle<FReal, 3>*>& ParticleHandles = Collection->PhysObject->GetSolverParticleHandles();
+		auto& Clustering = UnitTest.Solver->GetEvolution()->GetRigidClustering();
+		const auto& ClusterMap = Clustering.GetChildrenMap();
+
+		
+		Chaos::TArrayCollectionArray<float>& StrainArray = UnitTest.Solver->GetEvolution()->GetRigidClustering().GetStrainArray();
+
+		EXPECT_FALSE(ParticleHandles[0]->Disabled());
+		EXPECT_TRUE(ParticleHandles[1]->Disabled());
+		EXPECT_TRUE(ParticleHandles[2]->Disabled());
+		EXPECT_TRUE(ParticleHandles[3]->Disabled());
+		EXPECT_TRUE(ParticleHandles[4]->Disabled());
+
+		FName TargetName = GetFieldPhysicsName(EFieldPhysicsType::Field_InternalClusterStrain);
+		FFieldSystemCommand Command(TargetName, FalloffField->NewCopy());
+		FFieldSystemMetaDataProcessingResolution* ResolutionData = new FFieldSystemMetaDataProcessingResolution(EFieldResolutionType::Field_Resolution_Minimal);
+		Command.MetaData.Add(FFieldSystemMetaData::EMetaType::ECommandData_ProcessingResolution, TUniquePtr< FFieldSystemMetaDataProcessingResolution >(ResolutionData));
+		UnitTest.Solver->GetPerSolverField().AddTransientCommand(Command);
+		UnitTest.Advance();
+
+		// We expect only active clusters and their children to be affected by the field with Field_Resolution_Minimal
+		EXPECT_GT(StrainArray[0], 0.0f);
+		EXPECT_GT(StrainArray[1], 0.0f);
+		EXPECT_GT(0.0f, StrainArray[2]);
+		EXPECT_GT(0.0f, StrainArray[3]);
+
+	}
+
+	GTEST_TEST(AllTraits, GeometryCollection_RigidBodies_Field_FieldResolutionMaximum)
+	{
+		FFramework UnitTest;
+
+		// Construct Rest Collection...
+		TSharedPtr<FGeometryCollection> RestCollection = GeometryCollection::MakeCubeElement(FTransform(FQuat::MakeFromEuler(FVector(0.f)), FVector(0, 0, 0)), FVector(1.0));
+		RestCollection->AppendGeometry(*GeometryCollection::MakeCubeElement(FTransform(FQuat::MakeFromEuler(FVector(0.f)), FVector(100, 0, 0)), FVector(1.0)));
+		RestCollection->AppendGeometry(*GeometryCollection::MakeCubeElement(FTransform(FQuat::MakeFromEuler(FVector(0.f)), FVector(200, 0, 0)), FVector(1.0)));
+		RestCollection->AppendGeometry(*GeometryCollection::MakeCubeElement(FTransform(FQuat::MakeFromEuler(FVector(0.f)), FVector(300, 0, 0)), FVector(1.0)));
+		RestCollection->AppendGeometry(*GeometryCollection::MakeCubeElement(FTransform(FQuat::MakeFromEuler(FVector(0.f)), FVector(400, 0, 0)), FVector(1.0)));
+		RestCollection->Transform[0].SetTranslation(FVector(0.0f));
+		RestCollection->SimulationType[0] = FGeometryCollection::ESimulationTypes::FST_Clustered;
+		RestCollection->SimulationType[1] = FGeometryCollection::ESimulationTypes::FST_Clustered;
+		RestCollection->SimulationType[2] = FGeometryCollection::ESimulationTypes::FST_Clustered;
+		RestCollection->SimulationType[3] = FGeometryCollection::ESimulationTypes::FST_Clustered;
+		RestCollection->SimulationType[4] = FGeometryCollection::ESimulationTypes::FST_Rigid;
+		GeometryCollectionAlgo::ParentTransforms(RestCollection.Get(), 0, { 1 });
+		GeometryCollectionAlgo::ParentTransforms(RestCollection.Get(), 1, { 2 });
+		GeometryCollectionAlgo::ParentTransforms(RestCollection.Get(), 2, { 3 });
+		GeometryCollectionAlgo::ParentTransforms(RestCollection.Get(), 3, { 4 });
+
+		CreationParameters Params;
+		Params.RestCollection = RestCollection;
+		Params.DynamicState = EObjectStateTypeEnum::Chaos_Object_Dynamic;
+		Params.ImplicitType = EImplicitTypeEnum::Chaos_Implicit_Box;
+		Params.CollisionType = ECollisionTypeEnum::Chaos_Surface_Volumetric;
+		Params.Simulating = true;
+		Params.EnableClustering = true;
+		Params.DamageThreshold = { 91.0f, 92.0f, 93.0f, 94.0f };
+		Params.RemoveOnFractureEnabled = true;
+		FGeometryCollectionWrapper* Collection = TNewSimulationObject<GeometryType::GeometryCollectionWithSuppliedRestCollection>::Init(Params)->template As<FGeometryCollectionWrapper>();
+		UnitTest.AddSimulationObject(Collection);
+
+		FRadialFalloff* FalloffField = new FRadialFalloff();
+		FalloffField->Magnitude = -100;
+		FalloffField->Radius = 1000.0;
+		FalloffField->Position = FVector(0.0, 0.0, 0.0);
+		FalloffField->Falloff = EFieldFalloffType::Field_FallOff_None;
+
+		UnitTest.Initialize();
+		UnitTest.Advance();
+
+		TArray<Chaos::TPBDRigidClusteredParticleHandle<FReal, 3>*>& ParticleHandles = Collection->PhysObject->GetSolverParticleHandles();
+		auto& Clustering = UnitTest.Solver->GetEvolution()->GetRigidClustering();
+		const auto& ClusterMap = Clustering.GetChildrenMap();
+
+
+		Chaos::TArrayCollectionArray<float>& StrainArray = UnitTest.Solver->GetEvolution()->GetRigidClustering().GetStrainArray();
+
+		EXPECT_FALSE(ParticleHandles[0]->Disabled());
+		EXPECT_TRUE(ParticleHandles[1]->Disabled());
+		EXPECT_TRUE(ParticleHandles[2]->Disabled());
+		EXPECT_TRUE(ParticleHandles[3]->Disabled());
+		EXPECT_TRUE(ParticleHandles[4]->Disabled());
+
+		FName TargetName = GetFieldPhysicsName(EFieldPhysicsType::Field_InternalClusterStrain);
+		FFieldSystemCommand Command(TargetName, FalloffField->NewCopy());
+		FFieldSystemMetaDataProcessingResolution* ResolutionData = new FFieldSystemMetaDataProcessingResolution(EFieldResolutionType::Field_Resolution_Maximum);
+		Command.MetaData.Add(FFieldSystemMetaData::EMetaType::ECommandData_ProcessingResolution, TUniquePtr< FFieldSystemMetaDataProcessingResolution >(ResolutionData));
+		UnitTest.Solver->GetPerSolverField().AddTransientCommand(Command);
+		UnitTest.Advance();
+		UnitTest.Advance();
+		// We expect all clusters to be affected by the field with Field_Resolution_Maximum
+		EXPECT_GT(0.0f, StrainArray[0]);
+		EXPECT_GT(0.0f, StrainArray[1]);
+		EXPECT_GT(0.0f, StrainArray[2]);
+		EXPECT_GT(0.0f, StrainArray[3]);
+		
+	}
+
 }
 
 
