clustering w/ internal strain decay fields (resolution minimal)

diff --git a/Engine/Source/Programs/HeadlessChaos/Private/GeometryCollection/GeometryCollectionTestSimulationField.cpp b/Engine/Source/Programs/HeadlessChaos/Private/GeometryCollection/GeometryCollectionTestSimulationField.cpp
--- a/Engine/Source/Programs/HeadlessChaos/Private/GeometryCollection/GeometryCollectionTestSimulationField.cpp
+++ b/Engine/Source/Programs/HeadlessChaos/Private/GeometryCollection/GeometryCollectionTestSimulationField.cpp
@@ -828,6 +828,112 @@ namespace GeometryCollectionTest
 		delete FalloffField;
 	}
 
+	GTEST_TEST(AllTraits, GeometryCollection_RigidBodies_Field_ClusterBreak_StrainModel_TestResolutionMinimal)
+	{
+		FFramework UnitTest;
+		TSharedPtr<FGeometryCollection> RestCollection = CreateClusteredBody_ThreeByTwo_ThreeTransform(FVector(0));
+
+		CreationParameters Params;
+		Params.RestCollection = RestCollection;
+		Params.DynamicState = EObjectStateTypeEnum::Chaos_Object_Dynamic;
+		Params.CollisionType = ECollisionTypeEnum::Chaos_Surface_Volumetric;
+		Params.ImplicitType = EImplicitTypeEnum::Chaos_Implicit_Box;
+		Params.Simulating = true;
+		Params.EnableClustering = true;
+		Params.DamageThreshold = { 101.0, 103.0};
+		Params.MaxClusterLevel = 10;
+
+		FGeometryCollectionWrapper* Collection = TNewSimulationObject<GeometryType::GeometryCollectionWithSuppliedRestCollection>::Init(Params)->template As<FGeometryCollectionWrapper>();
+		UnitTest.AddSimulationObject(Collection);
+
+		FRadialFalloff* FalloffField = new FRadialFalloff(); // field resolution is minimal by default
+		FalloffField->Magnitude = -100.0;
+		FalloffField->Radius = 10000.0;
+		FalloffField->Position = FVector(0.0, 0.0, 0.0);
+		FalloffField->Falloff = EFieldFalloffType::Field_FallOff_None;
+
+		UnitTest.Initialize();
+		auto& ParticleHandles = Collection->PhysObject->GetSolverParticleHandles();
+		auto& ClusteredParticleHandles = Collection->PhysObject->GetSolverClusterHandles();
+		auto& Clustering = UnitTest.Solver->GetEvolution()->GetRigidClustering();
+		const auto& ClusterMap = Clustering.GetChildrenMap();
+		UnitTest.Solver->RegisterSimOneShotCallback([&]()
+		{
+			EXPECT_EQ(ClusterMap.Num(), 3);
+			EXPECT_EQ(ClusteredParticleHandles[1]->Strain(), 101);
+		});
+		FName TargetName = GetFieldPhysicsName(EFieldPhysicsType::Field_InternalClusterStrain);
+		FalloffField->Magnitude = 0.0;
+		UnitTest.Solver->GetPerSolverField().AddTransientCommand({ TargetName, FalloffField->NewCopy() });
+		UnitTest.Advance();
+
+
+		{
+			// { 101.0, 103.0 }
+			EXPECT_TRUE(ParticleHandles[0]->Disabled());
+			EXPECT_TRUE(ParticleHandles[1]->Disabled());
+			EXPECT_TRUE(ParticleHandles[2]->Disabled());
+			EXPECT_TRUE(ParticleHandles[3]->Disabled());
+			EXPECT_TRUE(ParticleHandles[4]->Disabled());
+			EXPECT_TRUE(ParticleHandles[5]->Disabled());
+			EXPECT_TRUE(ParticleHandles[6]->Disabled());
+			EXPECT_TRUE(ParticleHandles[7]->Disabled());
+			EXPECT_FALSE(ParticleHandles[8]->Disabled());
+
+			FalloffField->Magnitude = -100;
+			UnitTest.Solver->GetPerSolverField().AddTransientCommand({ TargetName, FalloffField->NewCopy() });
+			UnitTest.Advance();
+
+			// { 1.0, 3.0 }
+			EXPECT_TRUE(ParticleHandles[0]->Disabled());
+			EXPECT_TRUE(ParticleHandles[1]->Disabled());
+			EXPECT_TRUE(ParticleHandles[2]->Disabled());
+			EXPECT_TRUE(ParticleHandles[3]->Disabled());
+			EXPECT_TRUE(ParticleHandles[4]->Disabled());
+			EXPECT_TRUE(ParticleHandles[5]->Disabled());
+			EXPECT_TRUE(ParticleHandles[6]->Disabled());
+			EXPECT_TRUE(ParticleHandles[7]->Disabled());
+			EXPECT_FALSE(ParticleHandles[8]->Disabled()); // Parent still clustered w/ threshold > 0
+
+			EXPECT_EQ(ClusterMap.Num(), 3);
+			EXPECT_EQ(ClusterMap[ParticleHandles[6]].Num(), 3);
+			EXPECT_TRUE(ClusterMap[ParticleHandles[6]].Contains(ParticleHandles[0]));
+			EXPECT_TRUE(ClusterMap[ParticleHandles[6]].Contains(ParticleHandles[1]));
+			EXPECT_TRUE(ClusterMap[ParticleHandles[6]].Contains(ParticleHandles[2]));
+			EXPECT_EQ(ClusterMap[ParticleHandles[7]].Num(), 3);
+			EXPECT_TRUE(ClusterMap[ParticleHandles[7]].Contains(ParticleHandles[3]));
+			EXPECT_TRUE(ClusterMap[ParticleHandles[7]].Contains(ParticleHandles[4]));
+			EXPECT_TRUE(ClusterMap[ParticleHandles[7]].Contains(ParticleHandles[5]));
+			EXPECT_EQ(ClusterMap[ParticleHandles[8]].Num(), 2);
+			EXPECT_TRUE(ClusterMap[ParticleHandles[8]].Contains(ParticleHandles[6]));
+			EXPECT_TRUE(ClusterMap[ParticleHandles[8]].Contains(ParticleHandles[7]));
+
+			FalloffField->Magnitude = -1.0;
+			UnitTest.Solver->GetPerSolverField().AddTransientCommand({ TargetName, FalloffField->NewCopy() });
+			UnitTest.Advance();
+
+			// { 0.0, 2.0 } broken
+			EXPECT_FALSE(ParticleHandles[6]->Disabled());
+			EXPECT_FALSE(ParticleHandles[7]->Disabled());
+			EXPECT_TRUE(ParticleHandles[8]->Disabled()); 
+
+			FalloffField->Magnitude = -103;
+			UnitTest.Solver->GetPerSolverField().AddTransientCommand({ TargetName, FalloffField->NewCopy() });
+			UnitTest.Advance();
+
+			// { 0.0, 0.0 }
+			EXPECT_FALSE(ParticleHandles[0]->Disabled());
+			EXPECT_FALSE(ParticleHandles[1]->Disabled());
+			EXPECT_FALSE(ParticleHandles[2]->Disabled());
+			EXPECT_FALSE(ParticleHandles[3]->Disabled());
+			EXPECT_FALSE(ParticleHandles[4]->Disabled());
+			EXPECT_FALSE(ParticleHandles[5]->Disabled());
+			EXPECT_TRUE(ParticleHandles[6]->Disabled());
+			EXPECT_TRUE(ParticleHandles[7]->Disabled());
+			EXPECT_TRUE(ParticleHandles[8]->Disabled());
+		}
+	}
+
 	GTEST_TEST(AllTraits, GeometryCollection_RigidBodies_Field_Algebra1)
 	{
 		FFramework UnitTest;
@@ -1029,7 +1135,7 @@ namespace GeometryCollectionTest
 		const auto& ClusterMap = Clustering.GetChildrenMap();
 
 		
-		Chaos::TArrayCollectionArray<float>& StrainArray = UnitTest.Solver->GetEvolution()->GetRigidClustering().GetStrainArray();
+		Chaos::TArrayCollectionArray<Chaos::FReal>& StrainArray = UnitTest.Solver->GetEvolution()->GetRigidClustering().GetStrainArray();
 
 		EXPECT_FALSE(ParticleHandles[0]->Disabled());
 		EXPECT_TRUE(ParticleHandles[1]->Disabled());
@@ -1099,7 +1205,7 @@ namespace GeometryCollectionTest
 		const auto& ClusterMap = Clustering.GetChildrenMap();
 
 
-		Chaos::TArrayCollectionArray<float>& StrainArray = UnitTest.Solver->GetEvolution()->GetRigidClustering().GetStrainArray();
+		Chaos::TArrayCollectionArray<Chaos::FReal>& StrainArray = UnitTest.Solver->GetEvolution()->GetRigidClustering().GetStrainArray();
 
 		EXPECT_FALSE(ParticleHandles[0]->Disabled());
 		EXPECT_TRUE(ParticleHandles[1]->Disabled());
