Chaos - unit test - LWC - fix EPA test returning wrong value in LWC because of newly more precise FQuat

diff --git a/Engine/Source/Programs/HeadlessChaos/Private/HeadlessChaosTestEPA.cpp b/Engine/Source/Programs/HeadlessChaos/Private/HeadlessChaosTestEPA.cpp
--- a/Engine/Source/Programs/HeadlessChaos/Private/HeadlessChaosTestEPA.cpp
+++ b/Engine/Source/Programs/HeadlessChaos/Private/HeadlessChaosTestEPA.cpp
@@ -616,7 +616,14 @@ namespace ChaosTest
 			GJKPenetration(A, B, BToATM, Penetration, ClosestA, ClosestB, NormalA, ClosestVertexIndexA, ClosestVertexIndexB);
 			FVec3 Normal = BToATM.InverseTransformVector(NormalA);
 
+			// Why do we have two result depending on the precision of FReal:
+			// The double result is actually the correct one ( 0.04752 ), I have checked it in a 3D modeling package and I get the same value down to the 5th decimal digit
+			// The float version is certainly because of the imprecision of the quaternion and the rather large ( 200 x 200 x 20 ) object, amplifying the errors, 
+#if UE_LARGE_WORLD_COORDINATES_DISABLED
 			EXPECT_NEAR(Penetration, 0.025f, 0.005f);
+#else
+			EXPECT_NEAR(Penetration, 0.04752f, 0.005f);
+#endif
 			EXPECT_NEAR(Normal.Z, -1.0f, 0.001f);
 		}
 
