ChaosHeadless
00 Updated test to new ProximityUtility API

diff --git a/Engine/Source/Programs/HeadlessChaos/Private/GeometryCollection/GeometryCollectionTestProximity.cpp b/Engine/Source/Programs/HeadlessChaos/Private/GeometryCollection/GeometryCollectionTestProximity.cpp
--- a/Engine/Source/Programs/HeadlessChaos/Private/GeometryCollection/GeometryCollectionTestProximity.cpp
+++ b/Engine/Source/Programs/HeadlessChaos/Private/GeometryCollection/GeometryCollectionTestProximity.cpp
@@ -34,7 +34,8 @@ namespace GeometryCollectionTest
 		TArray<FTransform> GlobalTransform;
 		GeometryCollectionAlgo::GlobalMatrices(Collection->Transform, Collection->Parent, GlobalTransform);
 
-		FGeometryCollectionProximityUtility::UpdateProximity(Coll);
+		FGeometryCollectionProximityUtility ProximityUtility(Coll);
+		ProximityUtility.UpdateProximity();
 
 		// Breaking Data
 		TManagedArray<TSet<int32>>& Proximity = Collection->GetAttribute<TSet<int32>>("Proximity", FGeometryCollection::GeometryGroup);
@@ -83,7 +84,8 @@ namespace GeometryCollectionTest
 		TArray<FTransform> GlobalTransform;
 		GeometryCollectionAlgo::GlobalMatrices(Collection->Transform, Collection->Parent, GlobalTransform);
 
-		FGeometryCollectionProximityUtility::UpdateProximity(Coll);
+		FGeometryCollectionProximityUtility ProximityUtility(Coll);
+		ProximityUtility.UpdateProximity();
 
 		// Breaking Data
 		TManagedArray<TSet<int32>>& Proximity = Collection->GetAttribute<TSet<int32>>("Proximity", FGeometryCollection::GeometryGroup);
@@ -210,7 +212,8 @@ namespace GeometryCollectionTest
 		TArray<FTransform> GlobalTransform;
 		GeometryCollectionAlgo::GlobalMatrices(Collection->Transform, Collection->Parent, GlobalTransform);
 
-		FGeometryCollectionProximityUtility::UpdateProximity(Coll);
+		FGeometryCollectionProximityUtility ProximityUtility(Coll);
+		ProximityUtility.UpdateProximity();
 
 		// Breaking Data
 		TManagedArray<TSet<int32>>& Proximity = Collection->GetAttribute<TSet<int32>>("Proximity", FGeometryCollection::GeometryGroup);
@@ -337,7 +340,8 @@ namespace GeometryCollectionTest
 		TArray<FTransform> GlobalTransform;
 		GeometryCollectionAlgo::GlobalMatrices(Collection->Transform, Collection->Parent, GlobalTransform);
 
-		FGeometryCollectionProximityUtility::UpdateProximity(Coll);
+		FGeometryCollectionProximityUtility ProximityUtility(Coll);
+		ProximityUtility.UpdateProximity();
 
 		// Breaking Data
 		TManagedArray<TSet<int32>>& Proximity = Collection->GetAttribute<TSet<int32>>("Proximity", FGeometryCollection::GeometryGroup);
@@ -465,7 +469,8 @@ namespace GeometryCollectionTest
 		TArray<FTransform> GlobalTransform;
 		GeometryCollectionAlgo::GlobalMatrices(Collection->Transform, Collection->Parent, GlobalTransform);
 
-		FGeometryCollectionProximityUtility::UpdateProximity(Coll);
+		FGeometryCollectionProximityUtility ProximityUtility(Coll);
+		ProximityUtility.UpdateProximity();
 
 		// Breaking Data
 		TManagedArray<TSet<int32>>& Proximity = Collection->GetAttribute<TSet<int32>>("Proximity", FGeometryCollection::GeometryGroup);
@@ -578,7 +583,8 @@ namespace GeometryCollectionTest
 		TArray<FTransform> GlobalTransform;
 		GeometryCollectionAlgo::GlobalMatrices(Collection->Transform, Collection->Parent, GlobalTransform);
 
-		FGeometryCollectionProximityUtility::UpdateProximity(Coll);
+		FGeometryCollectionProximityUtility ProximityUtility(Coll);
+		ProximityUtility.UpdateProximity();
 
 		// Breaking Data
 		TManagedArray<TSet<int32>>& Proximity = Collection->GetAttribute<TSet<int32>>("Proximity", FGeometryCollection::GeometryGroup);
@@ -691,7 +697,8 @@ namespace GeometryCollectionTest
 		TArray<FTransform> GlobalTransform;
 		GeometryCollectionAlgo::GlobalMatrices(Collection->Transform, Collection->Parent, GlobalTransform);
 
-		FGeometryCollectionProximityUtility::UpdateProximity(Coll);
+		FGeometryCollectionProximityUtility ProximityUtility(Coll);
+		ProximityUtility.UpdateProximity();
 
 		// Breaking Data
 		TManagedArray<TSet<int32>>& Proximity = Collection->GetAttribute<TSet<int32>>("Proximity", FGeometryCollection::GeometryGroup);
@@ -797,7 +804,8 @@ namespace GeometryCollectionTest
 		TArray<FTransform> GlobalTransform;
 		GeometryCollectionAlgo::GlobalMatrices(Collection->Transform, Collection->Parent, GlobalTransform);
 
-		FGeometryCollectionProximityUtility::UpdateProximity(Coll);
+		FGeometryCollectionProximityUtility ProximityUtility(Coll);
+		ProximityUtility.UpdateProximity();
 
 		// Breaking Data
 		TManagedArray<TSet<int32>>& Proximity = Collection->GetAttribute<TSet<int32>>("Proximity", FGeometryCollection::GeometryGroup);
