Fix compiler error

diff --git a/Engine/Source/Programs/HeadlessChaos/Private/GeometryCollection/GeometryCollectionTestFramework.cpp b/Engine/Source/Programs/HeadlessChaos/Private/GeometryCollection/GeometryCollectionTestFramework.cpp
--- a/Engine/Source/Programs/HeadlessChaos/Private/GeometryCollection/GeometryCollectionTestFramework.cpp
+++ b/Engine/Source/Programs/HeadlessChaos/Private/GeometryCollection/GeometryCollectionTestFramework.cpp
@@ -242,7 +242,7 @@ namespace GeometryCollectionTest
 		TSharedPtr<Chaos::FChaosPhysicsMaterial> PhysicalMaterial = MakeShared<Chaos::FChaosPhysicsMaterial>(); InitMaterialToZero(PhysicalMaterial.Get());
 		auto Proxy = FSingleParticlePhysicsProxy::Create(Chaos::FGeometryParticle::CreateParticle());
 		auto& Particle = Proxy->GetGameThreadAPI();
-		Particle.SetGeometry(TUniquePtr<TPlane<FReal, 3>>(new TPlane<FReal, 3>(FVector(0), FVector(0, 0, 1))));
+		Particle.SetGeometry(TUniquePtr<Chaos::TPlane<FReal, 3>>(new Chaos::TPlane<FReal, 3>(FVector(0), FVector(0, 0, 1))));
 
 		FCollisionFilterData FilterData;
 		FilterData.Word1 = 0xFFFF;
