headlesschaos : Remove using namespace from GeometryCollectionTestFramework.cpp

diff --git a/Engine/Source/Programs/HeadlessChaos/Private/GeometryCollection/GeometryCollectionTestFramework.cpp b/Engine/Source/Programs/HeadlessChaos/Private/GeometryCollection/GeometryCollectionTestFramework.cpp
--- a/Engine/Source/Programs/HeadlessChaos/Private/GeometryCollection/GeometryCollectionTestFramework.cpp
+++ b/Engine/Source/Programs/HeadlessChaos/Private/GeometryCollection/GeometryCollectionTestFramework.cpp
@@ -15,17 +15,13 @@
 
 #include "../Resource/SphereGeometry.h"
 
-using namespace UE::Geometry;
-
-using namespace UE::Geometry;
-
 namespace GeometryCollectionTest
 {
 	TSharedPtr<FGeometryCollection> MakeSphereElement(
 		FTransform RootTranform, FTransform GeomTransform,
 		const int NumberOfMaterials = 2)
 	{
-		FSphereGenerator SphereGen;
+		UE::Geometry::FSphereGenerator SphereGen;
 		SphereGen.Radius = 1.f;
 		SphereGen.NumPhi = 16;		// Vertical divisions
 		SphereGen.NumTheta = 16;	// Horizontal divisions
@@ -89,6 +85,8 @@ namespace GeometryCollectionTest
 
 	TSharedPtr<FGeometryCollection> MakeCubeElement(FTransform RootTranform, FTransform GeomTransform)
 	{
+		using FVector3i = UE::Geometry::FVector3i;
+
 		const TArray<FVector3f> PointsIn = { FVector3f(-1,1,-1), FVector3f(1,1,-1), FVector3f(1,-1,-1), FVector3f(-1,-1,-1), FVector3f(-1,1,1), FVector3f(1,1,1), FVector3f(1,-1,1), FVector3f(-1,-1,1) };
 		const TArray<FVector3f> NormalsIn = { FVector3f(-1,1,-1).GetSafeNormal(), FVector3f(1,1,-1).GetSafeNormal(), FVector3f(1,-1,-1).GetSafeNormal(), FVector3f(-1,-1,-1).GetSafeNormal(), FVector3f(-1,1,1).GetSafeNormal(), FVector3f(1,1,1).GetSafeNormal(), FVector3f(1,-1,1).GetSafeNormal(), FVector3f(-1,-1,1).GetSafeNormal() };
 		const TArray<FVector3i> TrianglesIn = { FVector3i(0,1,2),FVector3i(0,2,3), FVector3i(2,1,6),FVector3i(1,5,6),FVector3i(2,6,7),FVector3i(3,2,7),FVector3i(4,7,3),FVector3i(4,0,3),FVector3i(4,1,0),FVector3i(4,5,1),FVector3i(5,4,7),FVector3i(5,7,6) };
@@ -98,6 +96,8 @@ namespace GeometryCollectionTest
 
 	TSharedPtr<FGeometryCollection> MakeTetrahedronElement(FTransform RootTranform, FTransform GeomTransform)
 	{
+		using FVector3i = UE::Geometry::FVector3i;
+
 		const TArray<FVector3f> PointsIn = { FVector3f(-1,1,-1), FVector3f(-1,-1,1), FVector3f(1,-1,-1), FVector3f(1,1,1) };
 		const TArray<FVector3f> NormalsIn = { FVector3f(-1,1,-1).GetSafeNormal(), FVector3f(-1,-1,1).GetSafeNormal(), FVector3f(1,-1,-1).GetSafeNormal(), FVector3f(1,1,1).GetSafeNormal() };
 		const TArray<FVector3i> TrianglesIn = { FVector3i(1,0,2), FVector3i(2,1,3), FVector3i(3,2,0), FVector3i(3,0,1) };
@@ -117,12 +117,12 @@ namespace GeometryCollectionTest
 	TSharedPtr<FGeometryCollection> MakeGriddedBoxElement(
 		FTransform RootTranform, FTransform GeomTransform,
 		const FVector& Extents = FVector(1, 1, 1),
-		const FIndex3i& EdgeVertices = FIndex3i(4, 4, 4),
+		const UE::Geometry::FIndex3i& EdgeVertices = UE::Geometry::FIndex3i(4, 4, 4),
 		const int NumberOfMaterials = 2)
 	{
 		using GeomVector = UE::Geometry::FVector3<double>;
-		FGridBoxMeshGenerator BoxGen;
-		BoxGen.Box = FOrientedBox3d(GeomVector::Zero(), GeomVector(Extents)); // box center, box dimensions
+		UE::Geometry::FGridBoxMeshGenerator BoxGen;
+		BoxGen.Box = UE::Geometry::FOrientedBox3d(GeomVector::Zero(), GeomVector(Extents)); // box center, box dimensions
 		BoxGen.EdgeVertices = EdgeVertices;
 		BoxGen.Generate();
 
