Fix headless chaos tests compiler after merge

diff --git a/Engine/Source/Programs/HeadlessChaos/Private/HeadlessChaosTestPhysicsScene.cpp b/Engine/Source/Programs/HeadlessChaos/Private/HeadlessChaosTestPhysicsScene.cpp
--- a/Engine/Source/Programs/HeadlessChaos/Private/HeadlessChaosTestPhysicsScene.cpp
+++ b/Engine/Source/Programs/HeadlessChaos/Private/HeadlessChaosTestPhysicsScene.cpp
@@ -491,7 +491,7 @@ namespace ChaosTest {
 			EXPECT_EQ(CollisionData.QueryData.Word0, 0); // ensure query filter is defaulted to 0 (no query allowed at all)
 
 			// Verify query is filtered out with default collision data on shape
-			bool bFiltered = PrePreFilterImp(QueryFilter, CollisionData.QueryData);
+			bool bFiltered = PrePreQueryFilterImp(QueryFilter, CollisionData.QueryData);
 			EXPECT_EQ(bFiltered, true);
 		}
 
@@ -517,7 +517,7 @@ namespace ChaosTest {
 			Shape->SetCollisionData(NewCollisionData);
 
 			// Filter with new data, ensuring we pass and are not filtered out.
-			bool bFiltered = PrePreFilterImp(QueryFilter, NewCollisionData.QueryData);
+			bool bFiltered = PrePreQueryFilterImp(QueryFilter, NewCollisionData.QueryData);
 			EXPECT_EQ(bFiltered, false);
 		}
 
diff --git a/Engine/Source/Programs/HeadlessChaos/Private/HeadlessChaosTestSolverProxies.cpp b/Engine/Source/Programs/HeadlessChaos/Private/HeadlessChaosTestSolverProxies.cpp
--- a/Engine/Source/Programs/HeadlessChaos/Private/HeadlessChaosTestSolverProxies.cpp
+++ b/Engine/Source/Programs/HeadlessChaos/Private/HeadlessChaosTestSolverProxies.cpp
@@ -126,7 +126,7 @@ namespace ChaosTest {
 		FChaosSolversModule* Module = FChaosSolversModule::GetModule();
 
 		// Make a solver
-		auto* Solver = Module->CreateSolver(nullptr);
+		auto* Solver = Module->CreateSolver(nullptr, -1);
 
 		// Make a particle
 		auto Proxy = FSingleParticlePhysicsProxy::Create(Chaos::FPBDRigidParticle::CreateParticle());
