HeadlessChaos: fix unit tests failing since adding preprefiltering for sim data
Tests are now properly marking the particle dirty to make sure the acceleration structure handlke is properly up to date with the sim data filters
(only an issue with test as they are using evolution directly instead of the GT level API )

diff --git a/Engine/Source/Programs/HeadlessChaos/Private/HeadlessChaosTestCCD.cpp b/Engine/Source/Programs/HeadlessChaos/Private/HeadlessChaosTestCCD.cpp
--- a/Engine/Source/Programs/HeadlessChaos/Private/HeadlessChaosTestCCD.cpp
+++ b/Engine/Source/Programs/HeadlessChaos/Private/HeadlessChaosTestCCD.cpp
@@ -279,6 +279,15 @@ namespace ChaosTest
 		Dynamic->SetCCDEnabled(true);
 		Dynamic->SetGravityEnabled(false);
 
+		// IMPORTANT : this is required to make sure the particles internal representation will reflect the sim data
+		Evolution.DirtyParticle(*ContainerFaces[0]);
+		Evolution.DirtyParticle(*ContainerFaces[1]);
+		Evolution.DirtyParticle(*ContainerFaces[2]);
+		Evolution.DirtyParticle(*ContainerFaces[3]);
+		Evolution.DirtyParticle(*ContainerFaces[4]);
+		Evolution.DirtyParticle(*ContainerFaces[5]);
+		Evolution.DirtyParticle(*Dynamic);
+
 		Dynamic->V() = InitialVelocity;
 		///////////////////////////////////
 		// Test 1: bouncing from two opposite walls
diff --git a/Engine/Source/Programs/HeadlessChaos/Private/HeadlessChaosTestSim.cpp b/Engine/Source/Programs/HeadlessChaos/Private/HeadlessChaosTestSim.cpp
--- a/Engine/Source/Programs/HeadlessChaos/Private/HeadlessChaosTestSim.cpp
+++ b/Engine/Source/Programs/HeadlessChaos/Private/HeadlessChaosTestSim.cpp
@@ -88,6 +88,9 @@ namespace ChaosTest {
 
 		::ChaosTest::SetParticleSimDataToCollide({ Static,Dynamic });
 
+		// IMPORTANT : this is required to make sure the particles internal representation will reflect the sim data
+		Evolution.DirtyParticle(*Static);
+		Evolution.DirtyParticle(*Dynamic);
 
 		const FReal Dt = 1 / 60.f;
 		for (int i = 0; i < 200; ++i)
@@ -121,6 +124,9 @@ namespace ChaosTest {
 
 		::ChaosTest::SetParticleSimDataToCollide({ Static,Dynamic });
 
+		// IMPORTANT : this is required to make sure the particles internal representation will reflect the sim data
+		Evolution.DirtyParticle(*Static);
+		Evolution.DirtyParticle(*Dynamic);
 
 		const FReal Dt = 1 / 60.f;
 		for (int i = 0; i < 100; ++i)
@@ -211,6 +217,11 @@ namespace ChaosTest {
 
 		::ChaosTest::SetParticleSimDataToCollide({ Static,Dynamic1,Dynamic2 });
 
+		// IMPORTANT : this is required to make sure the particles internal representation will reflect the sim data
+		Evolution.DirtyParticle(*Static);
+		Evolution.DirtyParticle(*Dynamic1);
+		Evolution.DirtyParticle(*Dynamic2);
+
 		bool Dynamic1WentToSleep = false;
 		bool Dynamic1HasWokeAgain = false;
 		for (int i = 0; i < 1000; ++i)
@@ -354,6 +365,10 @@ namespace ChaosTest {
 		Evolution.SetPhysicsMaterial(Dynamic1, MakeSerializable(PhysicsMaterial));
 		Evolution.SetPhysicsMaterial(Dynamic2, MakeSerializable(PhysicsMaterial));
 
+		// IMPORTANT : this is required to make sure the particles internal representation will reflect the sim data
+		Evolution.DirtyParticle(*Dynamic1);
+		Evolution.DirtyParticle(*Dynamic2);
+
 		const auto& SleepData = Evolution.GetParticles().GetDynamicParticles().GetSleepData();
 		EXPECT_EQ(SleepData.Num(), 0);
 
