Network Physics
-Prevent recursion into TriggerRewindIfNeeded_Internal on the last resim frame. We are not 100% sure if this is causing issues but the code was not intended to be recursive like this.
-Had to update headless tests which were relying on TriggerRewindIfNeeded to be called during resim

diff --git a/Engine/Plugins/Runtime/NetworkPrediction/Source/NetworkPrediction/Private/NetworkPhysics.cpp b/Engine/Plugins/Runtime/NetworkPrediction/Source/NetworkPrediction/Private/NetworkPhysics.cpp
--- a/Engine/Plugins/Runtime/NetworkPrediction/Source/NetworkPrediction/Private/NetworkPhysics.cpp
+++ b/Engine/Plugins/Runtime/NetworkPrediction/Source/NetworkPrediction/Private/NetworkPhysics.cpp
@@ -199,14 +199,11 @@ struct FNetworkPhysicsRewindCallback : public Chaos::IRewindCallback
 			return INDEX_NONE;
 		}
 
+		npCheckSlow(!Solver->GetEvolution()->IsResimming());
+		npCheckSlow(!RewindData->IsResim());
+		
 		TRACE_CPUPROFILER_EVENT_SCOPE(NPA_Physics_TriggerRewindIfNeeded_Internal);
 
-		if (ResimEndFrame != INDEX_NONE)
-		{
-			// We are already in a resim and shouldn't trigger a new one
-			return INDEX_NONE;
-		}
-
 		PendingCorrections.Reset();
 		const int32 EarliestFrame = RewindData->GetEarliestFrame_Internal();
 
@@ -214,7 +211,6 @@ struct FNetworkPhysicsRewindCallback : public Chaos::IRewindCallback
 		{
 			//const int32 ForceRewindFrame = RewindData->GetEarliestFrame_Internal()+1; // This was causing us to back > 64 frames?
 			const int32 ForceRewindFrame = LastCompletedStep - 9;
-			ResimEndFrame = LastCompletedStep;
 			//UE_LOG(LogTemp, Warning, TEXT("Forcing rewind to %d. LastCompletedStep: %d"), ForceRewindFrame, LastCompletedStep);
 			return ForceRewindFrame;
 		}
@@ -530,10 +526,6 @@ struct FNetworkPhysicsRewindCallback : public Chaos::IRewindCallback
 
 	void PostResimStep_Internal(int32 PhysicsStep) override
 	{
-		if (ResimEndFrame == PhysicsStep)
-		{
-			ResimEndFrame = INDEX_NONE;
-		}
 	}
 
 	void ProcessInputs_External(int32 PhysicsStep, const TArray<Chaos::FSimCallbackInputAndObject>& SimCallbackInputs) override 
@@ -568,7 +560,6 @@ struct FNetworkPhysicsRewindCallback : public Chaos::IRewindCallback
 	// Objects that we need to correct in PreResimStep_Internal
 	TArray<FNetworkPhysicsState> PendingCorrections;
 	int32 PendingCorrectionIdx = INDEX_NONE;
-	int32 ResimEndFrame = INDEX_NONE;
 
 	Chaos::FRewindData* RewindData=nullptr; // This be made to be accessed off of IRewindCallback
 	Chaos::FPhysicsSolver* Solver = nullptr; //todo: shouldn't have to cache this, should be part of base class
diff --git a/Engine/Source/Programs/HeadlessChaos/Private/HeadlessChaosTestRewind.cpp b/Engine/Source/Programs/HeadlessChaos/Private/HeadlessChaosTestRewind.cpp
--- a/Engine/Source/Programs/HeadlessChaos/Private/HeadlessChaosTestRewind.cpp
+++ b/Engine/Source/Programs/HeadlessChaos/Private/HeadlessChaosTestRewind.cpp
@@ -731,6 +731,7 @@ namespace ChaosTest {
 
 					virtual void ProcessInputs_Internal(int32 PhysicsStep, const TArray<FSimCallbackInputAndObject>& SimCallbackInputs) override
 					{
+						bool bIsResimming = RewindData->IsResim();
 						if(PhysicsStep == 2)
 						{
 							Proxy->GetPhysicsThreadAPI()->AddForce(FVec3(1, 0, 0));
@@ -760,24 +761,15 @@ namespace ChaosTest {
 
 					virtual int32 TriggerRewindIfNeeded_Internal(int32 LastCompletedStep) override
 					{
-						if (!bIsResimming && LastCompletedStep == ResimEndFrame)
-						{
-							bIsResimming = true;
-							return ResimStartFrame;
-						}
-
 						if (LastCompletedStep == ResimEndFrame)
 						{
-							bIsResimming = false;
+							return ResimStartFrame;
 						}
-
 						return INDEX_NONE;
 					}
 
 					int32 ResimStartFrame = 1;
 					int32 ResimEndFrame = 10;
-					bool bIsResimming = false;
-
 				};
 
 				const int32 LastGameStep = 32;
@@ -818,6 +810,7 @@ namespace ChaosTest {
 
 					virtual void ProcessInputs_Internal(int32 PhysicsStep, const TArray<FSimCallbackInputAndObject>& SimCallbackInputs) override
 					{
+						bool bIsResimming = RewindData->IsResim();
 						if (PhysicsStep == 2)
 						{
 							Proxy->GetPhysicsThreadAPI()->SetV(FVec3(0,0,2));
@@ -849,15 +842,9 @@ namespace ChaosTest {
 
 					virtual int32 TriggerRewindIfNeeded_Internal(int32 LastCompletedStep) override
 					{
-						if (!bIsResimming && LastCompletedStep == ResimEndFrame)
-						{
-							bIsResimming = true;
-							return ResimStartFrame;
-						}
-
 						if (LastCompletedStep == ResimEndFrame)
 						{
-							bIsResimming = false;
+							return ResimStartFrame;
 						}
 
 						return INDEX_NONE;
@@ -865,8 +852,6 @@ namespace ChaosTest {
 
 					int32 ResimStartFrame = 1;
 					int32 ResimEndFrame = 10;
-					bool bIsResimming = false;
-
 				};
 
 				const int32 LastGameStep = 32;
@@ -904,6 +889,7 @@ namespace ChaosTest {
 
 					virtual void ProcessInputs_Internal(int32 PhysicsStep, const TArray<FSimCallbackInputAndObject>& SimCallbackInputs) override
 					{
+						bool bIsResimming = RewindData->IsResim();
 						if (bIsResimming)
 						{
 							if (PhysicsStep == 5)
@@ -924,24 +910,15 @@ namespace ChaosTest {
 
 					virtual int32 TriggerRewindIfNeeded_Internal(int32 LastCompletedStep) override
 					{
-						if (!bIsResimming && LastCompletedStep == ResimEndFrame)
-						{
-							bIsResimming = true;
-							return ResimStartFrame;
-						}
-
 						if (LastCompletedStep == ResimEndFrame)
 						{
-							bIsResimming = false;
+							return ResimStartFrame;
 						}
-
 						return INDEX_NONE;
 					}
 
 					int32 ResimStartFrame = 1;
 					int32 ResimEndFrame = 10;
-					bool bIsResimming = false;
-
 				};
 
 				auto Sphere = TSharedPtr<FImplicitObject, ESPMode::ThreadSafe>(new TSphere<FReal, 3>(FVec3(0), 10));
@@ -986,6 +963,14 @@ namespace ChaosTest {
 
 					virtual void ProcessInputs_Internal(int32 PhysicsStep, const TArray<FSimCallbackInputAndObject>& SimCallbackInputs) override
 					{
+						bool bIsResimming = true;
+						if (PhysicsStep > HighestStep)
+						{
+							bIsResimming = false;
+							HighestStep = PhysicsStep;
+						}
+						ensure(bIsResimming == RewindData->IsResim()); // this would catch if IsResim is lieing to us
+
 						if (bIsResimming)
 						{
 							if (PhysicsStep == 5)
@@ -1006,15 +991,9 @@ namespace ChaosTest {
 
 					virtual int32 TriggerRewindIfNeeded_Internal(int32 LastCompletedStep) override
 					{
-						if (!bIsResimming && LastCompletedStep == ResimEndFrame)
-						{
-							bIsResimming = true;
-							return ResimStartFrame;
-						}
-
 						if (LastCompletedStep == ResimEndFrame)
 						{
-							bIsResimming = false;
+							return ResimStartFrame;
 						}
 
 						return INDEX_NONE;
@@ -1022,7 +1001,8 @@ namespace ChaosTest {
 
 					int32 ResimStartFrame = 1;
 					int32 ResimEndFrame = 10;
-					bool bIsResimming = false;
+
+					int32 HighestStep = INDEX_NONE;
 
 				};
 
diff --git a/Engine/Source/Runtime/Experimental/Chaos/Private/PBDRigidsSolver.cpp b/Engine/Source/Runtime/Experimental/Chaos/Private/PBDRigidsSolver.cpp
--- a/Engine/Source/Runtime/Experimental/Chaos/Private/PBDRigidsSolver.cpp
+++ b/Engine/Source/Runtime/Experimental/Chaos/Private/PBDRigidsSolver.cpp
@@ -1141,7 +1141,8 @@ namespace Chaos
 
 	void FPBDRigidsSolver::ConditionalApplyRewind_Internal()
 	{
-		if(!IsShuttingDown() && MRewindCallback && !MRewindData->IsResim())
+		// Note: checking MRewindData->IsResim() can lead to recursion into this function on the last resim frame since the call to AdvanceSolver is what advances RewindData's internal frame
+		if(!IsShuttingDown() && MRewindCallback && !GetEvolution()->IsResimming())
 		{
 			const int32 LastStep = MRewindData->CurrentFrame() - 1;
 			const int32 ResimStep = MRewindCallback->TriggerRewindIfNeeded_Internal(LastStep);
