Chaos: Restrict HeadlessChaos to desktop platforms for full solution builds.

diff --git a/Engine/Source/Programs/HeadlessChaos/HeadlessChaos.Target.cs b/Engine/Source/Programs/HeadlessChaos/HeadlessChaos.Target.cs
--- a/Engine/Source/Programs/HeadlessChaos/HeadlessChaos.Target.cs
+++ b/Engine/Source/Programs/HeadlessChaos/HeadlessChaos.Target.cs
@@ -3,7 +3,7 @@
 using UnrealBuildTool;
 using System.Collections.Generic;
 
-[SupportedPlatforms(UnrealPlatformClass.All)]
+[SupportedPlatforms(UnrealPlatformClass.Desktop)]
 public class HeadlessChaosTarget : TargetRules
 {
 	public HeadlessChaosTarget(TargetInfo Target) : base(Target)
