headlesschaos : Fix LWC runtime issues caused by the combination of the fact that we have unintialized managed array and we are still serializing in floats

diff --git a/Engine/Source/Programs/HeadlessChaos/Private/HeadlessChaosTestSerialization.cpp b/Engine/Source/Programs/HeadlessChaos/Private/HeadlessChaosTestSerialization.cpp
--- a/Engine/Source/Programs/HeadlessChaos/Private/HeadlessChaosTestSerialization.cpp
+++ b/Engine/Source/Programs/HeadlessChaos/Private/HeadlessChaosTestSerialization.cpp
@@ -255,6 +255,8 @@ namespace ChaosTest
 		{
 			FGeometryParticles OriginalParticles;
 			OriginalParticles.AddParticles(2);
+			OriginalParticles.R(0) = FRotation3::Identity;
+			OriginalParticles.R(1) = FRotation3::Identity;
 			OriginalParticles.SetGeometry(0, MakeSerializable(OriginalSpheres[0]));
 			OriginalParticles.SetGeometry(1, MakeSerializable(OriginalSpheres[1]));
 
@@ -289,6 +291,8 @@ namespace ChaosTest
 		{
 			auto OriginalParticles = MakeUnique<FGeometryParticles>();
 			OriginalParticles->AddParticles(2);
+			OriginalParticles->R(0) = FRotation3::Identity;
+			OriginalParticles->R(1) = FRotation3::Identity;
 			OriginalParticles->SetGeometry(0, MakeSerializable(OriginalSpheres[0]));
 			OriginalParticles->SetGeometry(1, MakeSerializable(OriginalSpheres[1]));
 
@@ -330,6 +334,8 @@ namespace ChaosTest
 
 			FGeometryParticles OriginalParticles;
 			OriginalParticles.AddParticles(2);
+			OriginalParticles.R(0) = FRotation3::Identity;
+			OriginalParticles.R(1) = FRotation3::Identity;
 			OriginalParticles.SetGeometry(0, MakeSerializable(OriginalSpheres[0]));
 			OriginalParticles.SetGeometry(1, MakeSerializable(OriginalSpheres[1]));
 			OriginalParticles.X(0) = FVec3(100, 1, 2);
@@ -388,6 +394,8 @@ namespace ChaosTest
 
 		TRigidParticles<FReal, 3> Particles;
 		Particles.AddParticles(2);
+		Particles.R(0) = FRotation3::Identity;
+		Particles.R(1) = FRotation3::Identity;
 		Particles.F(0) = F[0];
 		Particles.F(1) = F[1];
 		Particles.X(0) = X[0];
@@ -420,6 +428,8 @@ namespace ChaosTest
 
 		FGeometryParticles Particles;
 		Particles.AddParticles(3);
+		Particles.R(0) = FRotation3::Identity;
+		Particles.R(1) = FRotation3::Identity;
 		Particles.SetGeometry(0, MakeSerializable(Spheres[0]));
 		Particles.SetGeometry(1, MakeSerializable(Spheres[1]));
 		Particles.SetGeometry(2, MakeSerializable(Spheres[2]));
