Give an API for setting the global particle id.
Added unit test to show that order creation affects determinism.
Sort constraints by particle id

diff --git a/Engine/Source/Programs/HeadlessChaos/Private/HeadlessChaosTestRewind.cpp b/Engine/Source/Programs/HeadlessChaos/Private/HeadlessChaosTestRewind.cpp
--- a/Engine/Source/Programs/HeadlessChaos/Private/HeadlessChaosTestRewind.cpp
+++ b/Engine/Source/Programs/HeadlessChaos/Private/HeadlessChaosTestRewind.cpp
@@ -3730,7 +3730,7 @@ namespace ChaosTest {
 		}
 
 		static void ComputeMaxErrors(const FSimComparisonHelper& A, const FSimComparisonHelper& B, FReal& OutMaxLinearError,
-			FReal& OutMaxAngularError, int32 HistoryMultiple = 1)
+			FReal& OutMaxAngularError, int32 HistoryMultiple = 1, const TArray<int32>* BMapping = nullptr)
 		{
 			ensure(B.History.Num() == (A.History.Num() * HistoryMultiple));
 
@@ -3744,7 +3744,7 @@ namespace ChaosTest {
 				const FEntry& OtherEntry = B.History[OtherIdx];
 
 				FReal MaxLinearError, MaxAngularError;
-				FEntry::CompareEntry(Entry, OtherEntry, MaxLinearError, MaxAngularError);
+				FEntry::CompareEntry(Entry, OtherEntry, MaxLinearError, MaxAngularError, BMapping);
 
 				MaxLinearError2 = FMath::Max(MaxLinearError2, MaxLinearError * MaxLinearError);
 				MaxAngularError2 = FMath::Max(MaxAngularError2, MaxAngularError * MaxAngularError);
@@ -3760,23 +3760,28 @@ namespace ChaosTest {
 			TArray<FVec3> X;
 			TArray<FRotation3> R;
 
-			static void CompareEntry(const FEntry& A, const FEntry& B, FReal& OutMaxLinearError, FReal& OutMaxAngularError)
+			static void CompareEntry(const FEntry& A, const FEntry& B, FReal& OutMaxLinearError, FReal& OutMaxAngularError, const TArray<int32>* BMapping = nullptr)
 			{
 				FReal MaxLinearError2 = 0;
 				FReal MaxAngularError2 = 0;
+				
+				auto BMappingHelper = [BMapping](const int32 Idx)
+				{
+					return BMapping ? (*BMapping)[Idx] : Idx;
+				};
 
 				check(A.X.Num() == A.R.Num());
 				check(A.X.Num() == B.X.Num());
 				for (int32 Idx = 0; Idx < A.X.Num(); ++Idx)
 				{
-					const FReal LinearError2 = (A.X[Idx] - B.X[Idx]).SizeSquared();
+					const FReal LinearError2 = (A.X[Idx] - B.X[BMappingHelper(Idx)]).SizeSquared();
 					MaxLinearError2 = FMath::Max(LinearError2, MaxLinearError2);
 
 					//if exactly the same we want 0 for testing purposes, inverse does not get that so just skip it
-					if (B.R[Idx] != A.R[Idx])
+					if (B.R[BMappingHelper(Idx)] != A.R[Idx])
 					{
 						//For angular error we look at the rotation needed to go from B to A
-						const FRotation3 Delta = B.R[Idx] * A.R[Idx].Inverse();
+						const FRotation3 Delta = B.R[BMappingHelper(Idx)] * A.R[Idx].Inverse();
 
 						FVec3 Axis;
 						FReal Angle;
@@ -3795,7 +3800,7 @@ namespace ChaosTest {
 	};
 
 	template <typename InitLambda>
-	void RunHelper(FSimComparisonHelper& SimComparison, int32 NumSteps, FReal Dt, const InitLambda& InitFunc)
+	void RunHelper(FSimComparisonHelper& SimComparison, int32 NumSteps, FReal Dt, const InitLambda& InitFunc, const TArray<int32>* Mapping = nullptr)
 	{
 		FChaosSolversModule* Module = FChaosSolversModule::GetModule();
 
@@ -3803,7 +3808,7 @@ namespace ChaosTest {
 		auto* Solver = Module->CreateSolver(nullptr, /*AsyncDt=*/-1);
 		InitSolverSettings(Solver);
 
-		TArray<FPhysicsActorHandle> Storage = InitFunc(Solver);
+		TArray<FPhysicsActorHandle> Storage = InitFunc(Solver, Mapping);
 
 		for (int32 Step = 0; Step < NumSteps; ++Step)
 		{
@@ -3818,7 +3823,7 @@ namespace ChaosTest {
 	{
 		auto Box = TSharedPtr<FImplicitObject, ESPMode::ThreadSafe>(new TBox<FReal, 3>(FVec3(-10, -10, -10), FVec3(10, 10, 10)));
 
-		const auto InitLambda = [&Box](auto& Solver)
+		const auto InitLambda = [&Box](auto& Solver, auto)
 		{
 			TArray<FPhysicsActorHandle> Storage;
 			auto DynamicProxy = FSingleParticlePhysicsProxy::Create(Chaos::FPBDRigidParticle::CreateParticle());
@@ -3853,7 +3858,7 @@ namespace ChaosTest {
 		FVec3 StartPos(0);
 		FRotation3 StartRotation = FRotation3::FromIdentity();
 
-		const auto InitLambda = [&Box, &StartPos, &StartRotation](auto& Solver)
+		const auto InitLambda = [&Box, &StartPos, &StartRotation](auto& Solver, auto)
 		{
 			TArray<FPhysicsActorHandle> Storage;
 			auto DynamicProxy = FSingleParticlePhysicsProxy::Create(Chaos::FPBDRigidParticle::CreateParticle());
@@ -3903,7 +3908,7 @@ namespace ChaosTest {
 	{
 		auto Box = TSharedPtr<FImplicitObject, ESPMode::ThreadSafe>(new TBox<FReal, 3>(FVec3(-10, -10, -10), FVec3(10, 10, 10)));
 
-		const auto InitLambda = [&Box](auto& Solver)
+		const auto InitLambda = [&Box](auto& Solver, auto)
 		{
 			TArray<FPhysicsActorHandle> Storage;
 			auto DynamicProxy = FSingleParticlePhysicsProxy::Create(Chaos::FPBDRigidParticle::CreateParticle());
@@ -3938,7 +3943,7 @@ namespace ChaosTest {
 		auto Box = TSharedPtr<FImplicitObject, ESPMode::ThreadSafe>(new TBox<FReal, 3>(FVec3(-10, -10, -10), FVec3(10, 10, 10)));
 		const FReal Gravity = -980;
 
-		const auto InitLambda = [&Box, Gravity](auto& Solver)
+		const auto InitLambda = [&Box, Gravity](auto& Solver, auto)
 		{
 			TArray<FPhysicsActorHandle> Storage;
 			auto DynamicProxy = FSingleParticlePhysicsProxy::Create(Chaos::FPBDRigidParticle::CreateParticle());
@@ -3992,7 +3997,7 @@ namespace ChaosTest {
 	{
 		auto Sphere = TSharedPtr<FImplicitObject, ESPMode::ThreadSafe>(new TSphere<FReal, 3>(FVec3(0), 50));
 
-		const auto InitLambda = [&Sphere](auto& Solver)
+		const auto InitLambda = [&Sphere](auto& Solver, auto)
 		{
 			TArray<FPhysicsActorHandle> Storage;
 			auto DynamicProxy = FSingleParticlePhysicsProxy::Create(Chaos::FPBDRigidParticle::CreateParticle());
@@ -4039,7 +4044,7 @@ namespace ChaosTest {
 		auto SmallBox = TSharedPtr<FImplicitObject, ESPMode::ThreadSafe>(new TBox<FReal, 3>(FVec3(-50, -50, -50), FVec3(50, 50, 50)));
 		auto Box = TSharedPtr<FImplicitObject, ESPMode::ThreadSafe>(new TBox<FReal, 3>(FVec3(-1000, -1000, -1000), FVec3(1000, 1000, 0)));
 
-		const auto InitLambda = [&SmallBox, &Box](auto& Solver)
+		const auto InitLambda = [&SmallBox, &Box](auto& Solver, auto)
 		{
 			Solver->GetEvolution()->GetGravityForces().SetAcceleration(FVec3(0, 0, -980));
 			TArray<FPhysicsActorHandle> Storage;
@@ -4098,5 +4103,76 @@ namespace ChaosTest {
 
 		FSimComparisonHelper::ComputeMaxErrors(FirstRun, ThirdRun, MaxLinearError, MaxAngularError, 2);
 	}
+	
+	GTEST_TEST(AllTraits, DeterministicSim_DifferentCreationOrder)
+	{
+		auto SmallBox = TSharedPtr<FImplicitObject, ESPMode::ThreadSafe>(new TBox<FReal, 3>(FVec3(-50, -50, -50), FVec3(50, 50, 50)));
+		auto Box = TSharedPtr<FImplicitObject, ESPMode::ThreadSafe>(new TBox<FReal, 3>(FVec3(-1000, -1000, -1000), FVec3(1000, 1000, 0)));
+
+		const int32 NumParticles = 50;
+		const auto InitLambda = [&SmallBox, &Box, NumParticles](auto& Solver, const TArray<int32>* Mapping)
+		{
+			auto MappingHelper = [Mapping](const int32 Idx) { return Mapping ? (*Mapping)[Idx] : Idx; };
+			Solver->GetEvolution()->GetGravityForces().SetAcceleration(FVec3(0, 0, -980));
+			TArray<FPhysicsActorHandle> Storage;
+			for (int Idx = 0; Idx < NumParticles; ++Idx)
+			{
+				auto DynamicProxy = FSingleParticlePhysicsProxy::Create(Chaos::FPBDRigidParticle::CreateParticle());
+				auto& Dynamic = DynamicProxy->GetGameThreadAPI();
+
+				Dynamic.SetGeometry(SmallBox);
+				Solver->RegisterObject(DynamicProxy);
+				Dynamic.SetObjectState(EObjectStateType::Dynamic);
+				Dynamic.SetGravityEnabled(true);
+				Dynamic.SetParticleID(FParticleID{ MappingHelper(Idx), INDEX_NONE });
+				Dynamic.SetX(FVec3(0, 5 * MappingHelper(Idx), 100 * MappingHelper(Idx) + 50));	//slightly offset
+				Dynamic.SetI(FMatrix33(1000, 1000, 1000));
+				Dynamic.SetInvI(FMatrix33(1/1000.0, 1/1000.0, 1/1000.0));
+
+				Storage.Add(DynamicProxy);
+			}
+
+			auto KinematicProxy = FSingleParticlePhysicsProxy::Create(Chaos::FKinematicGeometryParticle::CreateParticle());
+			auto& Kinematic = KinematicProxy->GetGameThreadAPI();
+
+			Kinematic.SetGeometry(Box);
+			Solver->RegisterObject(KinematicProxy);
+			Kinematic.SetX(FVec3(0, 0, 0));
+
+			Storage.Add(KinematicProxy);
+
+			for (int i = 0; i < Storage.Num(); ++i)
+			{
+				for (int j = i + 1; j < Storage.Num(); ++j)
+				{
+					ChaosTest::SetParticleSimDataToCollide({ Storage[i]->GetParticle_LowLevel(),Storage[j]->GetParticle_LowLevel() });
+				}
+			}
+
+			return Storage;
+		};
+
+		const int32 NumSteps = 20;
+		FSimComparisonHelper FirstRun;
+
+		RunHelper(FirstRun, NumSteps, 1 / 30.f, InitLambda);
+
+		//tick twice as often
+
+		TArray<int32> Mapping;
+		for (int32 Idx = 0; Idx < NumParticles; ++Idx)
+		{
+			Mapping.Add(NumParticles - Idx - 1);
+		}
+
+		FSimComparisonHelper SecondRun;
+		RunHelper(SecondRun, NumSteps, 1 / 30.f, InitLambda, &Mapping);
+
+		//make sure deterministic
+		FReal MaxLinearError, MaxAngularError;
+		FSimComparisonHelper::ComputeMaxErrors(FirstRun, SecondRun, MaxLinearError, MaxAngularError, 1, &Mapping);
+		EXPECT_EQ(MaxLinearError, 0);
+		EXPECT_EQ(MaxAngularError, 0);
+	}
 
 }
diff --git a/Engine/Source/Runtime/Experimental/Chaos/Public/Chaos/Collision/CollisionReceiver.h b/Engine/Source/Runtime/Experimental/Chaos/Public/Chaos/Collision/CollisionReceiver.h
--- a/Engine/Source/Runtime/Experimental/Chaos/Public/Chaos/Collision/CollisionReceiver.h
+++ b/Engine/Source/Runtime/Experimental/Chaos/Public/Chaos/Collision/CollisionReceiver.h
@@ -77,30 +77,35 @@ namespace Chaos
 		 */
 		void ProcessCollisions()
 		{
-			FPBDCollisionConstraints::FConstraintAppendScope AppendHelper = CollisionConstraints.BeginAppendScope();
-
 			{
-				AppendHelper.ReserveSingle(NumSingle.load(std::memory_order_relaxed));
-				AppendHelper.ReserveSingleSwept(NumSingleSwept.load(std::memory_order_relaxed));
-			}
 
-			{
-				for(FPerParticleData& Data : ParticleCache)
+				FPBDCollisionConstraints::FConstraintAppendScope AppendHelper = CollisionConstraints.BeginAppendScope();
+
 				{
-					AppendHelper.Append(MoveTemp(Data.Single));
-					AppendHelper.Append(MoveTemp(Data.SingleSwept));
-					Data.Reset();
+					AppendHelper.ReserveSingle(NumSingle.load(std::memory_order_relaxed));
+					AppendHelper.ReserveSingleSwept(NumSingleSwept.load(std::memory_order_relaxed));
 				}
-			}
 
-			{
-				QUICK_SCOPE_CYCLE_COUNTER(QSTAT_ResetParticleCache);
-				ParticleCache.Reset();
+				{
+					for (FPerParticleData& Data : ParticleCache)
+					{
+						AppendHelper.Append(MoveTemp(Data.Single));
+						AppendHelper.Append(MoveTemp(Data.SingleSwept));
+						Data.Reset();
+					}
+				}
+
+				{
+					QUICK_SCOPE_CYCLE_COUNTER(QSTAT_ResetParticleCache);
+					ParticleCache.Reset();
 
-				NumSingle.store(0, std::memory_order_relaxed);
-				NumSingleSwept.store(0, std::memory_order_relaxed);
+					NumSingle.store(0, std::memory_order_relaxed);
+					NumSingleSwept.store(0, std::memory_order_relaxed);
+				}
 			}
 
+			CollisionConstraints.SortConstraints();
+
 			if(ResimCache)
 			{
 				ResimCache->SaveConstraints(CollisionConstraints.GetConstraintsArray());
diff --git a/Engine/Source/Runtime/Experimental/Chaos/Public/Chaos/GeometryParticles.h b/Engine/Source/Runtime/Experimental/Chaos/Public/Chaos/GeometryParticles.h
--- a/Engine/Source/Runtime/Experimental/Chaos/Public/Chaos/GeometryParticles.h
+++ b/Engine/Source/Runtime/Experimental/Chaos/Public/Chaos/GeometryParticles.h
@@ -254,34 +254,6 @@ namespace Chaos
 
 	void CHAOS_API UpdateShapesArrayFromGeometry(FShapesArray& ShapesArray, TSerializablePtr<FImplicitObject> Geometry, const FRigidTransform3& ActorTM, IPhysicsProxyBase* Proxy);
 
-
-	struct FParticleID
-	{
-		int32 GlobalID;	//Set by global ID system
-		int32 LocalID;		//Set by local client. This can only be used in cases where the LocalID will be set in the same way (for example we always spawn N client only particles)
-
-		bool operator<(const FParticleID& Other) const
-		{
-			if(GlobalID == Other.GlobalID)
-			{
-				return LocalID < Other.LocalID;
-			}
-			return GlobalID < Other.GlobalID;
-		}
-
-		bool operator==(const FParticleID& Other) const
-		{
-			return GlobalID == Other.GlobalID && LocalID == Other.LocalID;
-		}
-
-		FParticleID()
-		: GlobalID(INDEX_NONE)
-		, LocalID(INDEX_NONE)
-		{
-		}
-	};
-
-
 	FORCEINLINE uint32 GetTypeHash(const FParticleID& Unique)
 	{
 		return ::GetTypeHash(Unique.GlobalID);
diff --git a/Engine/Source/Runtime/Experimental/Chaos/Public/Chaos/PBDRigidsSOAs.h b/Engine/Source/Runtime/Experimental/Chaos/Public/Chaos/PBDRigidsSOAs.h
--- a/Engine/Source/Runtime/Experimental/Chaos/Public/Chaos/PBDRigidsSOAs.h
+++ b/Engine/Source/Runtime/Experimental/Chaos/Public/Chaos/PBDRigidsSOAs.h
@@ -850,7 +850,7 @@ private:
 			const int32 HandleIdx = Count + HandlesStartIdx;
 
 			TUniquePtr<TParticleHandleType> NewParticleHandle = TParticleHandleType::CreateParticleHandle(MakeSerializable(Particles), ParticleIdx, HandleIdx);
-			NewParticleHandle->ParticleID().LocalID = BiggestParticleID++;
+			NewParticleHandle->SetParticleID(FParticleID{ INDEX_NONE, BiggestParticleID++ });
 			ReturnHandles[Count] = NewParticleHandle.Get();
 			//If unique indices are null it means there is no GT particle that already registered an ID, so create one
 			if(ExistingIndices)
diff --git a/Engine/Source/Runtime/Experimental/Chaos/Public/Chaos/ParticleDirtyFlags.h b/Engine/Source/Runtime/Experimental/Chaos/Public/Chaos/ParticleDirtyFlags.h
--- a/Engine/Source/Runtime/Experimental/Chaos/Public/Chaos/ParticleDirtyFlags.h
+++ b/Engine/Source/Runtime/Experimental/Chaos/Public/Chaos/ParticleDirtyFlags.h
@@ -28,6 +28,26 @@ class FName;
 namespace Chaos
 {
 
+struct FParticleID
+{
+	int32 GlobalID = INDEX_NONE; //Set by global ID system
+	int32 LocalID = INDEX_NONE;	//Set by local client. This can only be used in cases where the LocalID will be set in the same way (for example we always spawn N client only particles)
+
+	bool operator<(const FParticleID& Other) const
+	{
+		if (GlobalID == Other.GlobalID)
+		{
+			return LocalID < Other.LocalID;
+		}
+		return GlobalID < Other.GlobalID;
+	}
+
+	bool operator==(const FParticleID& Other) const
+	{
+		return GlobalID == Other.GlobalID && LocalID == Other.LocalID;
+	}
+};
+
 using FKinematicTarget = TKinematicTarget<FReal, 3>;
 
 enum class EResimType: uint8;
@@ -482,6 +502,13 @@ public:
 		MResimType = InType;
 	}
 
+	void SetParticleID(const FParticleID& ParticleID)
+	{
+		MParticleID = ParticleID;
+	}
+
+	const FParticleID& ParticleID() const { return MParticleID; }
+
 	bool EnabledDuringResim() const { return MEnabledDuringResim; }
 	void SetEnabledDuringResim(bool bEnabledDuringResim) { MEnabledDuringResim = bEnabledDuringResim; }
 
@@ -493,6 +520,7 @@ private:
 	TSharedPtr<const FImplicitObject,ESPMode::ThreadSafe> MGeometry;
 	FUniqueIdx MUniqueIdx;
 	FSpatialAccelerationIdx MSpatialIdx;
+	FParticleID MParticleID;
 	EResimType MResimType;
 	bool MEnabledDuringResim;
 #if CHAOS_DEBUG_NAME
diff --git a/Engine/Source/Runtime/Experimental/Chaos/Public/Chaos/ParticleHandle.h b/Engine/Source/Runtime/Experimental/Chaos/Public/Chaos/ParticleHandle.h
--- a/Engine/Source/Runtime/Experimental/Chaos/Public/Chaos/ParticleHandle.h
+++ b/Engine/Source/Runtime/Experimental/Chaos/Public/Chaos/ParticleHandle.h
@@ -485,6 +485,7 @@ public:
 		SetUniqueIdx(InData.UniqueIdx());
 		SetSpatialIdx(InData.SpatialIdx());
 		SetResimType(InData.ResimType());
+		SetParticleID(InData.ParticleID());
 
 #if CHAOS_DEBUG_NAME
 		SetDebugName(InData.DebugName());
@@ -589,7 +590,21 @@ public:
 
 #if CHAOS_DETERMINISTIC
 	FParticleID ParticleID() const { return GeometryParticles->ParticleID(ParticleIdx); }
-	FParticleID& ParticleID() { return GeometryParticles->ParticleID(ParticleIdx); }
+	void SetParticleID(const FParticleID& ParticleID)
+	{
+		//When particles are created they are assigned a unique local index
+		//This index cannot be used for replicated particles (which use a global index)
+		//The global index (if used) is set on the game thread and comes over in the NonFrequent data
+		//However, it's possible that the particle id was never set (not used), so it will come over as -1,-1
+		//In this case we should continue to use the local index
+
+		//TODO: find a better way to deal with this, shouldn't be at such a low level API
+
+		if (ParticleID.GlobalID != INDEX_NONE || ParticleID.LocalID != INDEX_NONE)
+		{
+			GeometryParticles->ParticleID(ParticleIdx) = ParticleID;
+		}
+	}
 #endif
 
 	void MoveToSOA(TGeometryParticles<T, d>& ToSOA)
@@ -1760,6 +1775,12 @@ public:
 		MNonFrequentData.Modify(bInvalidate,MDirtyFlags,Proxy,[UniqueIdx](auto& Data){ Data.SetUniqueIdx(UniqueIdx);});
 	}
 
+	const FParticleID& ParticleID() const { return MNonFrequentData.Read().ParticleID(); }
+	void SetParticleID(const FParticleID& ParticleID, bool bInvalidate = true)
+	{
+		MNonFrequentData.Modify(bInvalidate, MDirtyFlags, Proxy, [ParticleID](auto& Data) { Data.SetParticleID(ParticleID); });
+	}
+
 	const TRotation<T, d>& R() const { return MXR.Read().R(); }
 	void SetR(const TRotation<T, d>& InR, bool bInvalidate = true);
 
diff --git a/Engine/Source/Runtime/Experimental/Chaos/Public/PhysicsProxy/SingleParticlePhysicsProxy.h b/Engine/Source/Runtime/Experimental/Chaos/Public/PhysicsProxy/SingleParticlePhysicsProxy.h
--- a/Engine/Source/Runtime/Experimental/Chaos/Public/PhysicsProxy/SingleParticlePhysicsProxy.h
+++ b/Engine/Source/Runtime/Experimental/Chaos/Public/PhysicsProxy/SingleParticlePhysicsProxy.h
@@ -862,6 +862,12 @@ public:
 	void SetShapeQueryCollisionEnabled(int32 InShapeIndex, bool bInEnabled) { VerifyContext(); GetParticle_LowLevel()->SetShapeQueryCollisionEnabled(InShapeIndex, bInEnabled); }
 	void SetShapeSimData(int32 InShapeIndex, const FCollisionFilterData& SimData) { VerifyContext(); GetParticle_LowLevel()->SetShapeSimData(InShapeIndex, SimData); }
 
+	void SetParticleID(const FParticleID& ParticleID)
+	{
+		VerifyContext();
+		GetParticle_LowLevel()->SetParticleID(ParticleID);
+	}
+
 	void SetX(const FVec3& InX, bool bInvalidate = true)
 	{
 		VerifyContext();
