HeadlessChaos fix binaries subfolder

diff --git a/Engine/Source/Programs/HeadlessChaos/HeadlessChaos.Build.cs b/Engine/Source/Programs/HeadlessChaos/HeadlessChaos.Build.cs
--- a/Engine/Source/Programs/HeadlessChaos/HeadlessChaos.Build.cs
+++ b/Engine/Source/Programs/HeadlessChaos/HeadlessChaos.Build.cs
@@ -56,6 +56,6 @@ public class HeadlessChaos : ModuleRules
 			"HeadlessChaos",
 			"Headless Chaos",
 			Target.LaunchModuleName,
-			Path.Combine("Engine", "Binaries", Target.Platform.ToString(), "NotForLicensees"));
+			Path.Combine("Engine", "Binaries"));
 	}
 }
diff --git a/Engine/Source/Programs/HeadlessChaos/HeadlessChaos.Target.cs b/Engine/Source/Programs/HeadlessChaos/HeadlessChaos.Target.cs
--- a/Engine/Source/Programs/HeadlessChaos/HeadlessChaos.Target.cs
+++ b/Engine/Source/Programs/HeadlessChaos/HeadlessChaos.Target.cs
@@ -11,8 +11,7 @@ public class HeadlessChaosTarget : TargetRules
 		Type = TargetType.Program;
 		LinkType = TargetLinkType.Monolithic;
 
-		ExeBinariesSubFolder = "NotForLicensees";
-		LaunchModuleName = "HeadlessChaos";
+		ExeBinariesSubFolder = LaunchModuleName = "HeadlessChaos";
 
 		bBuildDeveloperTools = false;
 
