Added callback to make it easy for NP to inject substep inputs

diff --git a/Engine/Source/Programs/HeadlessChaos/Private/HeadlessChaosTestRewind.cpp b/Engine/Source/Programs/HeadlessChaos/Private/HeadlessChaosTestRewind.cpp
--- a/Engine/Source/Programs/HeadlessChaos/Private/HeadlessChaosTestRewind.cpp
+++ b/Engine/Source/Programs/HeadlessChaos/Private/HeadlessChaosTestRewind.cpp
@@ -1765,6 +1765,51 @@ namespace ChaosTest {
 			});
 	}
 
+	struct FSimCallbackHelperInput2 : FSimCallbackInput
+	{
+		void Reset() { StepToCounter.Reset(); }
+		TMap<int32, int32> StepToCounter;
+	};
+
+	GTEST_TEST(AllTraits, RewindTest_SimCallbackProcessExternalInputs)
+	{
+		//If inputs are not set until external callback, make sure they are associated with the right frame
+		TRewindHelper::TestDynamicSphere([](auto* Solver, FReal SimDt, int32 Optimization, auto Proxy, auto Sphere)
+			{
+				struct FSimCallbackHelper : TSimCallbackObject<FSimCallbackHelperInput2>
+				{
+					FPBDRigidsSolver* Solver = nullptr;
+					virtual void OnPreSimulate_Internal() override
+					{
+						EXPECT_EQ(GetConsumerInput_Internal()->StepToCounter[Solver->GetCurrentFrame()], Solver->GetCurrentFrame());
+					}
+				};
+
+				struct FRewindCallback : public IRewindCallback
+				{
+					FSimCallbackHelper* SimCallback;
+					FRewindCallback(FSimCallbackHelper* Callback) : SimCallback(Callback){}
+					virtual void InjectInputs_External(int32 PhysicsStep, int32 NumSteps) override
+					{
+						for(int32 Idx = 0; Idx < NumSteps; ++Idx)
+						{
+							const int32 Step = Idx + PhysicsStep;
+							SimCallback->GetProducerInputData_External()->StepToCounter.Add(Step, Step);
+						}
+					}
+				};
+
+				FSimCallbackHelper* SimCallback = Solver->CreateAndRegisterSimCallbackObject_External<FSimCallbackHelper>();
+				SimCallback->Solver = Solver;
+				Solver->SetRewindCallback(MakeUnique<FRewindCallback>(SimCallback));
+
+				for (int Step = 0; Step < 32; ++Step)
+				{
+					TickSolverHelper(Solver);
+				}
+			});
+	}
+
 	GTEST_TEST(AllTraits, RewindTest_SimCallbackInputsCorrection)
 	{
 		TRewindHelper::TestDynamicSphere([](auto* Solver, FReal SimDt, int32 Optimization, auto Proxy, auto Sphere)
diff --git a/Engine/Source/Runtime/Experimental/Chaos/Private/PBDRigidsSolver.cpp b/Engine/Source/Runtime/Experimental/Chaos/Private/PBDRigidsSolver.cpp
--- a/Engine/Source/Runtime/Experimental/Chaos/Private/PBDRigidsSolver.cpp
+++ b/Engine/Source/Runtime/Experimental/Chaos/Private/PBDRigidsSolver.cpp
@@ -931,6 +931,11 @@ namespace Chaos
 
 		GetEvolution()->GetBroadPhase().GetIgnoreCollisionManager().PushProducerStorageData_External(MarshallingManager.GetExternalTimestamp_External());
 
+		if (MRewindCallback && !IsShuttingDown())
+		{
+			MRewindCallback->InjectInputs_External(MarshallingManager.GetInternalStep_External(), NumSteps);
+		}
+
 		MarshallingManager.Step_External(DeltaTime, NumSteps, GetSolverSubstep_External());
 	}
 
diff --git a/Engine/Source/Runtime/Experimental/Chaos/Public/Chaos/ChaosMarshallingManager.h b/Engine/Source/Runtime/Experimental/Chaos/Public/Chaos/ChaosMarshallingManager.h
--- a/Engine/Source/Runtime/Experimental/Chaos/Public/Chaos/ChaosMarshallingManager.h
+++ b/Engine/Source/Runtime/Experimental/Chaos/Public/Chaos/ChaosMarshallingManager.h
@@ -298,6 +298,9 @@ public:
 	/** Returns the amount of external time pushed so far. Any external commands or events should be associated with this time */
 	FReal GetExternalTime_External() const { return ExternalTime_External; }
 
+	/** Returns the internal step that the current PushData will be associated with once it is marshalled over*/
+	int32 GetInternalStep_External() const { return InternalStep_External; }
+
 	/** Used to delay marshalled data. This is mainly used for testing at the moment */
 	void SetTickDelay_External(int32 InDelay) { Delay = InDelay; }
 
diff --git a/Engine/Source/Runtime/Experimental/Chaos/Public/RewindData.h b/Engine/Source/Runtime/Experimental/Chaos/Public/RewindData.h
--- a/Engine/Source/Runtime/Experimental/Chaos/Public/RewindData.h
+++ b/Engine/Source/Runtime/Experimental/Chaos/Public/RewindData.h
@@ -1168,6 +1168,13 @@ public:
 	*/
 	virtual void ProcessInputs_External(int32 PhysicsStep, const TArray<FSimCallbackInputAndObject>& SimCallbackInputs) {}
 
+	/** Called before inputs are split into potential sub-steps and marshalled over to the physics thread.
+	*	The physics state has not been applied yet, and cannot be inspected anyway because this is triggered from the external thread (game thread)
+	*	Gives user the ability to call GetProducerInputData_External one last time.
+	*	Input data is shared amongst sub-steps. If NumSteps > 1 it means any input data injected will be shared for all sub-steps generated
+	*/
+	virtual void InjectInputs_External(int32 PhysicsStep, int32 NumSteps){}
+
 	/** Called after sim step to give the option to rewind. Any pending inputs for the next frame will remain in the queue
 	*   Return the PhysicsStep to start resimulating from. Resim will run up until latest step passed into RecordInputs (i.e. latest physics sim simulated so far)
 	*   Return INDEX_NONE to indicate no rewind
