Added an ensure when trying to access state from earlier than we rewound to. Added unit tests for sleep state changes

diff --git a/Engine/Source/Programs/HeadlessChaos/Private/HeadlessChaosTestRewind.cpp b/Engine/Source/Programs/HeadlessChaos/Private/HeadlessChaosTestRewind.cpp
--- a/Engine/Source/Programs/HeadlessChaos/Private/HeadlessChaosTestRewind.cpp
+++ b/Engine/Source/Programs/HeadlessChaos/Private/HeadlessChaosTestRewind.cpp
@@ -542,6 +542,227 @@ namespace ChaosTest {
 			});
 	}
 
+	GTEST_TEST(AllTraits, RewindTest_ResimSleepChangeRewind)
+	{
+		// Test puting object to sleep during Resim
+		TRewindHelper::TestDynamicSphere([](auto* Solver, FReal SimDt, int32 Optimization, auto Proxy, auto Sphere)
+			{
+				struct FRewindCallback : public IRewindCallback
+				{
+					FSingleParticlePhysicsProxy* Proxy;
+					FRewindData* RewindData;
+
+					virtual void ProcessInputs_Internal(int32 PhysicsStep, const TArray<FSimCallbackInputAndObject>& SimCallbackInputs) override
+					{
+						if (bIsResimming)
+						{
+							if (PhysicsStep >= ResimStartFrame+2)
+							{
+								Proxy->GetPhysicsThreadAPI()->SetObjectState(EObjectStateType::Sleeping, false, false);
+								ExpectedSleepFrame = ResimStartFrame+3;
+							}
+							else
+							{
+								Proxy->GetPhysicsThreadAPI()->SetObjectState(EObjectStateType::Dynamic, false, false);
+							}
+							return;
+						}
+
+						for (int32 Step = ResimStartFrame; Step < PhysicsStep; ++Step)
+						{
+							const EObjectStateType OldState = RewindData->GetPastStateAtFrame(*Proxy->GetHandle_LowLevel(), Step).ObjectState();
+
+							if (ExpectedSleepFrame != INDEX_NONE && Step >= ExpectedSleepFrame)
+							{
+								EXPECT_EQ(OldState, EObjectStateType::Sleeping);
+							}
+							else
+							{
+								EXPECT_EQ(OldState, EObjectStateType::Dynamic);								
+							}
+						}
+					}
+
+					virtual int32 TriggerRewindIfNeeded_Internal(int32 LastCompletedStep) override
+					{
+						if (!bIsResimming && LastCompletedStep == ResimEndFrame)
+						{
+							bIsResimming = true;
+							return ResimStartFrame;
+						}
+
+						return INDEX_NONE;
+					}
+
+					void PostResimStep_Internal(int32 PhysicsStep) override
+					{
+						if (ResimEndFrame == PhysicsStep)
+						{
+							bIsResimming = false;
+						}
+					}
+
+					
+					
+					bool bIsResimming = false;
+					int32 ResimStartFrame = 1;
+					int32 ResimEndFrame = 10;
+
+					int32 ExpectedSleepFrame = INDEX_NONE;
+				};
+
+				const int32 LastGameStep = 32;
+				const int32 NumPhysSteps = FMath::TruncToInt(LastGameStep / SimDt);
+
+				auto UniqueRewindCallback = MakeUnique<FRewindCallback>();
+				auto RewindCallback = UniqueRewindCallback.Get();
+				RewindCallback->Proxy = Proxy;
+				RewindCallback->RewindData = Solver->GetRewindData();
+
+				Solver->SetRewindCallback(MoveTemp(UniqueRewindCallback));
+
+				auto& Particle = Proxy->GetGameThreadAPI();
+				Particle.SetGravityEnabled(false);
+				Particle.SetV(FVec3(0, 0, 1));
+
+				for (int Step = 0; Step <= LastGameStep; ++Step)
+				{
+					TickSolverHelper(Solver);
+				}
+
+			});
+	}
+
+	GTEST_TEST(AllTraits, RewindTest_ResimSleepChangeRewind2)
+	{
+		// This does two ObjectState corrections
+		//	-Object starts out asleep
+		//	-On step 10 we force a rewind to frame 4 and wake it up during the resim (apply "correction")
+		//		-This on its own works fine
+		//	-On step 12 we force another rewind to frame 6 and put it to sleep during the resim
+		//		-This seems to have no effect: the object stays awake
+		//
+		//
+	
+		TRewindHelper::TestDynamicSphere([](auto* Solver, FReal SimDt, int32 Optimization, auto Proxy, auto Sphere)
+			{
+				struct FRewindCallback : public IRewindCallback
+				{
+					FSingleParticlePhysicsProxy* Proxy;
+					FRewindData* RewindData;
+
+					virtual void ProcessInputs_Internal(int32 PhysicsStep, const TArray<FSimCallbackInputAndObject>& SimCallbackInputs) override
+					{
+						if (bIsResimming)
+						{
+							if (ForceAwakeFrame != INDEX_NONE && PhysicsStep == ForceAwakeFrame)
+							{
+								Proxy->GetPhysicsThreadAPI()->SetObjectState(EObjectStateType::Dynamic);
+							}
+
+							if (ForceSleepFrame != INDEX_NONE && PhysicsStep == ForceSleepFrame)
+							{
+								Proxy->GetPhysicsThreadAPI()->SetObjectState(EObjectStateType::Sleeping);
+							}
+							
+							return;
+						}
+
+						for (int32 Step = ResimStartFrame; Step < PhysicsStep; ++Step)
+						{
+							const EObjectStateType OldState = RewindData->GetPastStateAtFrame(*Proxy->GetHandle_LowLevel(), Step).ObjectState();
+							if (TestAwakeFrame == INDEX_NONE || Step < TestAwakeFrame)
+							{
+								// If we haven't applied the first awake correction yet, or this is a step prior to that correction, we should be asleep
+								EXPECT_EQ(OldState, EObjectStateType::Sleeping);
+							}
+							else if (TestSleepFrame != INDEX_NONE && Step >= TestSleepFrame)
+							{
+								// If we *have* applied the second sleep correction and this is a step on or after that correction, we should be asleep
+								EXPECT_EQ(OldState, EObjectStateType::Sleeping);
+							}
+							else
+							{
+								// Everything else falls in the middle, we should be awake
+								EXPECT_EQ(OldState, EObjectStateType::Dynamic);
+							}
+						}
+					}
+
+					virtual int32 TriggerRewindIfNeeded_Internal(int32 LastCompletedStep) override
+					{
+						if (bIsResimming)
+						{
+							return INDEX_NONE;
+						}
+
+						if (LastCompletedStep == 10)
+						{
+							ForceAwakeFrame = 4;
+							bIsResimming = true;
+							ResimEndFrame = LastCompletedStep;
+							TestAwakeFrame = ForceAwakeFrame+1;
+							ResimStartFrame = ForceAwakeFrame;
+							return ResimStartFrame;
+						}
+
+						if (LastCompletedStep == 12)
+						{
+							ForceSleepFrame = 6;
+							bIsResimming = true;
+							ResimEndFrame = LastCompletedStep;
+							TestSleepFrame = ForceSleepFrame+1;
+							ResimStartFrame = ForceSleepFrame;
+							return ResimStartFrame;
+						}
+
+						return INDEX_NONE;
+					}
+
+					void PostResimStep_Internal(int32 PhysicsStep) override
+					{
+						if (ResimEndFrame == PhysicsStep)
+						{
+							ForceAwakeFrame = INDEX_NONE;
+							ForceSleepFrame = INDEX_NONE;
+							bIsResimming = false;
+						}
+					}
+					
+					int32 ForceAwakeFrame = INDEX_NONE;
+					int32 ForceSleepFrame = INDEX_NONE;
+
+					int32 TestAwakeFrame = INDEX_NONE;
+					int32 TestSleepFrame = INDEX_NONE;
+					int32 ResimStartFrame = 0;
+					
+					bool bIsResimming = false;
+					int32 ResimEndFrame = INDEX_NONE;
+				};
+
+				const int32 LastGameStep = 32;
+				const int32 NumPhysSteps = FMath::TruncToInt(LastGameStep / SimDt);
+
+				auto UniqueRewindCallback = MakeUnique<FRewindCallback>();
+				auto RewindCallback = UniqueRewindCallback.Get();
+				RewindCallback->Proxy = Proxy;
+				RewindCallback->RewindData = Solver->GetRewindData();
+
+				Solver->SetRewindCallback(MoveTemp(UniqueRewindCallback));
+
+				auto& Particle = Proxy->GetGameThreadAPI();
+				Particle.SetGravityEnabled(false);
+				Particle.SetV(FVec3(0, 0, 1));
+				Particle.SetObjectState(EObjectStateType::Sleeping);
+
+				for (int Step = 0; Step <= LastGameStep; ++Step)
+				{
+					TickSolverHelper(Solver);
+				}
+
+			});
+	}
+
 	GTEST_TEST(AllTraits, RewindTest_AddForce)
 	{
 		TRewindHelper::TestDynamicSphere([](auto* Solver, FReal SimDt, int32 Optimization, auto Proxy, auto Sphere)
diff --git a/Engine/Source/Runtime/Experimental/Chaos/Private/RewindData.cpp b/Engine/Source/Runtime/Experimental/Chaos/Private/RewindData.cpp
--- a/Engine/Source/Runtime/Experimental/Chaos/Private/RewindData.cpp
+++ b/Engine/Source/Runtime/Experimental/Chaos/Private/RewindData.cpp
@@ -590,7 +590,7 @@ bool FRewindData::RewindToFrame(int32 Frame)
 {
 	ensure(IsInPhysicsThreadContext());
 	//Can't go too far back
-	const int32 EarliestFrame = CurFrame - FramesSaved;
+	const int32 EarliestFrame = GetEarliestFrame_Internal();
 	if(Frame < EarliestFrame)
 	{
 		return false;
@@ -672,6 +672,8 @@ void FRewindData::RemoveParticle(const FUniqueIdx UniqueIdx)
 FGeometryParticleState FRewindData::GetPastStateAtFrame(const FGeometryParticleHandle& Handle,int32 Frame) const
 {
 	ensure(!IsResim());
+	ensure(Frame >= GetEarliestFrame_Internal());	//can't get state from before the frame we rewound to
+
 	const FDirtyParticleInfo* Info = FindParticle(Handle.UniqueIdx());
 	const FGeometryParticleStateBase* State = Info ? GetStateAtFrameImp(*Info, Frame) : nullptr;
 	return FGeometryParticleState(State,Handle, PropertiesPool);
diff --git a/Engine/Source/Runtime/Experimental/Chaos/Public/RewindData.h b/Engine/Source/Runtime/Experimental/Chaos/Public/RewindData.h
--- a/Engine/Source/Runtime/Experimental/Chaos/Public/RewindData.h
+++ b/Engine/Source/Runtime/Experimental/Chaos/Public/RewindData.h
@@ -756,6 +756,8 @@ public:
 
 	void CHAOS_API RemoveParticle(const FUniqueIdx UniqueIdx);
 
+	int32 CHAOS_API GetEarliestFrame_Internal() const { return CurFrame - FramesSaved; }
+
 	TArray<FDesyncedParticleInfo> CHAOS_API ComputeDesyncInfo() const;
 
 	/* Query the state of particles from the past. Once a rewind happens state captured must be queried using GetFutureStateAtFrame */
