Fix loss of momentum issue:
- Add ClearKinematicTarget function
- Do not invalidate X/R in SetKinematicTarget_AssumesLocked function
- Update HeadlessChaos tests

diff --git a/Engine/Source/Programs/HeadlessChaos/Private/HeadlessChaosTestPhysicsScene.cpp b/Engine/Source/Programs/HeadlessChaos/Private/HeadlessChaosTestPhysicsScene.cpp
--- a/Engine/Source/Programs/HeadlessChaos/Private/HeadlessChaosTestPhysicsScene.cpp
+++ b/Engine/Source/Programs/HeadlessChaos/Private/HeadlessChaosTestPhysicsScene.cpp
@@ -1011,6 +1011,142 @@ namespace ChaosTest {
 		EXPECT_NEAR(Particle.V()[2], ZVel, 1e-2);
 	}
 
+	GTEST_TEST(EngineInterface, SetKinematicTarget)
+	{
+		// Need to test:
+		// GT particle position is immediately updated after calling SetKinematicTarget_AssumesLocked
+		// GT particle positions and velocities are correctly updated
+		// PT particle positions and velocities are correctly updated
+		// Velocity becomes zero if no KinematicTarget is set in the current frame
+		// Particle positions and velocities are correct after SetKinematicTarget_AssumesLocked, SetKinematicTarget_AssumesLocked
+		// Velocity is zero if only SetGlobalPose_AssumesLocked is called (Teleport)
+		// Particle positions and velocities are correct after SetGlobalPose_AssumesLocked, SetKinematicTarget_AssumesLocked (Teleport)
+		// Particle positions and velocities are correct after SetKinematicTarget_AssumesLocked, SetGlobalPose_AssumesLocked (Teleport, KinematicTarget is cleared)
+		FChaosScene Scene(nullptr);
+		Scene.GetSolver()->SetThreadingMode_External(EThreadingModeTemp::SingleThread);
+
+		FActorCreationParams Params;
+		Params.Scene = &Scene;
+
+		FPhysicsActorHandle Proxy = nullptr;
+
+		FChaosEngineInterface::CreateActor(Params, Proxy);
+		auto& Particle = Proxy->GetGameThreadAPI();
+		{
+			auto Sphere = MakeUnique<TSphere<FReal, 3>>(FVec3(0), 3);
+			Particle.SetGeometry(MoveTemp(Sphere));
+		}
+
+		TArray<FPhysicsActorHandle> Proxys = { Proxy };
+		Scene.AddActorsToScene_AssumesLocked(Proxys);
+		Particle.SetObjectState(EObjectStateType::Kinematic);
+
+		struct FDummyInput : FSimCallbackInput
+		{
+			FSingleParticlePhysicsProxy* Proxy;
+			FVec3 CorrectX;
+			FVec3 CorrectV;
+			void Reset() {}
+		};
+
+		struct FCallback : public TSimCallbackObject<FDummyInput>
+		{
+			virtual void OnPreSimulate_Internal() override
+			{
+				auto Handle = GetConsumerInput_Internal()->Proxy->GetPhysicsThreadAPI();
+				EXPECT_EQ(Handle->X(), GetConsumerInput_Internal()->CorrectX);
+				EXPECT_EQ(Handle->V(), GetConsumerInput_Internal()->CorrectV);
+			}
+		};
+
+		auto Callback = Scene.GetSolver()->CreateAndRegisterSimCallbackObject_External<FCallback>();
+
+		Callback->GetProducerInputData_External()->Proxy = Proxy;
+
+		FVec3 Grav(0, 0, 0);
+		float Dt = 1;
+
+		auto AdvanceFrameAndRunTest = [&](const FVec3 &CorrectX, const FVec3 &CorrectV)
+		{
+			Scene.SetUpForFrame(&Grav, Dt, 99999, 99999, 10, false);
+			Scene.StartFrame();
+			Scene.EndFrame();
+			// Test X and V on GT
+			EXPECT_EQ(Particle.X(), CorrectX);
+			EXPECT_EQ(Particle.V(), CorrectV);
+			// Test X and V on PT, this is going to be used in OnPreSimulate_Internal in next frame.
+			Callback->GetProducerInputData_External()->CorrectX = CorrectX;
+			Callback->GetProducerInputData_External()->CorrectV = CorrectV;
+		};
+
+		// Set initial transform
+		FVec3 CurrentX = FVec3(1, 2, 3);
+		FVec3 CurrentV = FVec3(0, 0, 0);
+		FChaosEngineInterface::SetGlobalPose_AssumesLocked(Proxy, FTransform(CurrentX));
+
+		Callback->GetProducerInputData_External()->CorrectX = CurrentX;
+		Callback->GetProducerInputData_External()->CorrectV = CurrentV;
+		AdvanceFrameAndRunTest(CurrentX, CurrentV);
+
+		// Test SetKinematicTarget_AssumesLocked
+		CurrentX = FVec3(2, 3, 4);
+		CurrentV = FVec3(1, 1, 1);
+		FChaosEngineInterface::SetKinematicTarget_AssumesLocked(Proxy, FTransform(CurrentX));
+
+		// Test if position is immediately updated on GT after SetKinematicTarget_AssumesLocked
+		EXPECT_EQ(Particle.X(), CurrentX);
+
+		AdvanceFrameAndRunTest(CurrentX, CurrentV);
+
+		// Test if velocity becomes zero when no kinematic target is set
+		CurrentX = FVec3(2, 3, 4);
+		CurrentV = FVec3(0, 0, 0);
+
+		AdvanceFrameAndRunTest(CurrentX, CurrentV);
+
+		// Test if particle positions and velocities are correct after SetKinematicTarget_AssumesLocked, SetKinematicTarget_AssumesLocked
+		CurrentX = FVec3(0, 0, 0);
+		CurrentV = FVec3(-2, -3, -4);
+		FChaosEngineInterface::SetKinematicTarget_AssumesLocked(Proxy, FTransform(FVec3(1, 2, 3)));
+		FChaosEngineInterface::SetKinematicTarget_AssumesLocked(Proxy, FTransform(CurrentX));
+
+		AdvanceFrameAndRunTest(CurrentX, CurrentV);
+
+		// Test if velocity is zero if only SetGlobalPose_AssumesLocked is called (Teleport)
+		CurrentX = FVec3(0, 0, 0);
+		CurrentV = FVec3(0, 0, 0);
+		FChaosEngineInterface::SetGlobalPose_AssumesLocked(Proxy, FTransform(CurrentX));
+
+		Callback->GetProducerInputData_External()->CorrectX = CurrentX;
+		AdvanceFrameAndRunTest(CurrentX, CurrentV);
+
+		// Test if particle positions and velocities are correct after SetGlobalPose_AssumesLocked, SetKinematicTarget_AssumesLocked
+		CurrentX = FVec3(-1, -2, -3);
+		CurrentV = FVec3(0, 0, 0);
+		FChaosEngineInterface::SetGlobalPose_AssumesLocked(Proxy, FTransform(CurrentX));
+		FChaosEngineInterface::SetKinematicTarget_AssumesLocked(Proxy, FTransform(CurrentX));
+
+		Callback->GetProducerInputData_External()->CorrectX = CurrentX;
+		AdvanceFrameAndRunTest(CurrentX, CurrentV);
+
+		// Test if particle positions and velocities are correct after SetKinematicTarget_AssumesLocked, SetGlobalPose_AssumesLocked
+		CurrentX = FVec3(3, 2, 1);
+		CurrentV = FVec3(0, 0, 0);
+		FChaosEngineInterface::SetKinematicTarget_AssumesLocked(Proxy, FTransform(CurrentX));
+		FChaosEngineInterface::SetGlobalPose_AssumesLocked(Proxy, FTransform(CurrentX));
+
+		// Test if KinematicTarget is cleared
+		EXPECT_EQ(Particle.IsKinematicTargetDirty(), false);
+
+		Callback->GetProducerInputData_External()->CorrectX = CurrentX;
+		AdvanceFrameAndRunTest(CurrentX, CurrentV);
+
+		// Test if the PT positions and velocities are right from previous frame
+		CurrentX = FVec3(3, 2, 1);
+		CurrentV = FVec3(0, 0, 0);
+		AdvanceFrameAndRunTest(CurrentX, CurrentV);
+	}
+
 	GTEST_TEST(EngineInterface, PerPropertySetOnGT)
 	{
 		//Need to test:
diff --git a/Engine/Source/Runtime/Engine/Private/PhysicsEngine/BodyInstance.cpp b/Engine/Source/Runtime/Engine/Private/PhysicsEngine/BodyInstance.cpp
--- a/Engine/Source/Runtime/Engine/Private/PhysicsEngine/BodyInstance.cpp
+++ b/Engine/Source/Runtime/Engine/Private/PhysicsEngine/BodyInstance.cpp
@@ -2651,6 +2651,7 @@ void FBodyInstance::SetBodyTransform(const FTransform& NewTransform, ETeleportTy
 				}
 				else
 				{
+					// todo(chaos): Calling SetKinematicTarget_AssumesLocked before SetGlobalPose_AssumesLocked is unnessary for chaos. We should fix this when PhysX is removed.
 					if(bIsSimKinematic)
 					{
 						FPhysicsInterface::SetKinematicTarget_AssumesLocked(Actor, NewTransform);
diff --git a/Engine/Source/Runtime/Engine/Private/PhysicsEngine/PhysicalAnimationComponent.cpp b/Engine/Source/Runtime/Engine/Private/PhysicsEngine/PhysicalAnimationComponent.cpp
--- a/Engine/Source/Runtime/Engine/Private/PhysicsEngine/PhysicalAnimationComponent.cpp
+++ b/Engine/Source/Runtime/Engine/Private/PhysicsEngine/PhysicalAnimationComponent.cpp
@@ -311,6 +311,7 @@ void UPhysicalAnimationComponent::UpdateTargetActors(ETeleportType TeleportType)
 					if (BoneIdx != INDEX_NONE)	//It's possible the skeletal mesh has changed out from under us. In that case we should probably reset, but at the very least don't do work on non-existent bones
 					{
 						const FTransform TargetTM = ComputeTargetTM(PhysAnimData, *SkeletalMeshComponent, *PhysAsset, LocalTransforms, SpaceBases, BoneIdx);
+						// todo(chaos): Calling SetKinematicTarget_AssumesLocked before SetGlobalPose_AssumesLocked is unnessary for chaos. We should fix this when PhysX is removed.
 						FPhysicsInterface::SetKinematicTarget_AssumesLocked(TargetActor, TargetTM);
 
 						if (TeleportType == ETeleportType::TeleportPhysics)
diff --git a/Engine/Source/Runtime/Experimental/Chaos/Private/PhysicsProxy/SingleParticlePhysicsProxy.cpp b/Engine/Source/Runtime/Experimental/Chaos/Private/PhysicsProxy/SingleParticlePhysicsProxy.cpp
--- a/Engine/Source/Runtime/Experimental/Chaos/Private/PhysicsProxy/SingleParticlePhysicsProxy.cpp
+++ b/Engine/Source/Runtime/Experimental/Chaos/Private/PhysicsProxy/SingleParticlePhysicsProxy.cpp
@@ -67,13 +67,17 @@ void PushToPhysicsStateImp(const Chaos::FDirtyPropertiesManager& Manager, Chaos:
 			KinematicHandle->SetVelocities(*NewVelocities);
 		}
 
-		auto NewKinematicTarget = bHasKinematicData ? ParticleData.FindKinematicTarget(Manager, DataIdx) : nullptr;
-		if (NewKinematicTarget)
+		auto NewKinematicTargetGT = bHasKinematicData ? ParticleData.FindKinematicTarget(Manager, DataIdx) : nullptr;
+		if (NewKinematicTargetGT)
 		{
-			KinematicHandle->SetKinematicTarget(*NewKinematicTarget);
+			const TRigidTransform<float, 3>& TargetTransform = NewKinematicTargetGT->GetTarget();
+			const TRigidTransform<float, 3> PreviousTransform(KinematicHandle->X(), KinematicHandle->R());
+			TKinematicTarget<float, 3> NewKinematicTargetPT;
+			NewKinematicTargetPT.SetTargetMode(TargetTransform, PreviousTransform);
+			KinematicHandle->SetKinematicTarget(NewKinematicTargetPT);
 		}
 
-		if(NewXR || NewNonFrequentData || NewVelocities || NewKinematicTarget)
+		if(NewXR || NewNonFrequentData || NewVelocities || NewKinematicTargetGT)
 		{
 			const auto& Geometry = Handle->Geometry();
 			if(Geometry && Geometry->HasBoundingBox())
diff --git a/Engine/Source/Runtime/Experimental/Chaos/Public/Chaos/PBDRigidsEvolution.h b/Engine/Source/Runtime/Experimental/Chaos/Public/Chaos/PBDRigidsEvolution.h
--- a/Engine/Source/Runtime/Experimental/Chaos/Public/Chaos/PBDRigidsEvolution.h
+++ b/Engine/Source/Runtime/Experimental/Chaos/Public/Chaos/PBDRigidsEvolution.h
@@ -689,7 +689,7 @@ public:
 			}
 			}
 			
-			// Set previous and positions and velocities if we can
+			// Set positions and previous velocities if we can
 			// Note: At present kininematics are in fact rigid bodies
 			auto* Rigid = Particle.CastToRigidParticle();
 			if (Rigid)
diff --git a/Engine/Source/Runtime/Experimental/Chaos/Public/Chaos/ParticleHandle.h b/Engine/Source/Runtime/Experimental/Chaos/Public/Chaos/ParticleHandle.h
--- a/Engine/Source/Runtime/Experimental/Chaos/Public/Chaos/ParticleHandle.h
+++ b/Engine/Source/Runtime/Experimental/Chaos/Public/Chaos/ParticleHandle.h
@@ -2042,6 +2042,16 @@ public:
 		MKinematicTarget.Write(KinematicTarget, bInvalidate, MDirtyFlags, Proxy);
 	}
 
+	bool IsKinematicTargetDirty() const
+	{
+		return MKinematicTarget.IsDirty(MDirtyFlags);
+	}
+
+	void ClearKinematicTarget()
+	{
+		MKinematicTarget.Clear(MDirtyFlags, Proxy);
+	}
+
 	void SetVelocities(const FParticleVelocities& InVelocities,bool bInvalidate = true)
 	{
 		MVelocities.Write(InVelocities,bInvalidate,MDirtyFlags,Proxy);
diff --git a/Engine/Source/Runtime/Experimental/Chaos/Public/Chaos/Properties.h b/Engine/Source/Runtime/Experimental/Chaos/Public/Chaos/Properties.h
--- a/Engine/Source/Runtime/Experimental/Chaos/Public/Chaos/Properties.h
+++ b/Engine/Source/Runtime/Experimental/Chaos/Public/Chaos/Properties.h
@@ -41,6 +41,18 @@ public:
 		MarkDirty(bInvalidate, Dirty, Proxy);
 	}
 
+	void Clear(FParticleDirtyFlags& Dirty, IPhysicsProxyBase* Proxy)
+	{
+		Dirty.MarkClean(PropertyFlag);
+		if (Proxy && Dirty.IsClean())
+		{
+			if (FPhysicsSolverBase* PhysicsSolverBase = Proxy->GetSolver<FPhysicsSolverBase>())
+			{
+				PhysicsSolverBase->RemoveDirtyProxy(Proxy);
+			}
+		}
+	}
+
 	void SyncRemote(FDirtyPropertiesManager& Manager, int32 DataIdx, const FParticleDirtyData& Remote) const
 	{
 		Remote.SyncRemote<T,PropName>(Manager, DataIdx, Property);
diff --git a/Engine/Source/Runtime/Experimental/Chaos/Public/PhysicsProxy/SingleParticlePhysicsProxy.h b/Engine/Source/Runtime/Experimental/Chaos/Public/PhysicsProxy/SingleParticlePhysicsProxy.h
--- a/Engine/Source/Runtime/Experimental/Chaos/Public/PhysicsProxy/SingleParticlePhysicsProxy.h
+++ b/Engine/Source/Runtime/Experimental/Chaos/Public/PhysicsProxy/SingleParticlePhysicsProxy.h
@@ -1029,6 +1029,25 @@ public:
 
 	void MergeGeometry(TArray<TUniquePtr<FImplicitObject>>&& Objects) { VerifyContext(); GetParticle_LowLevel()->MergeGeometry(MoveTemp(Objects)); }
 
+	bool IsKinematicTargetDirty() const
+	{
+		VerifyContext();
+		if (auto Kinematic = GetParticle_LowLevel()->CastToKinematicParticle())
+		{
+			return Kinematic->IsKinematicTargetDirty();
+		}
+		return false;
+	}
+
+	void ClearKinematicTarget()
+	{
+		VerifyContext();
+		if (auto Kinematic = GetParticle_LowLevel()->CastToKinematicParticle())
+		{
+			return Kinematic->ClearKinematicTarget();
+		}
+	}
+
 	void SetCCDEnabled(bool bEnabled)
 	{
 		VerifyContext();
diff --git a/Engine/Source/Runtime/PhysicsCore/Private/ChaosEngineInterface.cpp b/Engine/Source/Runtime/PhysicsCore/Private/ChaosEngineInterface.cpp
--- a/Engine/Source/Runtime/PhysicsCore/Private/ChaosEngineInterface.cpp
+++ b/Engine/Source/Runtime/PhysicsCore/Private/ChaosEngineInterface.cpp
@@ -1647,6 +1647,10 @@ FChaosScene* FChaosEngineInterface::GetCurrentScene(const FPhysicsActorHandle& I
 void FChaosEngineInterface::SetGlobalPose_AssumesLocked(const FPhysicsActorHandle& InActorReference,const FTransform& InNewPose,bool bAutoWake)
 {
 	Chaos::FRigidBodyHandle_External& Body_External = InActorReference->GetGameThreadAPI();
+	if (Body_External.IsKinematicTargetDirty())
+	{
+		Body_External.ClearKinematicTarget();
+	}
 	Body_External.SetX(InNewPose.GetLocation());
 	Body_External.SetR(InNewPose.GetRotation());
 	Body_External.UpdateShapeBounds();
@@ -1658,13 +1662,16 @@ void FChaosEngineInterface::SetGlobalPose_AssumesLocked(const FPhysicsActorHandl
 void FChaosEngineInterface::SetKinematicTarget_AssumesLocked(const FPhysicsActorHandle& InActorReference,const FTransform& InNewTarget)
 {
 	{
-		Chaos::TKinematicTarget<Chaos::FReal, 3> newKinematicTarget;
-		Chaos::TRigidTransform<Chaos::FReal, 3> PreviousTM(InActorReference->GetGameThreadAPI().X(), InActorReference->GetGameThreadAPI().R());
-		newKinematicTarget.SetTargetMode(InNewTarget, PreviousTM);
-		InActorReference->GetGameThreadAPI().SetKinematicTarget(newKinematicTarget);
+		Chaos::TKinematicTarget <Chaos::FReal, 3 > NewKinematicTarget;
+		// SetKinematicTarget_AssumesLocked could be called multiple times in one time step
+		// Don't update Previous here. Previous is updated in FSingleParticlePhysicsProxy::PushToPhysicsState. Previous is not used in game thread and set to identity.
+		// @todo(chaos): create a new KinematicTarget class for game thread that does not have Preivous 
+		Chaos::FRigidTransform3 IdentityTransform;
+		NewKinematicTarget.SetTargetMode(InNewTarget, IdentityTransform);
+		InActorReference->GetGameThreadAPI().SetKinematicTarget(NewKinematicTarget);
 
-		InActorReference->GetGameThreadAPI().SetX(InNewTarget.GetLocation());
-		InActorReference->GetGameThreadAPI().SetR(InNewTarget.GetRotation());
+		InActorReference->GetGameThreadAPI().SetX(InNewTarget.GetLocation(), false);
+		InActorReference->GetGameThreadAPI().SetR(InNewTarget.GetRotation(), false);
 		InActorReference->GetGameThreadAPI().UpdateShapeBounds();
 
 		FChaosScene* Scene = GetCurrentScene(InActorReference);
