Make sure particles spawned by server before local version are used during resim

diff --git a/Engine/Source/Programs/HeadlessChaos/Private/HeadlessChaosTestRewind.cpp b/Engine/Source/Programs/HeadlessChaos/Private/HeadlessChaosTestRewind.cpp
--- a/Engine/Source/Programs/HeadlessChaos/Private/HeadlessChaosTestRewind.cpp
+++ b/Engine/Source/Programs/HeadlessChaos/Private/HeadlessChaosTestRewind.cpp
@@ -542,6 +542,96 @@ namespace ChaosTest {
 			});
 	}
 
+	GTEST_TEST(AllTraits, RewindTest_SpawnEarlierCorrection)
+	{
+		// Test resim when object spawned earlier as part of correction
+		TRewindHelper::TestDynamicSphere([](auto* Solver, FReal SimDt, int32 Optimization, auto Proxy, auto Sphere)
+			{
+				struct FRewindCallback : public IRewindCallback
+				{
+					FSingleParticlePhysicsProxy* Proxy;
+					FSingleParticlePhysicsProxy* Floor = nullptr;
+					FRewindData* RewindData;
+					FReal SimDt;
+
+					virtual void ProcessInputs_Internal(int32 PhysicsStep, const TArray<FSimCallbackInputAndObject>& SimCallbackInputs) override
+					{
+						const FReal Time = PhysicsStep * SimDt;
+
+						if (bIsResimming)
+						{
+							if (PhysicsStep == 1)
+							{
+								RewindData->SpawnProxyIfNeeded(*Floor);
+							}
+						}
+
+
+						if (!bIsResimming || Time < 4.5)
+						{
+							//simply movement without hitting floor because it's spawned too late
+							EXPECT_NEAR(Proxy->GetPhysicsThreadAPI()->X()[2], 14.5 - Time, 1e-2);
+						}
+						else
+						{
+							//floor spawned earlier so we hit it
+							EXPECT_GE(Proxy->GetPhysicsThreadAPI()->X()[2], 10 - 1e-2);
+						}
+					}
+
+					virtual int32 TriggerRewindIfNeeded_Internal(int32 LastCompletedStep) override
+					{
+						if (!bIsResimming && Floor)
+						{
+							bIsResimming = true;
+							return 0;
+						}
+
+						return INDEX_NONE;
+					}
+
+					bool bIsResimming = false;
+				};
+
+				const int32 LastGameStep = 32;
+				const int32 NumPhysSteps = FMath::TruncToInt(LastGameStep / SimDt);
+
+				auto UniqueRewindCallback = MakeUnique<FRewindCallback>();
+				auto RewindCallback = UniqueRewindCallback.Get();
+				RewindCallback->Proxy = Proxy;
+				RewindCallback->RewindData = Solver->GetRewindData();
+				RewindCallback->SimDt = SimDt;
+
+				auto FloorGeom = TSharedPtr<FImplicitObject, ESPMode::ThreadSafe>(new TBox<FReal,3>(FVec3(-100, -100, -1), FVec3(100, 100, 0)));
+
+				Solver->SetRewindCallback(MoveTemp(UniqueRewindCallback));
+
+				auto& Particle = Proxy->GetGameThreadAPI();
+				Particle.SetGravityEnabled(false);
+				Particle.SetV(FVec3(0, 0, -1));
+				Particle.SetX(FVec3(0, 0, 14.5));
+
+				ChaosTest::SetParticleSimDataToCollide({ Proxy->GetParticle_LowLevel() });
+
+				for (int Step = 0; Step <= LastGameStep; ++Step)
+				{
+					TickSolverHelper(Solver);
+
+					//spawn floor way late to ensure no collision on first run
+					if (Step == 12)
+					{
+						auto Floor = FSingleParticlePhysicsProxy::Create(Chaos::FGeometryParticle::CreateParticle());
+						auto& FloorParticle = Floor->GetGameThreadAPI();
+						FloorParticle.SetGeometry(FloorGeom);
+						Solver->RegisterObject(Floor);
+						ChaosTest::SetParticleSimDataToCollide({ Floor->GetParticle_LowLevel() });
+						RewindCallback->Floor = Floor;
+					}
+				}
+
+			});
+	}
+
 	GTEST_TEST(AllTraits, RewindTest_ResimSleepChangeRewind)
 	{
 		// Test puting object to sleep during Resim
diff --git a/Engine/Source/Runtime/Experimental/Chaos/Private/RewindData.cpp b/Engine/Source/Runtime/Experimental/Chaos/Private/RewindData.cpp
--- a/Engine/Source/Runtime/Experimental/Chaos/Private/RewindData.cpp
+++ b/Engine/Source/Runtime/Experimental/Chaos/Private/RewindData.cpp
@@ -926,6 +926,18 @@ void FRewindData::PushGTDirtyData(const FDirtyPropertiesManager& SrcManager,cons
 	}
 }
 
+void FRewindData::SpawnProxyIfNeeded(FSingleParticlePhysicsProxy& Proxy)
+{
+	if(Proxy.GetInitializedStep() > CurFrame)
+	{
+		FGeometryParticleHandle* Handle = Proxy.GetHandle_LowLevel();
+		FindOrAddParticle(*Handle, CurFrame);
+
+		Solver->GetEvolution()->EnableParticle(Handle, nullptr);
+		Proxy.SetInitialized(CurFrame);
+	}
+}
+
 void FRewindData::MarkDirtyFromPT(FGeometryParticleHandle& Handle)
 {
 	FDirtyParticleInfo& Info = FindOrAddParticle(Handle);
@@ -1023,7 +1035,6 @@ FRewindData::FDirtyParticleInfo& FRewindData::FindOrAddParticle(TGeometryParticl
 	ParticleToAllDirtyIdx.Add(UniqueIdx,DirtyIdx);
 	if(InitializedOnFrame != INDEX_NONE)
 	{
-		ensure(AllDirtyParticles[DirtyIdx].InitializedOnStep == INDEX_NONE);	//initializing now, shouldn't have this marked already
 		AllDirtyParticles[DirtyIdx].InitializedOnStep = InitializedOnFrame;
 	}
 
diff --git a/Engine/Source/Runtime/Experimental/Chaos/Public/RewindData.h b/Engine/Source/Runtime/Experimental/Chaos/Public/RewindData.h
--- a/Engine/Source/Runtime/Experimental/Chaos/Public/RewindData.h
+++ b/Engine/Source/Runtime/Experimental/Chaos/Public/RewindData.h
@@ -833,6 +833,8 @@ public:
 
 	void CHAOS_API MarkDirtyFromPT(FGeometryParticleHandle& Handle);
 
+	void CHAOS_API SpawnProxyIfNeeded(FSingleParticlePhysicsProxy& Proxy);
+
 private:
 
 	struct FSimWritableState
