Chaos: Fix for overlaps requiring that the start location is the centre of the query geometry when performing overlaps. Creating an offset collision primitive leads to missed hits if StartTM falls outside of the query geometry.

Added a unit test to HeadlessChaos that reproduces the case that identified this issue of an offset sphere contacting a box but the particle origin lies outside of the sphere.

diff --git a/Engine/Source/Programs/HeadlessChaos/Private/HeadlessChaosTestPhysicsScene.cpp b/Engine/Source/Programs/HeadlessChaos/Private/HeadlessChaosTestPhysicsScene.cpp
--- a/Engine/Source/Programs/HeadlessChaos/Private/HeadlessChaosTestPhysicsScene.cpp
+++ b/Engine/Source/Programs/HeadlessChaos/Private/HeadlessChaosTestPhysicsScene.cpp
@@ -1495,4 +1495,51 @@ namespace ChaosTest {
 		Scene.StartFrame();
 		Scene.EndFrame();
 	}
+	
+	GTEST_TEST(EngineInterface, OverlapOffsetActor)
+	{
+		FChaosScene Scene(nullptr);
+
+		FActorCreationParams Params;
+		Params.Scene = &Scene;
+		Params.bSimulatePhysics = false;
+		Params.bStatic = true;
+		Params.InitialTM = FTransform::Identity;
+		Params.Scene = &Scene;
+
+		TGeometryParticle<FReal, 3>* StaticCube = nullptr;
+
+		FChaosEngineInterface::CreateActor(Params, StaticCube);
+		ASSERT_NE(StaticCube, nullptr);
+
+		// Add geometry, placing a box at the origin
+		constexpr FReal BoxSize = static_cast<FReal>(50.0);
+		const FVec3 HalfBoxExtent{ BoxSize };
+
+		// We require a union here, although the second geometry isn't used we need the particle to
+		// have more than one shape in its shapes array otherwise the query acceleration will treat
+		// it as a special case and skip bounds checking during the overlap
+		TArray<TUniquePtr<FImplicitObject>> Geoms;
+		Geoms.Emplace(MakeUnique<TBox<FReal, 3>>(-HalfBoxExtent, HalfBoxExtent));
+		Geoms.Emplace(MakeUnique<TBox<FReal, 3>>(-HalfBoxExtent, HalfBoxExtent));
+
+		TUniquePtr<FImplicitObjectUnion> GeomUnion = MakeUnique<FImplicitObjectUnion>(MoveTemp(Geoms));
+		StaticCube->SetGeometry(MoveTemp(GeomUnion));
+
+		TArray<TGeometryParticle<FReal, 3>*> Particles{ StaticCube };
+		Scene.AddActorsToScene_AssumesLocked(Particles);
+
+		FChaosSQAccelerator SQ{ *Scene.GetSpacialAcceleration() };
+		FSQHitBuffer<ChaosInterface::FOverlapHit> HitBuffer;
+		FOverlapAllQueryCallback QueryCallback;
+
+		// Here we query from a position under the box, but using a shape that has an offset. This tests
+		// a failure case that was previously present where the query system assumed that the QueryTM
+		// was inside the geometry being used to query.
+		const FTransform QueryTM{ FVec3{0.0f, 0.0f, -110.0f} };
+		constexpr FReal SphereRadius = static_cast<FReal>(50.0);
+		SQ.Overlap(TSphere<FReal, 3>(FVec3(0.0f, 0.0f, 100.0f), SphereRadius), QueryTM, HitBuffer, FChaosQueryFilterData(), QueryCallback, FQueryDebugParams());
+
+		EXPECT_TRUE(HitBuffer.HasBlockingHit());
+	}
 }
diff --git a/Engine/Source/Runtime/PhysicsCore/Private/SQAccelerator.cpp b/Engine/Source/Runtime/PhysicsCore/Private/SQAccelerator.cpp
--- a/Engine/Source/Runtime/PhysicsCore/Private/SQAccelerator.cpp
+++ b/Engine/Source/Runtime/PhysicsCore/Private/SQAccelerator.cpp
@@ -213,6 +213,7 @@ private:
 		bool bContinue = true;
 
 		const FRigidTransform3 ActorTM(GeometryParticle->X(), GeometryParticle->R());
+		const TAABB<FReal, 3> QueryGeomWorldBounds = QueryGeom ? QueryGeom->BoundingBox().TransformedAABB(StartTM) : TAABB<FReal, 3>(-HalfExtents, HalfExtents);
 
 #if CHAOS_DEBUG_DRAW
 		bool bAllShapesIgnoredInPrefilter = true;
@@ -233,7 +234,7 @@ private:
 				else
 				{
 					// Transform to world bounds and get the proper half extent.
-					const FVec3 WorldHalfExtent = QueryGeom ? QueryGeom->BoundingBox().TransformedAABB(StartTM).Extents() * 0.5f : FVec3(HalfExtents);
+					const FVec3 WorldHalfExtent = QueryGeom ? QueryGeomWorldBounds.Extents() * 0.5f : FVec3(HalfExtents);
 
 					InflatedWorldBounds = FAABB3(Shape->GetWorldSpaceInflatedShapeBounds().Min() - WorldHalfExtent, Shape->GetWorldSpaceInflatedShapeBounds().Max() + WorldHalfExtent);
 				}
@@ -249,7 +250,8 @@ private:
 				}
 				else
 				{
-					if (!InflatedWorldBounds.Contains(StartTM.GetLocation()))
+					const FVec3 QueryCenter = QueryGeom ? QueryGeomWorldBounds.Center() : StartTM.GetLocation();
+					if (!InflatedWorldBounds.Contains(QueryCenter))
 					{
 						continue;
 					}
@@ -411,7 +413,7 @@ private:
 	ChaosInterface::FSQHitBuffer<THitType>& HitBuffer;
 	const FQueryFilterData& QueryFilterData;
 	const FCollisionFilterData QueryFilterDataConcrete;
-	const QueryGeometryType* QueryGeom;
+	const QueryGeometryType* QueryGeom = nullptr;
 	ICollisionQueryFilterCallbackBase& QueryCallback;
 	const FTransform StartTM;
 };
