Added a test to ensure objects put to sleep during resim correction are properly interpolated

diff --git a/Engine/Source/Programs/HeadlessChaos/Private/HeadlessChaosTestRewind.cpp b/Engine/Source/Programs/HeadlessChaos/Private/HeadlessChaosTestRewind.cpp
--- a/Engine/Source/Programs/HeadlessChaos/Private/HeadlessChaosTestRewind.cpp
+++ b/Engine/Source/Programs/HeadlessChaos/Private/HeadlessChaosTestRewind.cpp
@@ -632,6 +632,121 @@ namespace ChaosTest {
 			});
 	}
 
+	GTEST_TEST(AllTraits, RewindTest_MovingToNotMovingInterpolation)
+	{
+
+		//During correction we put a moving object to sleep
+		//This tests that even though it's only dirty for one frame, we still interpolate it over many
+		TRewindHelper::TestDynamicSphere([](auto* Solver, FReal SimDt, int32 Optimization, auto Proxy, auto Sphere)
+			{
+				//only care about resim when async results are used
+				if (Solver->IsUsingAsyncResults() == false)
+				{
+					return;
+				}
+
+				struct FRewindCallback : public IRewindCallback
+				{
+					FSingleParticlePhysicsProxy* Proxy;
+					FRewindData* RewindData;
+					FReal SimDt;
+
+					virtual void ProcessInputs_Internal(int32 PhysicsStep, const TArray<FSimCallbackInputAndObject>& SimCallbackInputs) override
+					{
+						const FReal Time = PhysicsStep * SimDt;
+
+						if (bIsResimming)
+						{
+							if (PhysicsStep == 0)
+							{
+								Proxy->GetPhysicsThreadAPI()->SetObjectState(EObjectStateType::Sleeping);
+							}
+						}
+					}
+
+					virtual int32 TriggerRewindIfNeeded_Internal(int32 LastCompletedStep) override
+					{
+						const FReal Time = LastCompletedStep * SimDt;
+						if (!bIsResimming && Time == ResimTime)
+						{
+							bIsResimming = true;
+							return 0;
+						}
+
+						return INDEX_NONE;
+					}
+
+					bool bIsResimming = false;
+					FReal ResimTime = 32;
+				};
+
+				const int32 LastGameStep = 64;
+				const int32 NumPhysSteps = FMath::TruncToInt(LastGameStep / SimDt);
+
+				auto UniqueRewindCallback = MakeUnique<FRewindCallback>();
+				auto RewindCallback = UniqueRewindCallback.Get();
+				RewindCallback->Proxy = Proxy;
+				RewindCallback->RewindData = Solver->GetRewindData();
+				RewindCallback->SimDt = SimDt;
+
+				auto FloorGeom = TSharedPtr<FImplicitObject, ESPMode::ThreadSafe>(new TBox<FReal, 3>(FVec3(-100, -100, -1), FVec3(100, 100, 0)));
+
+				const FReal LeashTime = 3;
+
+				Solver->SetRewindCallback(MoveTemp(UniqueRewindCallback));
+				Solver->GetResultsManager().SetResimInterpTime(LeashTime);
+
+				auto& Particle = Proxy->GetGameThreadAPI();
+				Particle.SetGravityEnabled(false);
+				Particle.SetV(FVec3(0, 0, 1));
+				Particle.SetX(FVec3(0, 0, 0));
+
+				FReal Time = 0;
+				const FReal GTDt = 1;
+				const FReal InterpStartTime = RewindCallback->ResimTime + SimDt;
+				const FReal InterpEndTime = RewindCallback->ResimTime + LeashTime + SimDt;
+				FReal PrevZDuringInterp = InterpStartTime;
+
+				for (int Step = 0; Step <= LastGameStep; ++Step)
+				{
+					TickSolverHelper(Solver);
+
+					Time += GTDt;
+					const FReal InterpolatedTime = Time - SimDt * Chaos::AsyncInterpolationMultiplier;
+
+					if (InterpolatedTime <= InterpStartTime)
+					{
+						if(InterpolatedTime < 0)
+						{
+							//no interpolation yet so just take first result
+							EXPECT_NEAR(Particle.X()[2], 0, 1e-2);
+						}
+						else
+						{
+							//simple movement with constant velocity
+							EXPECT_NEAR(Particle.X()[2], InterpolatedTime, 1e-2);
+						}
+					}
+					else
+					{
+						if (InterpolatedTime >= InterpEndTime)
+						{
+							//not moving and no longer in leash mode
+							EXPECT_NEAR(Particle.X()[2], 0, 1e-2);
+						}
+						else
+						{
+							//leash mode
+							EXPECT_GT(Particle.X()[2], 0);
+							EXPECT_LT(Particle.X()[2], PrevZDuringInterp);
+							PrevZDuringInterp = Particle.X()[2];
+						}
+					}
+				}
+
+			});
+	}
+
 	GTEST_TEST(AllTraits, RewindTest_ResimSleepChangeRewind)
 	{
 		// Test puting object to sleep during Resim
