Fix compiler errors in HeadlessChaosTestPhysicsScene

diff --git a/Engine/Source/Programs/HeadlessChaos/Private/HeadlessChaosTestPhysicsScene.cpp b/Engine/Source/Programs/HeadlessChaos/Private/HeadlessChaosTestPhysicsScene.cpp
--- a/Engine/Source/Programs/HeadlessChaos/Private/HeadlessChaosTestPhysicsScene.cpp
+++ b/Engine/Source/Programs/HeadlessChaos/Private/HeadlessChaosTestPhysicsScene.cpp
@@ -1068,7 +1068,7 @@ namespace ChaosTest {
 
 		auto AdvanceFrameAndRunTest = [&](const FVec3 &CorrectX, const FVec3 &CorrectV)
 		{
-			Scene.SetUpForFrame(&Grav, Dt, 99999, 99999, 10, false);
+			Scene.SetUpForFrame(&Grav, Dt, 99999, 99999, 10, false, -1);
 			Scene.StartFrame();
 			Scene.EndFrame();
 			// Test X and V on GT
