Fix ArrayView crash due to array resizing

diff --git a/Engine/Source/Programs/HeadlessChaos/Private/GeometryCollection/GeometryCollectionTestFields.cpp b/Engine/Source/Programs/HeadlessChaos/Private/GeometryCollection/GeometryCollectionTestFields.cpp
--- a/Engine/Source/Programs/HeadlessChaos/Private/GeometryCollection/GeometryCollectionTestFields.cpp
+++ b/Engine/Source/Programs/HeadlessChaos/Private/GeometryCollection/GeometryCollectionTestFields.cpp
@@ -48,9 +48,9 @@ namespace GeometryCollectionTest
 			0.0
 		};
 
-		TArray<float> ResultsArray;
+		TArray<Chaos::FReal> ResultsArray;
 		ResultsArray.Init(false, Bounds*Bounds);
-		TFieldArrayView<float> ResultsView(ResultsArray, 0, ResultsArray.Num());
+		TFieldArrayView<Chaos::FReal> ResultsView(ResultsArray, 0, ResultsArray.Num());
 		NoiseField->Evaluate(Context, ResultsView);
 
 		Chaos::FReal min = FLT_MAX, max = -FLT_MAX;
@@ -133,9 +133,9 @@ namespace GeometryCollectionTest
 			0.0
 		};
 
-		TArray<float> ResultsArray;
+		TArray<Chaos::FReal> ResultsArray;
 		ResultsArray.Init(false, 10);
-		TFieldArrayView<float> ResultsView(ResultsArray, 0, ResultsArray.Num());
+		TFieldArrayView<Chaos::FReal> ResultsView(ResultsArray, 0, ResultsArray.Num());
 		RadialFalloff->Evaluate(Context, ResultsView);
 
 		for (int32 Index = 0; Index < 10; Index++)
@@ -181,9 +181,9 @@ namespace GeometryCollectionTest
 			0.0
 		};
 
-		TArray<float> ResultsArray;
+		TArray<Chaos::FReal> ResultsArray;
 		ResultsArray.Init(false, 10);
-		TFieldArrayView<float> ResultsView(ResultsArray, 0, ResultsArray.Num());
+		TFieldArrayView<Chaos::FReal> ResultsView(ResultsArray, 0, ResultsArray.Num());
 		PlaneFalloff->Evaluate(Context, ResultsView);
 
 		FPlane Plane(PlaneFalloff->Position, PlaneFalloff->Normal);
@@ -715,9 +715,9 @@ namespace GeometryCollectionTest
 			0.0
 		};
 
-		TArray<float> ResultsArray;
+		TArray<Chaos::FReal> ResultsArray;
 		ResultsArray.Init(0.f, ExecutionDatas.SamplePositions.Num());
-		TFieldArrayView<float> ResultsView(ResultsArray, 0, ResultsArray.Num());
+		TFieldArrayView<Chaos::FReal> ResultsView(ResultsArray, 0, ResultsArray.Num());
 		SumScalar->Evaluate(Context, ResultsView);
 
 		for (int32 Index = 0; Index < NumPoints; Index++)
@@ -778,9 +778,9 @@ namespace GeometryCollectionTest
 			0.0
 		};
 
-		TArray<float> ResultsArray;
+		TArray<Chaos::FReal> ResultsArray;
 		ResultsArray.Init(0.f, ExecutionDatas.SamplePositions.Num());
-		TFieldArrayView<float> ResultsView(ResultsArray, 0, ResultsArray.Num());
+		TFieldArrayView<Chaos::FReal> ResultsView(ResultsArray, 0, ResultsArray.Num());
 		SumScalar->Evaluate(Context, ResultsView);
 
 		for (int32 Index = 0; Index < NumPoints; Index++)
@@ -833,9 +833,9 @@ namespace GeometryCollectionTest
 			0.0
 		};
 
-		TArray<float> ResultsArray;
+		TArray<Chaos::FReal> ResultsArray;
 		ResultsArray.Init(0.f, ExecutionDatas.SamplePositions.Num());
-		TFieldArrayView<float> ResultsView(ResultsArray, 0, ResultsArray.Num());
+		TFieldArrayView<Chaos::FReal> ResultsView(ResultsArray, 0, ResultsArray.Num());
 		SumScalar->Evaluate(Context, ResultsView);
 
 		for (int32 Index = 0; Index < NumPoints; Index++)
@@ -895,9 +895,9 @@ namespace GeometryCollectionTest
 			0.0
 		};
 
-		TArray<float> ResultsArray;
+		TArray<Chaos::FReal> ResultsArray;
 		ResultsArray.Init(0.f, ExecutionDatas.SamplePositions.Num());
-		TFieldArrayView<float> ResultsView(ResultsArray, 0, ResultsArray.Num());
+		TFieldArrayView<Chaos::FReal> ResultsView(ResultsArray, 0, ResultsArray.Num());
 		CullingField->Evaluate(Context, ResultsView);
 
 		for (int32 Index = 0; Index < NumPoints; Index++)
