Chaos: Fix headless tests compile after SOA changes in 17869846

diff --git a/Engine/Source/Programs/HeadlessChaos/Private/HeadlessChaosTestBroadphase.cpp b/Engine/Source/Programs/HeadlessChaos/Private/HeadlessChaosTestBroadphase.cpp
--- a/Engine/Source/Programs/HeadlessChaos/Private/HeadlessChaosTestBroadphase.cpp
+++ b/Engine/Source/Programs/HeadlessChaos/Private/HeadlessChaosTestBroadphase.cpp
@@ -443,7 +443,8 @@ namespace ChaosTest
 	void GridBPTest2()
 	{
 		TUniquePtr<TBox<FReal, 3>> Box = MakeUnique<TBox<FReal, 3>>(FVec3(0, 0, 0), FVec3(100, 100, 100));
-		FPBDRigidsSOAs SOAs;
+		FParticleUniqueIndicesMultithreaded UniqueIndices;
+		FPBDRigidsSOAs SOAs(UniqueIndices);
 		const int32 NumRows = 10;
 		const int32 NumCols = 10;
 		const int32 NumHeight = 10;
@@ -858,7 +859,8 @@ namespace ChaosTest
 		const int32 ParticleCount = NumRows * NumCols * NumHeight;
 		const FReal BoxSize = 100;
 
-		FPBDRigidsSOAs Particles;
+		FParticleUniqueIndicesMultithreaded UniqueIndices;
+		FPBDRigidsSOAs Particles(UniqueIndices);
 		TArray<FPBDRigidParticleHandle*> ParticleHandles = Particles.CreateDynamicParticles(ParticleCount);
 		for (auto& Handle : ParticleHandles)
 		{
diff --git a/Engine/Source/Programs/HeadlessChaos/Private/HeadlessChaosTestCCD.cpp b/Engine/Source/Programs/HeadlessChaos/Private/HeadlessChaosTestCCD.cpp
--- a/Engine/Source/Programs/HeadlessChaos/Private/HeadlessChaosTestCCD.cpp
+++ b/Engine/Source/Programs/HeadlessChaos/Private/HeadlessChaosTestCCD.cpp
@@ -27,7 +27,8 @@ namespace ChaosTest
 		const FReal InitialSpeed = BoxHalfSize * 5 * Fps; // More than enough to tunnel
 		const FReal InitialPosition = BoxHalfSize * 2 + 30;
 		using TEvolution = FPBDRigidsEvolutionGBF;
-		FPBDRigidsSOAs Particles;
+		FParticleUniqueIndicesMultithreaded UniqueIndices;
+		FPBDRigidsSOAs Particles(UniqueIndices);
 		THandleArray<FChaosPhysicsMaterial> PhysicalMaterials;
 		TEvolution Evolution(Particles, PhysicalMaterials);
 		InitEvolutionSettings(Evolution);
@@ -87,7 +88,8 @@ namespace ChaosTest
 		const FReal BoxHalfSize = 50; // cm
 		const FReal InitialSpeed = BoxHalfSize * 5 * Fps; // More than enough to tunnel
 		using TEvolution = FPBDRigidsEvolutionGBF;
-		FPBDRigidsSOAs Particles;
+		FParticleUniqueIndicesMultithreaded UniqueIndices;
+		FPBDRigidsSOAs Particles(UniqueIndices);
 		THandleArray<FChaosPhysicsMaterial> PhysicalMaterials;
 		TEvolution Evolution(Particles, PhysicalMaterials);
 		InitEvolutionSettings(Evolution);
@@ -147,7 +149,8 @@ namespace ChaosTest
 		const FReal SphereRadius = 100; // cm
 		const FReal InitialSpeed = SphereRadius * 5 * Fps; // More than enough to tunnel
 		using TEvolution = FPBDRigidsEvolutionGBF;
-		FPBDRigidsSOAs Particles;
+		FParticleUniqueIndicesMultithreaded UniqueIndices;
+		FPBDRigidsSOAs Particles(UniqueIndices);
 		THandleArray<FChaosPhysicsMaterial> PhysicalMaterials;
 		TEvolution Evolution(Particles, PhysicalMaterials);
 		InitEvolutionSettings(Evolution);
@@ -216,7 +219,8 @@ namespace ChaosTest
 		const FVec3 InitialVelocity = FVec3(-ContainerBoxHalfSize * Fps * 10, 0 ,0);
 
 		using TEvolution = FPBDRigidsEvolutionGBF;
-		FPBDRigidsSOAs Particles;
+		FParticleUniqueIndicesMultithreaded UniqueIndices;
+		FPBDRigidsSOAs Particles(UniqueIndices);
 		THandleArray<FChaosPhysicsMaterial> PhysicalMaterials;
 		TEvolution Evolution(Particles, PhysicalMaterials);
 		InitEvolutionSettings(Evolution);
diff --git a/Engine/Source/Programs/HeadlessChaos/Private/HeadlessChaosTestClustering.cpp b/Engine/Source/Programs/HeadlessChaos/Private/HeadlessChaosTestClustering.cpp
--- a/Engine/Source/Programs/HeadlessChaos/Private/HeadlessChaosTestClustering.cpp
+++ b/Engine/Source/Programs/HeadlessChaos/Private/HeadlessChaosTestClustering.cpp
@@ -23,7 +23,8 @@ namespace ChaosTest {
 
 	void ImplicitCluster()
 	{
-		FPBDRigidsSOAs Particles;
+		FParticleUniqueIndicesMultithreaded UniqueIndices;
+		FPBDRigidsSOAs Particles(UniqueIndices);
 		THandleArray<FChaosPhysicsMaterial> PhysicalMaterials;
 		FPBDRigidsEvolution Evolution(Particles, PhysicalMaterials);
 		FPBDRigidClusteredParticles& ClusteredParticles = Particles.GetClusteredParticles();
@@ -68,7 +69,8 @@ namespace ChaosTest {
 
 	void FractureCluster()
 	{
-		FPBDRigidsSOAs Particles;
+		FParticleUniqueIndicesMultithreaded UniqueIndices;
+		FPBDRigidsSOAs Particles(UniqueIndices);
 		THandleArray<FChaosPhysicsMaterial> PhysicalMaterials;
 		FPBDRigidsEvolution Evolution(Particles, PhysicalMaterials);
 		auto& ClusteredParticles = Particles.GetClusteredParticles();
@@ -192,7 +194,8 @@ namespace ChaosTest {
 
 	void PartialFractureCluster()
 	{
-		FPBDRigidsSOAs Particles;
+		FParticleUniqueIndicesMultithreaded UniqueIndices;
+		FPBDRigidsSOAs Particles(UniqueIndices);
 		THandleArray<FChaosPhysicsMaterial> PhysicalMaterials;
 		FPBDRigidsEvolution Evolution(Particles, PhysicalMaterials);
 		auto& ClusteredParticles = Particles.GetClusteredParticles();
diff --git a/Engine/Source/Programs/HeadlessChaos/Private/HeadlessChaosTestCollisions.cpp b/Engine/Source/Programs/HeadlessChaos/Private/HeadlessChaosTestCollisions.cpp
--- a/Engine/Source/Programs/HeadlessChaos/Private/HeadlessChaosTestCollisions.cpp
+++ b/Engine/Source/Programs/HeadlessChaos/Private/HeadlessChaosTestCollisions.cpp
@@ -37,7 +37,8 @@ namespace ChaosTest {
 		TArrayCollectionArray<TSerializablePtr<FChaosPhysicsMaterial>> PhysicsMaterials;
 		TArrayCollectionArray<TUniquePtr<FChaosPhysicsMaterial>> PerParticlePhysicsMaterials;
 		
-		FPBDRigidsSOAs Particles;
+		FParticleUniqueIndicesMultithreaded UniqueIndices;
+		FPBDRigidsSOAs Particles(UniqueIndices);
 		Particles.GetParticleHandles().AddArray(&Collided);
 		Particles.GetParticleHandles().AddArray(&PhysicsMaterials);
 		Particles.GetParticleHandles().AddArray(&PerParticlePhysicsMaterials);
@@ -83,7 +84,8 @@ namespace ChaosTest {
 		TArrayCollectionArray<TSerializablePtr<FChaosPhysicsMaterial>> PhysicsMaterials;
 		TArrayCollectionArray<TUniquePtr<FChaosPhysicsMaterial>> PerParticlePhysicsMaterials;
 
-		FPBDRigidsSOAs Particles;
+		FParticleUniqueIndicesMultithreaded UniqueIndices;
+		FPBDRigidsSOAs Particles(UniqueIndices);
 		Particles.GetParticleHandles().AddArray(&Collided);
 		Particles.GetParticleHandles().AddArray(&PhysicsMaterials);
 		Particles.GetParticleHandles().AddArray(&PerParticlePhysicsMaterials);
@@ -125,7 +127,8 @@ namespace ChaosTest {
 		TArrayCollectionArray<TSerializablePtr<FChaosPhysicsMaterial>> PhysicsMaterials;
 		TArrayCollectionArray<TUniquePtr<FChaosPhysicsMaterial>> PerParticlePhysicsMaterials;
 
-		FPBDRigidsSOAs Particles;
+		FParticleUniqueIndicesMultithreaded UniqueIndices;
+		FPBDRigidsSOAs Particles(UniqueIndices);
 		Particles.GetParticleHandles().AddArray(&Collided);
 		Particles.GetParticleHandles().AddArray(&PhysicsMaterials);
 		Particles.GetParticleHandles().AddArray(&PerParticlePhysicsMaterials);
@@ -202,7 +205,8 @@ namespace ChaosTest {
 		TArrayCollectionArray<TSerializablePtr<FChaosPhysicsMaterial>> PhysicsMaterials;
 		TArrayCollectionArray<TUniquePtr<FChaosPhysicsMaterial>> PerParticlePhysicsMaterials;
 
-		FPBDRigidsSOAs Particles;
+		FParticleUniqueIndicesMultithreaded UniqueIndices;
+		FPBDRigidsSOAs Particles(UniqueIndices);
 		Particles.GetParticleHandles().AddArray(&Collided);
 		Particles.GetParticleHandles().AddArray(&PhysicsMaterials);
 		Particles.GetParticleHandles().AddArray(&PerParticlePhysicsMaterials);
@@ -275,7 +279,8 @@ namespace ChaosTest {
 		TArrayCollectionArray<TSerializablePtr<FChaosPhysicsMaterial>> PhysicsMaterials;
 		TArrayCollectionArray<TUniquePtr<FChaosPhysicsMaterial>> PerParticlePhysicsMaterials;
 		
-		FPBDRigidsSOAs Particles;
+		FParticleUniqueIndicesMultithreaded UniqueIndices;
+		FPBDRigidsSOAs Particles(UniqueIndices);
 		Particles.GetParticleHandles().AddArray(&Collided);
 		Particles.GetParticleHandles().AddArray(&PhysicsMaterials);
 		Particles.GetParticleHandles().AddArray(&PerParticlePhysicsMaterials);
@@ -346,7 +351,8 @@ namespace ChaosTest {
 		PhysicsMaterial->Restitution = (FReal)1;
 		TArrayCollectionArray<TSerializablePtr<FChaosPhysicsMaterial>> PhysicsMaterials;
 		TArrayCollectionArray<TUniquePtr<FChaosPhysicsMaterial>> PerParticlePhysicsMaterials;
-		FPBDRigidsSOAs Particles;
+		FParticleUniqueIndicesMultithreaded UniqueIndices;
+		FPBDRigidsSOAs Particles(UniqueIndices);
 		Particles.GetParticleHandles().AddArray(&Collided);
 		Particles.GetParticleHandles().AddArray(&PhysicsMaterials);
 		Particles.GetParticleHandles().AddArray(&PerParticlePhysicsMaterials);
@@ -422,7 +428,8 @@ namespace ChaosTest {
 		PhysicsMaterial->Restitution = (FReal)0.5;
 		TArrayCollectionArray<TSerializablePtr<FChaosPhysicsMaterial>> PhysicsMaterials;
 		TArrayCollectionArray<TUniquePtr<FChaosPhysicsMaterial>> PerParticlePhysicsMaterials;
-		FPBDRigidsSOAs Particles;
+		FParticleUniqueIndicesMultithreaded UniqueIndices;
+		FPBDRigidsSOAs Particles(UniqueIndices);
 		Particles.GetParticleHandles().AddArray(&Collided);
 		Particles.GetParticleHandles().AddArray(&PhysicsMaterials);
 		Particles.GetParticleHandles().AddArray(&PerParticlePhysicsMaterials);
@@ -481,7 +488,8 @@ namespace ChaosTest {
 		PhysicsMaterial->Restitution = (FReal)0;
 		TArrayCollectionArray<TSerializablePtr<FChaosPhysicsMaterial>> PhysicsMaterials;
 		TArrayCollectionArray<TUniquePtr<FChaosPhysicsMaterial>> PerParticlePhysicsMaterials;
-		FPBDRigidsSOAs Particles;
+		FParticleUniqueIndicesMultithreaded UniqueIndices;
+		FPBDRigidsSOAs Particles(UniqueIndices);
 		Particles.GetParticleHandles().AddArray(&Collided);
 		Particles.GetParticleHandles().AddArray(&PhysicsMaterials);
 		Particles.GetParticleHandles().AddArray(&PerParticlePhysicsMaterials);
diff --git a/Engine/Source/Programs/HeadlessChaos/Private/HeadlessChaosTestConstraintHandle.cpp b/Engine/Source/Programs/HeadlessChaos/Private/HeadlessChaosTestConstraintHandle.cpp
--- a/Engine/Source/Programs/HeadlessChaos/Private/HeadlessChaosTestConstraintHandle.cpp
+++ b/Engine/Source/Programs/HeadlessChaos/Private/HeadlessChaosTestConstraintHandle.cpp
@@ -109,7 +109,8 @@ namespace ChaosTest {
 	template<typename TEvolution>
 	void JointConstraintHandles()
 	{
-		FPBDRigidsSOAs Particles;
+		FParticleUniqueIndicesMultithreaded UniqueIndices;
+		FPBDRigidsSOAs Particles(UniqueIndices);
 		THandleArray<FChaosPhysicsMaterial> PhysicalMaterials;
 		TEvolution Evolution(Particles, PhysicalMaterials);
 
@@ -130,7 +131,8 @@ namespace ChaosTest {
 	template<typename TEvolution>
 	void PositionConstraintHandles()
 	{
-		FPBDRigidsSOAs Particles;
+		FParticleUniqueIndicesMultithreaded UniqueIndices;
+		FPBDRigidsSOAs Particles(UniqueIndices);
 		THandleArray<FChaosPhysicsMaterial> PhysicalMaterials;
 		TEvolution Evolution(Particles, PhysicalMaterials);
 		TArray<FPBDRigidParticleHandle*> ParticleHandles = Evolution.CreateDynamicParticles(5);
@@ -150,7 +152,8 @@ namespace ChaosTest {
 	template<typename TEvolution>
 	void RigidSpringConstraintHandles()
 	{
-		FPBDRigidsSOAs Particles;
+		FParticleUniqueIndicesMultithreaded UniqueIndices;
+		FPBDRigidsSOAs Particles(UniqueIndices);
 		THandleArray<FChaosPhysicsMaterial> PhysicalMaterials;
 		TEvolution Evolution(Particles, PhysicalMaterials);
 		TArray<FPBDRigidParticleHandle*> ParticleHandles = Evolution.CreateDynamicParticles(5);
@@ -171,7 +174,8 @@ namespace ChaosTest {
 	template<typename TEvolution>
 	void RigidDynamicSpringConstraintHandles()
 	{
-		FPBDRigidsSOAs Particles;
+		FParticleUniqueIndicesMultithreaded UniqueIndices;
+		FPBDRigidsSOAs Particles(UniqueIndices);
 		THandleArray<FChaosPhysicsMaterial> PhysicalMaterials;
 		TEvolution Evolution(Particles, PhysicalMaterials);
 		TArray<FPBDRigidParticleHandle*> ParticleHandles = Evolution.CreateDynamicParticles(5);
diff --git a/Engine/Source/Programs/HeadlessChaos/Private/HeadlessChaosTestConstraints.cpp b/Engine/Source/Programs/HeadlessChaos/Private/HeadlessChaosTestConstraints.cpp
--- a/Engine/Source/Programs/HeadlessChaos/Private/HeadlessChaosTestConstraints.cpp
+++ b/Engine/Source/Programs/HeadlessChaos/Private/HeadlessChaosTestConstraints.cpp
@@ -30,7 +30,8 @@ namespace ChaosTest {
 	void Position()
 	{
 		{
-			FPBDRigidsSOAs Particles;
+			FParticleUniqueIndicesMultithreaded UniqueIndices;
+			FPBDRigidsSOAs Particles(UniqueIndices);
 			THandleArray<FChaosPhysicsMaterial> PhysicalMaterials;
 			TEvolution Evolution(Particles, PhysicalMaterials);
 			TArray<FPBDRigidParticleHandle*> Dynamics = Evolution.CreateDynamicParticles(1);
@@ -45,7 +46,8 @@ namespace ChaosTest {
 			EXPECT_LT(Evolution.GetParticleHandles().Handle(0)->X().SizeSquared(), SMALL_NUMBER);
 		}
 		{
-			FPBDRigidsSOAs Particles;
+			FParticleUniqueIndicesMultithreaded UniqueIndices;
+			FPBDRigidsSOAs Particles(UniqueIndices);
 			THandleArray<FChaosPhysicsMaterial> PhysicalMaterials;
 			TEvolution Evolution(Particles, PhysicalMaterials);
 			InitEvolutionSettings(Evolution);
@@ -86,7 +88,8 @@ namespace ChaosTest {
 	void PositionAndJoint()
 	{
 		const int32 Iterations = 10;
-		FPBDRigidsSOAs Particles;
+		FParticleUniqueIndicesMultithreaded UniqueIndices;
+		FPBDRigidsSOAs Particles(UniqueIndices);
 		THandleArray<FChaosPhysicsMaterial> PhysicalMaterials;
 		TEvolution Evolution(Particles, PhysicalMaterials);
 		InitEvolutionSettings(Evolution);
@@ -144,7 +147,8 @@ namespace ChaosTest {
 		SuspensionSettings.Axis = FVec3(0.0f, 0.0f, 1.0f);
 
 		{
-			FPBDRigidsSOAs Particles;
+			FParticleUniqueIndicesMultithreaded UniqueIndices;
+			FPBDRigidsSOAs Particles(UniqueIndices);
 			THandleArray<FChaosPhysicsMaterial> PhysicalMaterials;
 			TEvolution Evolution(Particles, PhysicalMaterials);
 			InitEvolutionSettings(Evolution);
@@ -194,7 +198,8 @@ namespace ChaosTest {
 
 		{
 
-			FPBDRigidsSOAs Particles;
+			FParticleUniqueIndicesMultithreaded UniqueIndices;
+			FPBDRigidsSOAs Particles(UniqueIndices);
 			THandleArray<FChaosPhysicsMaterial> PhysicalMaterials;
 			TEvolution Evolution(Particles, PhysicalMaterials);
 			InitEvolutionSettings(Evolution);
@@ -251,7 +256,8 @@ namespace ChaosTest {
 	template<typename TEvolution>
 	void SuspensionConstraintSpring()
 	{
-		FPBDRigidsSOAs Particles;
+		FParticleUniqueIndicesMultithreaded UniqueIndices;
+		FPBDRigidsSOAs Particles(UniqueIndices);
 		THandleArray<FChaosPhysicsMaterial> PhysicalMaterials;
 		TEvolution Evolution(Particles, PhysicalMaterials);
 		InitEvolutionSettings(Evolution);
diff --git a/Engine/Source/Programs/HeadlessChaos/Private/HeadlessChaosTestForces.cpp b/Engine/Source/Programs/HeadlessChaos/Private/HeadlessChaosTestForces.cpp
--- a/Engine/Source/Programs/HeadlessChaos/Private/HeadlessChaosTestForces.cpp
+++ b/Engine/Source/Programs/HeadlessChaos/Private/HeadlessChaosTestForces.cpp
@@ -16,7 +16,8 @@ namespace ChaosTest
 	template<typename TEvolution>
 	void Gravity()
 	{
-		FPBDRigidsSOAs Particles;
+		FParticleUniqueIndicesMultithreaded UniqueIndices;
+		FPBDRigidsSOAs Particles(UniqueIndices);
 		THandleArray<FChaosPhysicsMaterial> PhysicalMaterials;
 		TEvolution Evolution(Particles, PhysicalMaterials);
 		
diff --git a/Engine/Source/Programs/HeadlessChaos/Private/HeadlessChaosTestGraph.cpp b/Engine/Source/Programs/HeadlessChaos/Private/HeadlessChaosTestGraph.cpp
--- a/Engine/Source/Programs/HeadlessChaos/Private/HeadlessChaosTestGraph.cpp
+++ b/Engine/Source/Programs/HeadlessChaos/Private/HeadlessChaosTestGraph.cpp
@@ -128,7 +128,8 @@ namespace ChaosTest {
 		// Create some dynamic particles - doesn't matter what position or other state they have
 		TArray<TGeometryParticleHandle<FReal, 3>*> AllParticles;
 
-		FPBDRigidsSOAs SOAs;
+		FParticleUniqueIndicesMultithreaded UniqueIndices;
+		FPBDRigidsSOAs SOAs(UniqueIndices);
 		TArray<TPBDRigidParticleHandle<FReal,3>*> Dynamics = SOAs.CreateDynamicParticles(17);
 		for (auto Dyn : Dynamics) { AllParticles.Add(Dyn);}
 
@@ -288,7 +289,8 @@ namespace ChaosTest {
 		// Create some dynamic particles - doesn't matter what position or other state they have
 		TArray<TGeometryParticleHandle<FReal, 3>*> AllParticles;
 
-		FPBDRigidsSOAs SOAs;
+		FParticleUniqueIndicesMultithreaded UniqueIndices;
+		FPBDRigidsSOAs SOAs(UniqueIndices);
 		TArray<TPBDRigidParticleHandle<FReal, 3>*> Dynamics = SOAs.CreateDynamicParticles(2);
 		for (auto Dyn : Dynamics) { AllParticles.Add(Dyn); }
 
@@ -636,7 +638,8 @@ namespace ChaosTest {
 
 			// Create some dynamic particles
 			int32 NumParticles = 6;
-			FPBDRigidsSOAs SOAs;
+			FParticleUniqueIndicesMultithreaded UniqueIndices;
+			FPBDRigidsSOAs SOAs(UniqueIndices);
 			TArray<TPBDRigidParticleHandle<FReal, 3>*> Particles = SOAs.CreateDynamicParticles(NumParticles);
 			TArrayCollectionArray<TSerializablePtr<FChaosPhysicsMaterial>> PhysicsMaterials;
 			THandleArray<FChaosPhysicsMaterial> PhysicalMaterials;
@@ -714,7 +717,8 @@ namespace ChaosTest {
 
 			// Create some dynamic particles
 			int32 NumParticles = 6;
-			FPBDRigidsSOAs SOAs;
+			FParticleUniqueIndicesMultithreaded UniqueIndices;
+			FPBDRigidsSOAs SOAs(UniqueIndices);
 			TArray<TPBDRigidParticleHandle<FReal, 3>*> Particles = SOAs.CreateDynamicParticles(NumParticles);
 			TArrayCollectionArray<TSerializablePtr<FChaosPhysicsMaterial>> PhysicsMaterials;
 			THandleArray<FChaosPhysicsMaterial> PhysicalMaterials;
@@ -807,7 +811,8 @@ namespace ChaosTest {
 
 			// Create some dynamic particles
 			int32 NumParticles = 6;
-			FPBDRigidsSOAs SOAs;
+			FParticleUniqueIndicesMultithreaded UniqueIndices;
+			FPBDRigidsSOAs SOAs(UniqueIndices);
 			TArray<TPBDRigidParticleHandle<FReal,3>*> Particles = SOAs.CreateDynamicParticles(NumParticles);
 			TArrayCollectionArray<TSerializablePtr<FChaosPhysicsMaterial>> PhysicsMaterials;
 			THandleArray<FChaosPhysicsMaterial> PhysicalMaterials;
@@ -926,7 +931,8 @@ namespace ChaosTest {
 	void GraphColorGrid(const int32 NumParticlesX, const int32 NumParticlesY, const int32 Multiplicity, const bool bRandomize)
 	{
 		// Create a grid of particles
-		FPBDRigidsSOAs SOAs;
+		FParticleUniqueIndicesMultithreaded UniqueIndices;
+		FPBDRigidsSOAs SOAs(UniqueIndices);
 		int32 NumParticles = NumParticlesX * NumParticlesY;
 		TArray<TGeometryParticleHandle<FReal,3>*> AllParticles;
 
@@ -1156,7 +1162,8 @@ namespace ChaosTest {
 		// Create some particles - doesn't matter what position or other state they have
 		TArray<TGeometryParticleHandle<FReal, 3>*> AllParticles;
 
-		FPBDRigidsSOAs SOAs;
+		FParticleUniqueIndicesMultithreaded UniqueIndices;
+		FPBDRigidsSOAs SOAs(UniqueIndices);
 		TArray<FPBDRigidParticleHandle*> Dynamics = SOAs.CreateDynamicParticles(500);
 		for (auto Dynamic : Dynamics) { AllParticles.Add(Dynamic); }
 
diff --git a/Engine/Source/Programs/HeadlessChaos/Private/HeadlessChaosTestJoints.cpp b/Engine/Source/Programs/HeadlessChaos/Private/HeadlessChaosTestJoints.cpp
--- a/Engine/Source/Programs/HeadlessChaos/Private/HeadlessChaosTestJoints.cpp
+++ b/Engine/Source/Programs/HeadlessChaos/Private/HeadlessChaosTestJoints.cpp
@@ -480,7 +480,8 @@ namespace ChaosTest {
 		PhysicalMaterial->DisabledLinearThreshold = 0;
 		PhysicalMaterial->DisabledAngularThreshold = 0;
 
-		FPBDRigidsSOAs Particles;
+		FParticleUniqueIndicesMultithreaded UniqueIndices;
+		FPBDRigidsSOAs Particles(UniqueIndices);
 
 		auto StaticBox = AppendStaticParticleBox(Particles, FVec3((FReal)100, (FReal)100, (FReal)100));
 		auto Box2 = AppendDynamicParticleBox(Particles, FVec3((FReal)100, (FReal)100, (FReal)100));
@@ -522,7 +523,8 @@ namespace ChaosTest {
 		PhysicalMaterial->DisabledAngularThreshold = 0;
 
 		{
-			FPBDRigidsSOAs Particles;
+			FParticleUniqueIndicesMultithreaded UniqueIndices;
+			FPBDRigidsSOAs Particles(UniqueIndices);
 
 			auto& StaticBox = *AppendStaticParticleBox(Particles, FVec3((FReal)100, (FReal)100, (FReal)100));
 			auto& Box2 = *AppendDynamicParticleBox(Particles, FVec3((FReal)100, (FReal)100, (FReal)100));
@@ -552,7 +554,8 @@ namespace ChaosTest {
 		}
 
 		{
-			FPBDRigidsSOAs Particles;
+			FParticleUniqueIndicesMultithreaded UniqueIndices;
+			FPBDRigidsSOAs Particles(UniqueIndices);
 
 			auto& StaticBox = *AppendStaticParticleBox(Particles, FVec3((FReal)100, (FReal)100, (FReal)100));
 			auto& Box2 = *AppendDynamicParticleBox(Particles, FVec3((FReal)100, (FReal)100, (FReal)100));
@@ -616,7 +619,8 @@ namespace ChaosTest {
 	// Check that constraints end up in the same island when graph is fully connected
 	GTEST_TEST(JointTests, TestJointConstraintGraph_Connected)
 	{
-		FPBDRigidsSOAs ParticleContainer;
+		FParticleUniqueIndicesMultithreaded UniqueIndices;
+		FPBDRigidsSOAs ParticleContainer(UniqueIndices);
 		FPBDJointConstraints JointContainer;
 
 		// Create 3 particles
@@ -646,7 +650,8 @@ namespace ChaosTest {
 	// Check that constraints islands are not merged through shared kinematic particles
 	GTEST_TEST(JointTests, TestJointConstraintGraph_NotConnected)
 	{
-		FPBDRigidsSOAs ParticleContainer;
+		FParticleUniqueIndicesMultithreaded UniqueIndices;
+		FPBDRigidsSOAs ParticleContainer(UniqueIndices);
 		FPBDJointConstraints JointContainer;
 
 		// Create 3 particles
diff --git a/Engine/Source/Programs/HeadlessChaos/Private/HeadlessChaosTestMinEvolution.cpp b/Engine/Source/Programs/HeadlessChaos/Private/HeadlessChaosTestMinEvolution.cpp
--- a/Engine/Source/Programs/HeadlessChaos/Private/HeadlessChaosTestMinEvolution.cpp
+++ b/Engine/Source/Programs/HeadlessChaos/Private/HeadlessChaosTestMinEvolution.cpp
@@ -30,7 +30,8 @@ namespace ChaosTest
 		using FParticlePair = TVec2<FGeometryParticleHandle*>;
 
 		// Particles
-		FRigidParticleSOAs ParticlesContainer;
+		FParticleUniqueIndicesMultithreaded UniqueIndices;
+		FRigidParticleSOAs ParticlesContainer(UniqueIndices);
 
 		// @todo(ccaulfield): we shouldn't require collisions to use an evolution...
 		// Stuff needed for collisions
diff --git a/Engine/Source/Programs/HeadlessChaos/Private/HeadlessChaosTestParticleHandle.cpp b/Engine/Source/Programs/HeadlessChaos/Private/HeadlessChaosTestParticleHandle.cpp
--- a/Engine/Source/Programs/HeadlessChaos/Private/HeadlessChaosTestParticleHandle.cpp
+++ b/Engine/Source/Programs/HeadlessChaos/Private/HeadlessChaosTestParticleHandle.cpp
@@ -310,7 +310,8 @@ namespace ChaosTest {
 	void ParticleLifetimeAndThreading()
 	{
 		{
-			FPBDRigidsSOAs SOAs;
+			FParticleUniqueIndicesMultithreaded UniqueIndices;
+			FPBDRigidsSOAs SOAs(UniqueIndices);
 
 			TArray<TUniquePtr<FPBDRigidParticle>> GTRawParticles;
 			for (int i = 0; i < 3; ++i)
@@ -380,7 +381,8 @@ namespace ChaosTest {
 	void ParticleDestroyOrdering()
 	{
 		{
-			FPBDRigidsSOAs SOAs;
+			FParticleUniqueIndicesMultithreaded UniqueIndices;
+			FPBDRigidsSOAs SOAs(UniqueIndices);
 			SOAs.CreateDynamicParticles(10);
 			FReal Count = 0;
 			FGeometryParticleHandle* ThirdParticle = nullptr;
@@ -416,7 +418,8 @@ namespace ChaosTest {
 
 		//now test non swapping remove
 		{
-			FPBDRigidsSOAs SOAs;
+			FParticleUniqueIndicesMultithreaded UniqueIndices;
+			FPBDRigidsSOAs SOAs(UniqueIndices);
 			SOAs.CreateClusteredParticles(10);
 			FReal Count = 0;
 			FGeometryParticleHandle* ThirdParticle = nullptr;
@@ -458,7 +461,8 @@ namespace ChaosTest {
 	{
 		FWeakParticleHandle WeakHandle;
 		{
-			FPBDRigidsSOAs SOAs;
+			FParticleUniqueIndicesMultithreaded UniqueIndices;
+			FPBDRigidsSOAs SOAs(UniqueIndices);
 			SOAs.CreateStaticParticles(1);
 			for(auto& Particle : SOAs.GetAllParticlesView())
 			{
@@ -547,7 +551,8 @@ namespace ChaosTest {
 
 			//Use an SOA with an active list
 			{
-				FPBDRigidsSOAs SOAsWithHandles;	//todo: create a mock object so we can more easily create handles
+				FParticleUniqueIndicesMultithreaded UniqueIndices;
+				FPBDRigidsSOAs SOAsWithHandles(UniqueIndices);	//todo: create a mock object so we can more easily create handles
 				auto PartialDynamics = SOAsWithHandles.CreateDynamicParticles(10);
 
 				TArray<FPBDRigidParticleHandle*> ActiveParticles = { PartialDynamics[3], PartialDynamics[5] };
@@ -583,7 +588,8 @@ namespace ChaosTest {
 
 		{
 			// try using SOA manager
-			FPBDRigidsSOAs SOAs;
+			FParticleUniqueIndicesMultithreaded UniqueIndices;
+			FPBDRigidsSOAs SOAs(UniqueIndices);
 			SOAs.CreateStaticParticles(3);
 			auto KinematicParticles = SOAs.CreateKinematicParticles(3);
 			SOAs.CreateDynamicParticles(3);
@@ -686,7 +692,8 @@ namespace ChaosTest {
 
 	void HandleObjectStateChangeTest()
 	{
-		FPBDRigidsSOAs SOAs;
+		FParticleUniqueIndicesMultithreaded UniqueIndices;
+		FPBDRigidsSOAs SOAs(UniqueIndices);
 
 		// Lambda for adding a particle to the dynamic-backed kinematic SOA
 		const auto CreateDynamicKinematic = [&]()
diff --git a/Engine/Source/Programs/HeadlessChaos/Private/HeadlessChaosTestPerf.cpp b/Engine/Source/Programs/HeadlessChaos/Private/HeadlessChaosTestPerf.cpp
--- a/Engine/Source/Programs/HeadlessChaos/Private/HeadlessChaosTestPerf.cpp
+++ b/Engine/Source/Programs/HeadlessChaos/Private/HeadlessChaosTestPerf.cpp
@@ -34,7 +34,8 @@ namespace ChaosTest {
 
 	GTEST_TEST(AllEvolutions, PerfTests_Sim)
 	{
-		FPBDRigidsSOAs Particles;
+		FParticleUniqueIndicesMultithreaded UniqueIndices;
+		FPBDRigidsSOAs Particles(UniqueIndices);
 		THandleArray<FChaosPhysicsMaterial> PhysicalMaterials;
 		FPBDRigidsEvolutionGBF Evolution(Particles, PhysicalMaterials);
 		InitEvolutionSettings(Evolution);
diff --git a/Engine/Source/Programs/HeadlessChaos/Private/HeadlessChaosTestSerialization.cpp b/Engine/Source/Programs/HeadlessChaos/Private/HeadlessChaosTestSerialization.cpp
--- a/Engine/Source/Programs/HeadlessChaos/Private/HeadlessChaosTestSerialization.cpp
+++ b/Engine/Source/Programs/HeadlessChaos/Private/HeadlessChaosTestSerialization.cpp
@@ -470,7 +470,8 @@ namespace ChaosTest
 			if (File)
 			{
 				Chaos::FChaosArchive ChaosAr(*File);
-				FPBDRigidsSOAs Particles;
+				FParticleUniqueIndicesMultithreaded UniqueIndices;
+				FPBDRigidsSOAs Particles(UniqueIndices);
 
 				THandleArray<FChaosPhysicsMaterial> PhysicalMaterials;
 				FPBDRigidsEvolution Evolution(Particles, PhysicalMaterials);
diff --git a/Engine/Source/Programs/HeadlessChaos/Private/HeadlessChaosTestShapePadding.cpp b/Engine/Source/Programs/HeadlessChaos/Private/HeadlessChaosTestShapePadding.cpp
--- a/Engine/Source/Programs/HeadlessChaos/Private/HeadlessChaosTestShapePadding.cpp
+++ b/Engine/Source/Programs/HeadlessChaos/Private/HeadlessChaosTestShapePadding.cpp
@@ -37,7 +37,8 @@ namespace ChaosTest {
 		TArrayCollectionArray<TSerializablePtr<FChaosPhysicsMaterial>> PhysicsMaterials;
 		TArrayCollectionArray<TUniquePtr<FChaosPhysicsMaterial>> PerParticlePhysicsMaterials;
 
-		FPBDRigidsSOAs Particles;
+		FParticleUniqueIndicesMultithreaded UniqueIndices;
+		FPBDRigidsSOAs Particles(UniqueIndices);
 		Particles.GetParticleHandles().AddArray(&Collided);
 		Particles.GetParticleHandles().AddArray(&PhysicsMaterials);
 		Particles.GetParticleHandles().AddArray(&PerParticlePhysicsMaterials);
@@ -144,7 +145,8 @@ namespace ChaosTest {
 		TArrayCollectionArray<TSerializablePtr<FChaosPhysicsMaterial>> PhysicsMaterials;
 		TArrayCollectionArray<TUniquePtr<FChaosPhysicsMaterial>> PerParticlePhysicsMaterials;
 
-		FPBDRigidsSOAs Particles;
+		FParticleUniqueIndicesMultithreaded UniqueIndices;
+		FPBDRigidsSOAs Particles(UniqueIndices);
 		Particles.GetParticleHandles().AddArray(&Collided);
 		Particles.GetParticleHandles().AddArray(&PhysicsMaterials);
 		Particles.GetParticleHandles().AddArray(&PerParticlePhysicsMaterials);
@@ -257,7 +259,8 @@ namespace ChaosTest {
 		TArrayCollectionArray<TSerializablePtr<FChaosPhysicsMaterial>> PhysicsMaterials;
 		TArrayCollectionArray<TUniquePtr<FChaosPhysicsMaterial>> PerParticlePhysicsMaterials;
 
-		FPBDRigidsSOAs Particles;
+		FParticleUniqueIndicesMultithreaded UniqueIndices;
+		FPBDRigidsSOAs Particles(UniqueIndices);
 		Particles.GetParticleHandles().AddArray(&Collided);
 		Particles.GetParticleHandles().AddArray(&PhysicsMaterials);
 		Particles.GetParticleHandles().AddArray(&PerParticlePhysicsMaterials);
diff --git a/Engine/Source/Programs/HeadlessChaos/Private/HeadlessChaosTestSim.cpp b/Engine/Source/Programs/HeadlessChaos/Private/HeadlessChaosTestSim.cpp
--- a/Engine/Source/Programs/HeadlessChaos/Private/HeadlessChaosTestSim.cpp
+++ b/Engine/Source/Programs/HeadlessChaos/Private/HeadlessChaosTestSim.cpp
@@ -61,7 +61,8 @@ namespace ChaosTest {
 
 	GTEST_TEST(AllEvolutions, SimTests_SphereSphereSimTest)
 	{
-		FPBDRigidsSOAs Particles;
+		FParticleUniqueIndicesMultithreaded UniqueIndices;
+		FPBDRigidsSOAs Particles(UniqueIndices);
 		THandleArray<FChaosPhysicsMaterial> PhysicalMaterials;
 		FPBDRigidsEvolutionGBF Evolution(Particles, PhysicalMaterials);
 		InitEvolutionSettings(Evolution);
@@ -104,7 +105,8 @@ namespace ChaosTest {
 
 	GTEST_TEST(AllEvolutions, SimTests_BoxBoxSimTest)
 	{
-		FPBDRigidsSOAs Particles;
+		FParticleUniqueIndicesMultithreaded UniqueIndices;
+		FPBDRigidsSOAs Particles(UniqueIndices);
 		THandleArray<FChaosPhysicsMaterial> PhysicalMaterials;
 		FPBDRigidsEvolutionGBF Evolution(Particles, PhysicalMaterials);
 		InitEvolutionSettings(Evolution);
@@ -148,7 +150,8 @@ namespace ChaosTest {
 	//
 	GTEST_TEST(AllEvolutions, DISABLED_SimTests_VeryLowInertiaSimTest)
 	{
-		FPBDRigidsSOAs Particles;
+		FParticleUniqueIndicesMultithreaded UniqueIndices;
+		FPBDRigidsSOAs Particles(UniqueIndices);
 		THandleArray<FChaosPhysicsMaterial> PhysicalMaterials;
 		FPBDRigidsEvolutionGBF Evolution(Particles, PhysicalMaterials);
 		InitEvolutionSettings(Evolution);
@@ -179,7 +182,8 @@ namespace ChaosTest {
 
 	GTEST_TEST(AllEvolutions, SimTests_SleepAndWakeSimTest)
 	{
-		FPBDRigidsSOAs Particles;
+		FParticleUniqueIndicesMultithreaded UniqueIndices;
+		FPBDRigidsSOAs Particles(UniqueIndices);
 		THandleArray<FChaosPhysicsMaterial> PhysicalMaterials;
 		FPBDRigidsEvolutionGBF Evolution(Particles, PhysicalMaterials);
 		InitEvolutionSettings(Evolution);
@@ -311,7 +315,8 @@ namespace ChaosTest {
 
 	GTEST_TEST(AllEvolutions, SimTests_SleepAndWakeSimTest3)
 	{
-		FPBDRigidsSOAs Particles;
+		FParticleUniqueIndicesMultithreaded UniqueIndices;
+		FPBDRigidsSOAs Particles(UniqueIndices);
 		THandleArray<FChaosPhysicsMaterial> PhysicalMaterials;
 		FPBDRigidsEvolutionGBF Evolution(Particles, PhysicalMaterials);
 		InitEvolutionSettings(Evolution);
diff --git a/Engine/Source/Programs/HeadlessChaos/Private/HeadlessChaosTestVehicles.cpp b/Engine/Source/Programs/HeadlessChaos/Private/HeadlessChaosTestVehicles.cpp
--- a/Engine/Source/Programs/HeadlessChaos/Private/HeadlessChaosTestVehicles.cpp
+++ b/Engine/Source/Programs/HeadlessChaos/Private/HeadlessChaosTestVehicles.cpp
@@ -1305,7 +1305,8 @@ namespace ChaosTest
 	// #todo: break out vehicle simulation setup so it can be used across number of tests
 	GTEST_TEST(AllEvolutions, VehicleTest_SuspensionSpringLoad)
 	{
-		FPBDRigidsSOAs Particles;
+		FParticleUniqueIndicesMultithreaded UniqueIndices;
+		FPBDRigidsSOAs Particles(UniqueIndices);
 		THandleArray<FChaosPhysicsMaterial> PhysicalMaterials;
 		FPBDRigidsEvolutionGBF Evolution(Particles, PhysicalMaterials);
 
diff --git a/Engine/Source/Programs/HeadlessChaos/Public/HeadlessChaosCollisionConstraints.h b/Engine/Source/Programs/HeadlessChaos/Public/HeadlessChaosCollisionConstraints.h
--- a/Engine/Source/Programs/HeadlessChaos/Public/HeadlessChaosCollisionConstraints.h
+++ b/Engine/Source/Programs/HeadlessChaos/Public/HeadlessChaosCollisionConstraints.h
@@ -34,7 +34,8 @@ public:
 	using FAccelerationStructure = TBoundingVolume<FAccelerationStructureHandle>;
 
 	FPBDCollisionConstraintAccessor()
-		: SpatialAcceleration(EmptyParticles.GetNonDisabledView())
+		: EmptyParticles(UniqueIndices)
+		, SpatialAcceleration(EmptyParticles.GetNonDisabledView())
 		, BroadPhase(EmptyParticles, (FReal)1, (FReal)0, (FReal)0)
 		, CollisionConstraints(EmptyParticles, EmptyCollided, EmptyPhysicsMaterials, EmptyUniquePhysicsMaterials, 1, 1)
 		, CollisionDetector(BroadPhase, NarrowPhase, CollisionConstraints)
@@ -42,7 +43,8 @@ public:
 
 	FPBDCollisionConstraintAccessor(const FPBDRigidsSOAs& InParticles, TArrayCollectionArray<bool>& Collided, const TArrayCollectionArray<TSerializablePtr<FChaosPhysicsMaterial>>& PerParticleMaterials, const TArrayCollectionArray<TUniquePtr<FChaosPhysicsMaterial>>& PerParticleUniqueMaterials,
 	const int32 PushOutIterations, const int32 PushOutPairIterations) 
-		: SpatialAcceleration(InParticles.GetNonDisabledView())
+		: EmptyParticles(UniqueIndices)
+		, SpatialAcceleration(InParticles.GetNonDisabledView())
 		, BroadPhase(InParticles, (FReal)1, (FReal)0, (FReal)0)
 		, CollisionConstraints(InParticles, Collided, PerParticleMaterials, PerParticleUniqueMaterials, 1, 1)
 		, CollisionDetector(BroadPhase, NarrowPhase, CollisionConstraints)
@@ -114,6 +116,7 @@ public:
 	}
 
 	FPointContactConstraint EmptyConstraint;
+	FParticleUniqueIndicesMultithreaded UniqueIndices;
 	FPBDRigidsSOAs EmptyParticles;
 	TArrayCollectionArray<bool> EmptyCollided;
 	TArrayCollectionArray<TSerializablePtr<FChaosPhysicsMaterial>> EmptyPhysicsMaterials;
diff --git a/Engine/Source/Programs/HeadlessChaos/Public/HeadlessChaosTestConstraints.h b/Engine/Source/Programs/HeadlessChaos/Public/HeadlessChaosTestConstraints.h
--- a/Engine/Source/Programs/HeadlessChaos/Public/HeadlessChaosTestConstraints.h
+++ b/Engine/Source/Programs/HeadlessChaos/Public/HeadlessChaosTestConstraints.h
@@ -21,7 +21,8 @@ namespace ChaosTest
 	public:
 
 		FConstraintsTest(const int32 NumIterations, const FReal Gravity)
-			: Evolution(SOAs, PhysicalMaterials)
+			: SOAs(UniqueIndices)
+			, Evolution(SOAs, PhysicalMaterials)
 		{
 			PhysicalMaterial = MakeUnique<FChaosPhysicsMaterial>();
 			PhysicalMaterial->Friction = 0;
@@ -76,6 +77,7 @@ namespace ChaosTest
 		}
 
 		// Solver state
+		FParticleUniqueIndicesMultithreaded UniqueIndices;
 		FPBDRigidsSOAs SOAs;
 		TEvolution Evolution;
 		TUniquePtr<FChaosPhysicsMaterial> PhysicalMaterial;
