Chaos - fix post-EPA transform properly and add unit tests

diff --git a/Engine/Source/Programs/HeadlessChaos/Private/HeadlessChaosTestGJK.cpp b/Engine/Source/Programs/HeadlessChaos/Private/HeadlessChaosTestGJK.cpp
--- a/Engine/Source/Programs/HeadlessChaos/Private/HeadlessChaosTestGJK.cpp
+++ b/Engine/Source/Programs/HeadlessChaos/Private/HeadlessChaosTestGJK.cpp
@@ -1868,8 +1868,11 @@ namespace ChaosTest
 
 		FGJKSimplexData WarmStartData;
 
+		// Separated (GJK)
 		GJKPenetrationWarmStartable<true>(Box, Box, BTM.GetRelativeTransformNoScale(ATM), FReal(0), FReal(0), Penetration, ClosestA, ClosestB, NormalA, NormalB, WarmStartData);
 		EXPECT_NEAR(Penetration, -5.0f, KINDA_SMALL_NUMBER);
+		EXPECT_NEAR(ClosestA.Z, (FReal)50.0f, (FReal)KINDA_SMALL_NUMBER);
+		EXPECT_NEAR(ClosestB.Z, (FReal)-50.0f, (FReal)KINDA_SMALL_NUMBER);
 
 		GJKPenetrationWarmStartable<true>(Box, Box, BTM.GetRelativeTransformNoScale(ATM), FReal(0), FReal(0), Penetration, ClosestA, ClosestB, NormalA, NormalB, WarmStartData);
 		EXPECT_NEAR(Penetration, -5.0f, KINDA_SMALL_NUMBER);
@@ -1886,4 +1889,62 @@ namespace ChaosTest
 		GJKPenetrationWarmStartable<true>(Box, Box, BTM.GetRelativeTransformNoScale(ATM), FReal(0), FReal(0), Penetration2, ClosestA, ClosestB, NormalA, NormalB, WarmStartData);
 		EXPECT_NEAR(Penetration2, Penetration, KINDA_SMALL_NUMBER);
 	}
+
+	GTEST_TEST(GJKTests, BoxBoxWarmStartDeepTest)
+	{
+		FAABB3 Box({ -50, -50, -50 }, { 50, 50, 50 });
+
+		FRigidTransform3 ATM = FRigidTransform3::Identity;
+		FRigidTransform3 BTM = FRigidTransform3(FVec3(0, 0, 60), FRotation3::FromIdentity());
+		FVec3 ClosestA, ClosestB, NormalA, NormalB;
+		FReal Penetration;
+
+		FGJKSimplexData WarmStartData;
+
+		// Deep (EPA) No Margin
+		GJKPenetrationWarmStartable<true>(Box, Box, BTM.GetRelativeTransformNoScale(ATM), FReal(0), FReal(0), Penetration, ClosestA, ClosestB, NormalA, NormalB, WarmStartData);
+		EXPECT_NEAR(Penetration, 40.0f, KINDA_SMALL_NUMBER);
+		EXPECT_NEAR(ClosestA.Z, (FReal)50.0f, (FReal)KINDA_SMALL_NUMBER);
+		EXPECT_NEAR(ClosestB.Z, (FReal)-50.0f, (FReal)KINDA_SMALL_NUMBER);
+
+		// Deep (EPA) With Margin
+		WarmStartData = FGJKSimplexData();
+		TGJKCoreShape MarginBox(Box, FReal(10));
+		GJKPenetrationWarmStartable<true>(MarginBox, MarginBox, BTM.GetRelativeTransformNoScale(ATM), FReal(0), FReal(0), Penetration, ClosestA, ClosestB, NormalA, NormalB, WarmStartData);
+		EXPECT_NEAR(Penetration, 40.0f, KINDA_SMALL_NUMBER);
+		EXPECT_NEAR(ClosestA.Z, (FReal)50.0f, (FReal)KINDA_SMALL_NUMBER);
+		EXPECT_NEAR(ClosestB.Z, (FReal)-50.0f, (FReal)KINDA_SMALL_NUMBER);
+
+		// Deep (EPA) With Margin and Relative Rotation
+		WarmStartData = FGJKSimplexData();
+		ATM = FRigidTransform3(FVec3(0, 0, 0), FRotation3::FromAxisAngle(FVec3(1, 0, 0), FMath::DegreesToRadians(180)));
+		GJKPenetrationWarmStartable<true>(MarginBox, MarginBox, BTM.GetRelativeTransformNoScale(ATM), FReal(0), FReal(0), Penetration, ClosestA, ClosestB, NormalA, NormalB, WarmStartData);
+		EXPECT_NEAR(Penetration, 40.0f, KINDA_SMALL_NUMBER);
+		EXPECT_NEAR(ClosestA.Z, (FReal)-50.0f, (FReal)KINDA_SMALL_NUMBER);
+		EXPECT_NEAR(ClosestB.Z, (FReal)-50.0f, (FReal)KINDA_SMALL_NUMBER);
+	}
+
+	GTEST_TEST(GJKTests, SphereSphereWarmStartTest)
+	{
+		FImplicitSphere3 Sphere(FVec3(0), FReal(50));
+
+		FRigidTransform3 ATM = FRigidTransform3::Identity;
+		FRigidTransform3 BTM = FRigidTransform3(FVec3(0, 0, 105), FRotation3::FromIdentity());
+		FVec3 ClosestA, ClosestB, NormalA, NormalB;
+		FReal Penetration;
+
+		FGJKSimplexData WarmStartData;
+
+		GJKPenetrationWarmStartable<true>(Sphere, Sphere, BTM.GetRelativeTransformNoScale(ATM), FReal(0), FReal(0), Penetration, ClosestA, ClosestB, NormalA, NormalB, WarmStartData);
+		EXPECT_NEAR(Penetration, (FReal)-5.0f, (FReal)KINDA_SMALL_NUMBER);
+		EXPECT_NEAR(ClosestA.Z, (FReal)50.0f, (FReal)KINDA_SMALL_NUMBER);
+		EXPECT_NEAR(ClosestB.Z, (FReal)-50.0f, (FReal)KINDA_SMALL_NUMBER);
+
+		BTM = FRigidTransform3(FVec3(0, 0, 105), FRotation3::FromAxisAngle(FVec3(1,0,0), FMath::DegreesToRadians(180)));
+		GJKPenetrationWarmStartable<true>(Sphere, Sphere, BTM.GetRelativeTransformNoScale(ATM), FReal(0), FReal(0), Penetration, ClosestA, ClosestB, NormalA, NormalB, WarmStartData);
+		EXPECT_NEAR(Penetration, (FReal)-5.0f, (FReal)KINDA_SMALL_NUMBER);
+		EXPECT_NEAR(ClosestA.Z, (FReal)50.0f, (FReal)KINDA_SMALL_NUMBER);
+		EXPECT_NEAR(ClosestB.Z, (FReal)50.0f, (FReal)KINDA_SMALL_NUMBER);
+
+	}
 }
\ No newline at end of file
diff --git a/Engine/Source/Runtime/Experimental/Chaos/Public/Chaos/GJK.h b/Engine/Source/Runtime/Experimental/Chaos/Public/Chaos/GJK.h
--- a/Engine/Source/Runtime/Experimental/Chaos/Public/Chaos/GJK.h
+++ b/Engine/Source/Runtime/Experimental/Chaos/Public/Chaos/GJK.h
@@ -319,7 +319,7 @@ namespace Chaos
 				OutNormalB = BToATM.InverseTransformVectorNoScale(MTD);
 				OutPenetration = Penetration + ThicknessA + ThicknessB;
 				OutClosestA = ClosestA + MTD * ThicknessA;
-				OutClosestB = BToATM.InverseTransformPositionNoScale(ClosestBInA) - MTD * ThicknessB;
+				OutClosestB = BToATM.InverseTransformPositionNoScale(ClosestBInA - MTD * ThicknessB);
 				return true;
 			case EEPAResult::BadInitialSimplex:
 				// The origin is outside the simplex. Must be a touching contact and EPA setup will have calculated the normal
