Fix headlesschaos CIS

diff --git a/Engine/Source/Programs/HeadlessChaos/Private/HeadlessChaosTestGraph.cpp b/Engine/Source/Programs/HeadlessChaos/Private/HeadlessChaosTestGraph.cpp
--- a/Engine/Source/Programs/HeadlessChaos/Private/HeadlessChaosTestGraph.cpp
+++ b/Engine/Source/Programs/HeadlessChaos/Private/HeadlessChaosTestGraph.cpp
@@ -528,7 +528,7 @@ namespace ChaosTest {
 			Graph.GetIslandGraph()->InitSorting();
 
 			Graph.GetIslandGraph()->ComputeLevels(ContainerId);
-			Graph.GetIslandGraph()->ComputeColors(ContainerId, INT32_MAX);
+			Graph.GetIslandGraph()->ComputeColors(ContainerId, 0);
 
 			// Assign color to constraints
 			GraphColor.InitializeColor(Graph);
