Chaos - fix healdess chaos levelset collision test

diff --git a/Engine/Source/Programs/HeadlessChaos/Private/HeadlessChaosTestCollisions.cpp b/Engine/Source/Programs/HeadlessChaos/Private/HeadlessChaosTestCollisions.cpp
--- a/Engine/Source/Programs/HeadlessChaos/Private/HeadlessChaosTestCollisions.cpp
+++ b/Engine/Source/Programs/HeadlessChaos/Private/HeadlessChaosTestCollisions.cpp
@@ -52,7 +52,7 @@ namespace ChaosTest {
 		Box1->UpdateWorldSpaceState(FRigidTransform3(Box1->X(), Box1->R()), FVec3(0));
 
 		auto Box2 = AppendDynamicParticleBox(Particles);
-		Box2->X() = FVec3(1.5f, 1.5f, 1.9f);
+		Box2->X() = FVec3(0.5f, 0.5f, 1.9f);
 		Box2->R() = FRotation3(FQuat::Identity);
 		Box2->P() = Box2->X();
 		Box2->Q() = Box2->R();
@@ -76,8 +76,8 @@ namespace ChaosTest {
 		EXPECT_TRUE(Constraint.GetNormal().operator==(FVec3(0, 0, 1)));
 
 		// The contact point is the average of the surface contact points, so it should be half of Phi inside particle0
-		Chaos::FReal Distance = ChaosTest::SignedDistance(*Constraint.Particle[0], Constraint.GetLocation());
-		EXPECT_NEAR(Distance, 0.5f * Constraint.GetPhi(), KINDA_SMALL_NUMBER);
+		EXPECT_NEAR(Constraint.GetPhi(), -0.1f, KINDA_SMALL_NUMBER);
+		EXPECT_NEAR(Constraint.GetLocation().Z, 1.45f, KINDA_SMALL_NUMBER);
 	}
 
 	void LevelsetConstraintGJK()
diff --git a/Engine/Source/Programs/HeadlessChaos/Public/HeadlessChaosCollisionConstraints.h b/Engine/Source/Programs/HeadlessChaos/Public/HeadlessChaosCollisionConstraints.h
--- a/Engine/Source/Programs/HeadlessChaos/Public/HeadlessChaosCollisionConstraints.h
+++ b/Engine/Source/Programs/HeadlessChaos/Public/HeadlessChaosCollisionConstraints.h
@@ -99,6 +99,7 @@ public:
 		FRigidTransform3 WorldTransform0 = Constraint.ImplicitTransform[0] * Collisions::GetTransform(Constraint.Particle[0]);
 		FRigidTransform3 WorldTransform1 = Constraint.ImplicitTransform[1] * Collisions::GetTransform(Constraint.Particle[1]);
 
+		Constraint.ResetPhi(TNumericLimits<FReal>::Max());
 		Collisions::UpdateLevelsetLevelsetConstraint<ECollisionUpdateType::Deepest>(WorldTransform0, WorldTransform1, FReal(1 / 30.0f), Constraint);
 	}
 
diff --git a/Engine/Source/Runtime/Experimental/Chaos/Private/Chaos/CollisionResolution.cpp b/Engine/Source/Runtime/Experimental/Chaos/Private/Chaos/CollisionResolution.cpp
--- a/Engine/Source/Runtime/Experimental/Chaos/Private/Chaos/CollisionResolution.cpp
+++ b/Engine/Source/Runtime/Experimental/Chaos/Private/Chaos/CollisionResolution.cpp
@@ -1748,8 +1748,9 @@ namespace Chaos
 
 				if (ContactPoint.IsSet())
 				{
-					ContactPoint.ShapeContactPoints[0] = WorldTransform0.InverseTransformPosition(ContactPoint.Location);
-					ContactPoint.ShapeContactPoints[1] = WorldTransform1.InverseTransformPosition(ContactPoint.Location - ContactPoint.Phi * ContactPoint.Normal);
+					// @todo(chaos): this work is partially duplicated from SampleObject - fix this
+					ContactPoint.ShapeContactPoints[0] = WorldTransform0.InverseTransformPosition(ContactPoint.Location + FReal(0.5) * ContactPoint.Phi * ContactPoint.Normal);
+					ContactPoint.ShapeContactPoints[1] = WorldTransform1.InverseTransformPosition(ContactPoint.Location - FReal(0.5) * ContactPoint.Phi * ContactPoint.Normal);
 					ContactPoint.ShapeContactNormal = WorldTransform1.InverseTransformVector(ContactPoint.Normal);
 				}
 
diff --git a/Engine/Source/Runtime/Experimental/Chaos/Private/Chaos/CollisionResolutionUtil.cpp b/Engine/Source/Runtime/Experimental/Chaos/Private/Chaos/CollisionResolutionUtil.cpp
--- a/Engine/Source/Runtime/Experimental/Chaos/Private/Chaos/CollisionResolutionUtil.cpp
+++ b/Engine/Source/Runtime/Experimental/Chaos/Private/Chaos/CollisionResolutionUtil.cpp
@@ -625,7 +625,7 @@ namespace Chaos
 					if (NewPhi < Contact.Phi)
 					{
 						Contact.Phi = NewPhi;
-						Contact.Location = ObjectTransform.TransformPositionNoScale(LocalPoint);
+						Contact.Location = ObjectTransform.TransformPositionNoScale(LocalPoint - FReal(0.5) * NewPhi * LocalNormal);
 						Contact.Normal = ObjectTransform.TransformVectorNoScale(LocalNormal);
 					}
 				}
@@ -640,7 +640,7 @@ namespace Chaos
 				FVec3 LocalPoint = SampleToObjectTM.TransformPositionNoScale(SampleParticles.X(DeepestParticle));
 				FVec3 LocalNormal;
 				Contact.Phi = Object.PhiWithNormal(LocalPoint, LocalNormal);
-				Contact.Location = ObjectTransform.TransformPositionNoScale(LocalPoint);
+				Contact.Location = ObjectTransform.TransformPositionNoScale(LocalPoint - FReal(0.5) * Contact.Phi * LocalNormal);
 				Contact.Normal = ObjectTransform.TransformVectorNoScale(LocalNormal);
 
 			}
