Improve contact plane selection

diff --git a/Engine/Source/Programs/HeadlessChaos/Private/HeadlessChaosTestEPA.cpp b/Engine/Source/Programs/HeadlessChaos/Private/HeadlessChaosTestEPA.cpp
--- a/Engine/Source/Programs/HeadlessChaos/Private/HeadlessChaosTestEPA.cpp
+++ b/Engine/Source/Programs/HeadlessChaos/Private/HeadlessChaosTestEPA.cpp
@@ -421,9 +421,10 @@ namespace ChaosTest
 		{
 			//sphere with deep penetration to make sure we have max iterations
 			TSphere<FReal,3> Sphere(FVec3(0), 10);
-			auto Support = [&Sphere](const auto& V)
+			int32 VertexIndex = INDEX_NONE;
+			auto Support = [&Sphere,&VertexIndex](const auto& V)
 			{
-				return Sphere.Support(V, 0);
+				return Sphere.Support(V, 0,VertexIndex);
 			};
 
 			TArray<FVec3> Tetrahedron = { 
@@ -444,9 +445,10 @@ namespace ChaosTest
 		{
 			//capsule with origin in middle
 			FCapsule Capsule(FVec3(0, 0, 10), FVec3(0, 0, -10), 3);
-			auto Support = [&Capsule](const auto& V)
+			int32 VertexIndex = INDEX_NONE;
+			auto Support = [&Capsule,&VertexIndex](const auto& V)
 			{
-				return Capsule.Support(V, 0);
+				return Capsule.Support(V, 0, VertexIndex);
 			};
 
 			TArray<FVec3> Tetrahedron = { 
@@ -465,9 +467,10 @@ namespace ChaosTest
 		{
 			//capsule with origin near top
 			FCapsule Capsule(FVec3(0, 0, -2), FVec3(0, 0, -12), 3);
-			auto Support = [&Capsule](const auto& V)
+			int32 VertexIndex = INDEX_NONE;
+			auto Support = [&Capsule,&VertexIndex](const auto& V)
 			{
-				return Capsule.Support(V, 0);
+				return Capsule.Support(V, 0, VertexIndex);
 			};
 
 			TArray<FVec3> Tetrahedron = { 
diff --git a/Engine/Source/Programs/HeadlessChaos/Private/HeadlessChaosTestGJK.cpp b/Engine/Source/Programs/HeadlessChaos/Private/HeadlessChaosTestGJK.cpp
--- a/Engine/Source/Programs/HeadlessChaos/Private/HeadlessChaosTestGJK.cpp
+++ b/Engine/Source/Programs/HeadlessChaos/Private/HeadlessChaosTestGJK.cpp
@@ -1868,26 +1868,28 @@ namespace ChaosTest
 
 		FGJKSimplexData WarmStartData;
 		FReal SupportDelta = FReal(0);
+		int32 VertexIndexA = INDEX_NONE;
+		int32 VertexIndexB = INDEX_NONE;
 
 		// Separated (GJK)
-		GJKPenetrationWarmStartable<true>(Box, Box, BTM.GetRelativeTransformNoScale(ATM), FReal(0), FReal(0), Penetration, ClosestA, ClosestB, NormalA, NormalB, WarmStartData, SupportDelta);
+		GJKPenetrationWarmStartable<true>(Box, Box, BTM.GetRelativeTransformNoScale(ATM), FReal(0), FReal(0), Penetration, ClosestA, ClosestB, NormalA, NormalB, VertexIndexA, VertexIndexB, WarmStartData, SupportDelta);
 		EXPECT_NEAR(Penetration, -5.0f, KINDA_SMALL_NUMBER);
 		EXPECT_NEAR(ClosestA.Z, (FReal)50.0f, (FReal)KINDA_SMALL_NUMBER);
 		EXPECT_NEAR(ClosestB.Z, (FReal)-50.0f, (FReal)KINDA_SMALL_NUMBER);
 
-		GJKPenetrationWarmStartable<true>(Box, Box, BTM.GetRelativeTransformNoScale(ATM), FReal(0), FReal(0), Penetration, ClosestA, ClosestB, NormalA, NormalB, WarmStartData, SupportDelta);
+		GJKPenetrationWarmStartable<true>(Box, Box, BTM.GetRelativeTransformNoScale(ATM), FReal(0), FReal(0), Penetration, ClosestA, ClosestB, NormalA, NormalB, VertexIndexA, VertexIndexB, WarmStartData, SupportDelta);
 		EXPECT_NEAR(Penetration, -5.0f, KINDA_SMALL_NUMBER);
 
 		BTM = FRigidTransform3(FVec3(0, 0, 145), FRotation3::FromIdentity());
-		GJKPenetrationWarmStartable<true>(Box, Box, BTM.GetRelativeTransformNoScale(ATM), FReal(0), FReal(0), Penetration, ClosestA, ClosestB, NormalA, NormalB, WarmStartData, SupportDelta);
+		GJKPenetrationWarmStartable<true>(Box, Box, BTM.GetRelativeTransformNoScale(ATM), FReal(0), FReal(0), Penetration, ClosestA, ClosestB, NormalA, NormalB, VertexIndexA, VertexIndexB, WarmStartData, SupportDelta);
 		EXPECT_NEAR(Penetration, -45.0f, KINDA_SMALL_NUMBER);
 
 		BTM = FRigidTransform3(FVec3(0, 0, 145), FRotation3::FromAxisAngle(FVec3(1, 0, 0), FMath::DegreesToRadians(110.0f)));
-		GJKPenetrationWarmStartable<true>(Box, Box, BTM.GetRelativeTransformNoScale(ATM), FReal(0), FReal(0), Penetration, ClosestA, ClosestB, NormalA, NormalB, WarmStartData, SupportDelta);
+		GJKPenetrationWarmStartable<true>(Box, Box, BTM.GetRelativeTransformNoScale(ATM), FReal(0), FReal(0), Penetration, ClosestA, ClosestB, NormalA, NormalB, VertexIndexA, VertexIndexB, WarmStartData, SupportDelta);
 		EXPECT_NEAR(Penetration, -30.9144f, KINDA_SMALL_NUMBER);
 
 		FReal Penetration2;
-		GJKPenetrationWarmStartable<true>(Box, Box, BTM.GetRelativeTransformNoScale(ATM), FReal(0), FReal(0), Penetration2, ClosestA, ClosestB, NormalA, NormalB, WarmStartData, SupportDelta);
+		GJKPenetrationWarmStartable<true>(Box, Box, BTM.GetRelativeTransformNoScale(ATM), FReal(0), FReal(0), Penetration2, ClosestA, ClosestB, NormalA, NormalB, VertexIndexA, VertexIndexB, WarmStartData, SupportDelta);
 		EXPECT_NEAR(Penetration2, Penetration, KINDA_SMALL_NUMBER);
 	}
 
@@ -1903,8 +1905,11 @@ namespace ChaosTest
 		FGJKSimplexData WarmStartData;
 		FReal SupportDelta = FReal(0);
 
+		int32 VertexIndexA = INDEX_NONE;
+		int32 VertexIndexB = INDEX_NONE;
+
 		// Deep (EPA) No Margin
-		GJKPenetrationWarmStartable<true>(Box, Box, BTM.GetRelativeTransformNoScale(ATM), FReal(0), FReal(0), Penetration, ClosestA, ClosestB, NormalA, NormalB, WarmStartData, SupportDelta);
+		GJKPenetrationWarmStartable<true>(Box, Box, BTM.GetRelativeTransformNoScale(ATM), FReal(0), FReal(0), Penetration, ClosestA, ClosestB, NormalA, NormalB, VertexIndexA, VertexIndexB,  WarmStartData, SupportDelta);
 		EXPECT_NEAR(Penetration, 40.0f, KINDA_SMALL_NUMBER);
 		EXPECT_NEAR(ClosestA.Z, (FReal)50.0f, (FReal)KINDA_SMALL_NUMBER);
 		EXPECT_NEAR(ClosestB.Z, (FReal)-50.0f, (FReal)KINDA_SMALL_NUMBER);
@@ -1912,7 +1917,7 @@ namespace ChaosTest
 		// Deep (EPA) With Margin
 		WarmStartData = FGJKSimplexData();
 		TGJKCoreShape MarginBox(Box, FReal(10));
-		GJKPenetrationWarmStartable<true>(MarginBox, MarginBox, BTM.GetRelativeTransformNoScale(ATM), FReal(0), FReal(0), Penetration, ClosestA, ClosestB, NormalA, NormalB, WarmStartData, SupportDelta);
+		GJKPenetrationWarmStartable<true>(MarginBox, MarginBox, BTM.GetRelativeTransformNoScale(ATM), FReal(0), FReal(0), Penetration, ClosestA, ClosestB, NormalA, NormalB,  VertexIndexA, VertexIndexB, WarmStartData, SupportDelta);
 		EXPECT_NEAR(Penetration, 40.0f, KINDA_SMALL_NUMBER);
 		EXPECT_NEAR(ClosestA.Z, (FReal)50.0f, (FReal)KINDA_SMALL_NUMBER);
 		EXPECT_NEAR(ClosestB.Z, (FReal)-50.0f, (FReal)KINDA_SMALL_NUMBER);
@@ -1920,7 +1925,7 @@ namespace ChaosTest
 		// Deep (EPA) With Margin and Relative Rotation
 		WarmStartData = FGJKSimplexData();
 		ATM = FRigidTransform3(FVec3(0, 0, 0), FRotation3::FromAxisAngle(FVec3(1, 0, 0), FMath::DegreesToRadians(180)));
-		GJKPenetrationWarmStartable<true>(MarginBox, MarginBox, BTM.GetRelativeTransformNoScale(ATM), FReal(0), FReal(0), Penetration, ClosestA, ClosestB, NormalA, NormalB, WarmStartData, SupportDelta);
+		GJKPenetrationWarmStartable<true>(MarginBox, MarginBox, BTM.GetRelativeTransformNoScale(ATM), FReal(0), FReal(0), Penetration, ClosestA, ClosestB, NormalA, NormalB,  VertexIndexA, VertexIndexB, WarmStartData, SupportDelta);
 		EXPECT_NEAR(Penetration, 40.0f, KINDA_SMALL_NUMBER);
 		EXPECT_NEAR(ClosestA.Z, (FReal)-50.0f, (FReal)KINDA_SMALL_NUMBER);
 		EXPECT_NEAR(ClosestB.Z, (FReal)-50.0f, (FReal)KINDA_SMALL_NUMBER);
@@ -1937,14 +1942,17 @@ namespace ChaosTest
 
 		FGJKSimplexData WarmStartData;
 		FReal SupportDelta = FReal(0);
+		
+		int32 VertexIndexA = INDEX_NONE;
+		int32 VertexIndexB = INDEX_NONE;
 
-		GJKPenetrationWarmStartable<true>(Sphere, Sphere, BTM.GetRelativeTransformNoScale(ATM), FReal(0), FReal(0), Penetration, ClosestA, ClosestB, NormalA, NormalB, WarmStartData, SupportDelta);
+		GJKPenetrationWarmStartable<true>(Sphere, Sphere, BTM.GetRelativeTransformNoScale(ATM), FReal(0), FReal(0), Penetration, ClosestA, ClosestB, NormalA, NormalB,  VertexIndexA, VertexIndexB, WarmStartData, SupportDelta);
 		EXPECT_NEAR(Penetration, (FReal)-5.0f, (FReal)KINDA_SMALL_NUMBER);
 		EXPECT_NEAR(ClosestA.Z, (FReal)50.0f, (FReal)KINDA_SMALL_NUMBER);
 		EXPECT_NEAR(ClosestB.Z, (FReal)-50.0f, (FReal)KINDA_SMALL_NUMBER);
 
 		BTM = FRigidTransform3(FVec3(0, 0, 105), FRotation3::FromAxisAngle(FVec3(1,0,0), FMath::DegreesToRadians(180)));
-		GJKPenetrationWarmStartable<true>(Sphere, Sphere, BTM.GetRelativeTransformNoScale(ATM), FReal(0), FReal(0), Penetration, ClosestA, ClosestB, NormalA, NormalB, WarmStartData, SupportDelta);
+		GJKPenetrationWarmStartable<true>(Sphere, Sphere, BTM.GetRelativeTransformNoScale(ATM), FReal(0), FReal(0), Penetration, ClosestA, ClosestB, NormalA, NormalB, VertexIndexA, VertexIndexB, WarmStartData, SupportDelta);
 		EXPECT_NEAR(Penetration, (FReal)-5.0f, (FReal)KINDA_SMALL_NUMBER);
 		EXPECT_NEAR(ClosestA.Z, (FReal)50.0f, (FReal)KINDA_SMALL_NUMBER);
 		EXPECT_NEAR(ClosestB.Z, (FReal)50.0f, (FReal)KINDA_SMALL_NUMBER);
diff --git a/Engine/Source/Programs/HeadlessChaos/Private/HeadlessChaosTestImplicits.cpp b/Engine/Source/Programs/HeadlessChaos/Private/HeadlessChaosTestImplicits.cpp
--- a/Engine/Source/Programs/HeadlessChaos/Private/HeadlessChaosTestImplicits.cpp
+++ b/Engine/Source/Programs/HeadlessChaos/Private/HeadlessChaosTestImplicits.cpp
@@ -122,21 +122,22 @@ namespace ChaosTest {
 	void UnitImplicitObjectSupportPhis(GeometryType& Subject, FString Caller)
 	{
 		FString Error = FString("Called by ") + Caller + FString(".");
+		int32 VertexIndex = INDEX_NONE;
 
 		// support phi
-		EXPECT_VECTOR_NEAR_ERR(Subject.Support(FVec3(0, 0, 1), FReal(0)), (FVec3(0, 0, 1)), KINDA_SMALL_NUMBER, Error);
-		EXPECT_VECTOR_NEAR_ERR(Subject.Support(FVec3(0, 0, -1), FReal(0)), (FVec3(0, 0, -1)), KINDA_SMALL_NUMBER, Error);
-		EXPECT_VECTOR_NEAR_ERR(Subject.Support(FVec3(0, 1, 0), FReal(0)), (FVec3(0, 1, 0)), KINDA_SMALL_NUMBER, Error);
-		EXPECT_VECTOR_NEAR_ERR(Subject.Support(FVec3(0, -1, 0), FReal(0)), (FVec3(0, -1, 0)), KINDA_SMALL_NUMBER, Error);
-		EXPECT_VECTOR_NEAR_ERR(Subject.Support(FVec3(1, 0, 0), FReal(0)), (FVec3(1, 0, 0)), KINDA_SMALL_NUMBER, Error);
-		EXPECT_VECTOR_NEAR_ERR(Subject.Support(FVec3(-1, 0, 0), FReal(0)), (FVec3(-1, 0, 0)), KINDA_SMALL_NUMBER, Error);
-
-		EXPECT_VECTOR_NEAR_ERR(Subject.Support(FVec3(0, 0, 1), FReal(1)), (FVec3(0, 0, 2)), KINDA_SMALL_NUMBER, Error);
-		EXPECT_VECTOR_NEAR_ERR(Subject.Support(FVec3(0, 0, -1), FReal(1)), (FVec3(0, 0, -2)), KINDA_SMALL_NUMBER, Error);
-		EXPECT_VECTOR_NEAR_ERR(Subject.Support(FVec3(0, 1, 0), FReal(1)), (FVec3(0, 2, 0)), KINDA_SMALL_NUMBER, Error);
-		EXPECT_VECTOR_NEAR_ERR(Subject.Support(FVec3(0, -1, 0), FReal(1)), (FVec3(0, -2, 0)), KINDA_SMALL_NUMBER, Error);
-		EXPECT_VECTOR_NEAR_ERR(Subject.Support(FVec3(1, 0, 0), FReal(1)), (FVec3(2, 0, 0)), KINDA_SMALL_NUMBER, Error);
-		EXPECT_VECTOR_NEAR_ERR(Subject.Support(FVec3(-1, 0, 0), FReal(1)), (FVec3(-2, 0, 0)), KINDA_SMALL_NUMBER, Error);
+		EXPECT_VECTOR_NEAR_ERR(Subject.Support(FVec3(0, 0, 1), FReal(0), VertexIndex), (FVec3(0, 0, 1)), KINDA_SMALL_NUMBER, Error);
+		EXPECT_VECTOR_NEAR_ERR(Subject.Support(FVec3(0, 0, -1), FReal(0), VertexIndex), (FVec3(0, 0, -1)), KINDA_SMALL_NUMBER, Error);
+		EXPECT_VECTOR_NEAR_ERR(Subject.Support(FVec3(0, 1, 0), FReal(0), VertexIndex), (FVec3(0, 1, 0)), KINDA_SMALL_NUMBER, Error);
+		EXPECT_VECTOR_NEAR_ERR(Subject.Support(FVec3(0, -1, 0), FReal(0), VertexIndex), (FVec3(0, -1, 0)), KINDA_SMALL_NUMBER, Error);
+		EXPECT_VECTOR_NEAR_ERR(Subject.Support(FVec3(1, 0, 0), FReal(0), VertexIndex), (FVec3(1, 0, 0)), KINDA_SMALL_NUMBER, Error);
+		EXPECT_VECTOR_NEAR_ERR(Subject.Support(FVec3(-1, 0, 0), FReal(0), VertexIndex), (FVec3(-1, 0, 0)), KINDA_SMALL_NUMBER, Error);
+
+		EXPECT_VECTOR_NEAR_ERR(Subject.Support(FVec3(0, 0, 1), FReal(1), VertexIndex), (FVec3(0, 0, 2)), KINDA_SMALL_NUMBER, Error);
+		EXPECT_VECTOR_NEAR_ERR(Subject.Support(FVec3(0, 0, -1), FReal(1), VertexIndex), (FVec3(0, 0, -2)), KINDA_SMALL_NUMBER, Error);
+		EXPECT_VECTOR_NEAR_ERR(Subject.Support(FVec3(0, 1, 0), FReal(1), VertexIndex), (FVec3(0, 2, 0)), KINDA_SMALL_NUMBER, Error);
+		EXPECT_VECTOR_NEAR_ERR(Subject.Support(FVec3(0, -1, 0), FReal(1), VertexIndex), (FVec3(0, -2, 0)), KINDA_SMALL_NUMBER, Error);
+		EXPECT_VECTOR_NEAR_ERR(Subject.Support(FVec3(1, 0, 0), FReal(1), VertexIndex), (FVec3(2, 0, 0)), KINDA_SMALL_NUMBER, Error);
+		EXPECT_VECTOR_NEAR_ERR(Subject.Support(FVec3(-1, 0, 0), FReal(1), VertexIndex), (FVec3(-2, 0, 0)), KINDA_SMALL_NUMBER, Error);
 	}
 
 	/* Takes 3 ImplictObject of unit size (circumscribed inside a 2x2 cube)
@@ -233,7 +234,8 @@ namespace ChaosTest {
 		UnitImplicitObjectIntersections(Subject, Caller);
 		
 		{// support phi - expects the corners for boxes
-			// Iterate through every face, edge, and corner direction, and ensure it snaps to the proper corner. 
+			// Iterate through every face, edge, and corner direction, and ensure it snaps to the proper corner.
+			int32 VertexIndex = INDEX_NONE;
 			for (int i0 = -1; i0 < 2; ++i0)
 			{
 				for (int i1 = -1; i1 < 2; ++i1)
@@ -250,7 +252,7 @@ namespace ChaosTest {
 						FString Error("Direction: ");
 						Error += FString::Printf(TEXT("(%d, %d, %d)"), i0, i1, i2);
 
-						EXPECT_VECTOR_NEAR_ERR(Subject.Support(FVec3(i0, i1, i2), FReal(0)), Expected, KINDA_SMALL_NUMBER, Error);
+						EXPECT_VECTOR_NEAR_ERR(Subject.Support(FVec3(i0, i1, i2), FReal(0), VertexIndex), Expected, KINDA_SMALL_NUMBER, Error);
 					}
 				}
 			}
@@ -267,22 +269,22 @@ namespace ChaosTest {
 
 		{// support phi off origin
 			TBox<FReal, 3> Subject2(FVec3(2), FVec3(4));
-
-			EXPECT_VECTOR_NEAR_DEFAULT(Subject2.Support(FVec3(0, 0, 1), FReal(0)), (FVec3(4, 4, 4)));
-			EXPECT_VECTOR_NEAR_DEFAULT(Subject2.Support(FVec3(0, 0, -1), FReal(0)), (FVec3(4, 4, 2)));
-			EXPECT_VECTOR_NEAR_DEFAULT(Subject2.Support(FVec3(0, 1, 0), FReal(0)), (FVec3(4, 4, 4)));
-			EXPECT_VECTOR_NEAR_DEFAULT(Subject2.Support(FVec3(0, -1, 0), FReal(0)), (FVec3(4, 2, 4)));
-			EXPECT_VECTOR_NEAR_DEFAULT(Subject2.Support(FVec3(0, 1, 0), FReal(0)), (FVec3(4, 4, 4)));
-			EXPECT_VECTOR_NEAR_DEFAULT(Subject2.Support(FVec3(1, 0, 0), FReal(0)), (FVec3(4, 4, 4)));
-			EXPECT_VECTOR_NEAR_DEFAULT(Subject2.Support(FVec3(-1, 0, 0), FReal(0)), (FVec3(2, 4, 4)));
+			int32 VertexIndex = INDEX_NONE;
+			EXPECT_VECTOR_NEAR_DEFAULT(Subject2.Support(FVec3(0, 0, 1), FReal(0), VertexIndex), (FVec3(4, 4, 4)));
+			EXPECT_VECTOR_NEAR_DEFAULT(Subject2.Support(FVec3(0, 0, -1), FReal(0), VertexIndex), (FVec3(4, 4, 2)));
+			EXPECT_VECTOR_NEAR_DEFAULT(Subject2.Support(FVec3(0, 1, 0), FReal(0), VertexIndex), (FVec3(4, 4, 4)));
+			EXPECT_VECTOR_NEAR_DEFAULT(Subject2.Support(FVec3(0, -1, 0), FReal(0), VertexIndex), (FVec3(4, 2, 4)));
+			EXPECT_VECTOR_NEAR_DEFAULT(Subject2.Support(FVec3(0, 1, 0), FReal(0), VertexIndex), (FVec3(4, 4, 4)));
+			EXPECT_VECTOR_NEAR_DEFAULT(Subject2.Support(FVec3(1, 0, 0), FReal(0), VertexIndex), (FVec3(4, 4, 4)));
+			EXPECT_VECTOR_NEAR_DEFAULT(Subject2.Support(FVec3(-1, 0, 0), FReal(0), VertexIndex), (FVec3(2, 4, 4)));
 
 #if RUN_KNOWN_BROKEN_TESTS
-			EXPECT_VECTOR_NEAR_DEFAULT(Subject2.Support(FVec3(0, 0, 1), FReal(1)), (FVec3(5, 5, 5)));
-			EXPECT_VECTOR_NEAR_DEFAULT(Subject2.Support(FVec3(0, 0, -1), FReal(1)), (FVec3(5, 5, 1)));
-			EXPECT_VECTOR_NEAR_DEFAULT(Subject2.Support(FVec3(0, 1, 0), FReal(1)), (FVec3(5, 5, 5)));
-			EXPECT_VECTOR_NEAR_DEFAULT(Subject2.Support(FVec3(0, -1, 0), FReal(1)), (FVec3(5, 1, 5)));
-			EXPECT_VECTOR_NEAR_DEFAULT(Subject2.Support(FVec3(1, 0, 0), FReal(1)), (FVec3(5, 5, 5)));
-			EXPECT_VECTOR_NEAR_DEFAULT(Subject2.Support(FVec3(-1, 0, 0), FReal(1)), (FVec3(1, 5, 5)));
+			EXPECT_VECTOR_NEAR_DEFAULT(Subject2.Support(FVec3(0, 0, 1), FReal(1), VertexIndex), (FVec3(5, 5, 5)));
+			EXPECT_VECTOR_NEAR_DEFAULT(Subject2.Support(FVec3(0, 0, -1), FReal(1), VertexIndex), (FVec3(5, 5, 1)));
+			EXPECT_VECTOR_NEAR_DEFAULT(Subject2.Support(FVec3(0, 1, 0), FReal(1), VertexIndex), (FVec3(5, 5, 5)));
+			EXPECT_VECTOR_NEAR_DEFAULT(Subject2.Support(FVec3(0, -1, 0), FReal(1), VertexIndex), (FVec3(5, 1, 5)));
+			EXPECT_VECTOR_NEAR_DEFAULT(Subject2.Support(FVec3(1, 0, 0), FReal(1), VertexIndex), (FVec3(5, 5, 5)));
+			EXPECT_VECTOR_NEAR_DEFAULT(Subject2.Support(FVec3(-1, 0, 0), FReal(1), VertexIndex), (FVec3(1, 5, 5)));
 #endif
 		}
 
@@ -812,21 +814,23 @@ namespace ChaosTest {
 		}
 		
 		{//support
+			int32 VertexIndex = INDEX_NONE;
+			
 			const FVec3 DirX(1, 0, 0);
-			FVec3 SupportPt = Unscaled.Support(DirX, 1);
+			FVec3 SupportPt = Unscaled.Support(DirX, 1, VertexIndex);
 			EXPECT_VECTOR_NEAR(SupportPt, FVec3(9, 0, 0), 0);
 
-			SupportPt = UniformScale.Support(DirX, 1);
+			SupportPt = UniformScale.Support(DirX, 1, VertexIndex);
 			EXPECT_VECTOR_NEAR(SupportPt, FVec3(17, 0, 0), 0);
 
 			const FVec3 DirZ(0, 0, -1);
-			SupportPt = UniformScale.Support(DirZ, 1);
+			SupportPt = UniformScale.Support(DirZ, 1, VertexIndex);
 			EXPECT_VECTOR_NEAR(SupportPt, FVec3(6, 0, -11), 0);
 
-			SupportPt = NonUniformScale.Support(DirX, 1);
+			SupportPt = NonUniformScale.Support(DirX, 1, VertexIndex);
 			EXPECT_VECTOR_NEAR(SupportPt, FVec3(17, 0, 0), 0);
 
-			SupportPt = NonUniformScale.Support(DirZ, 1);
+			SupportPt = NonUniformScale.Support(DirZ, 1, VertexIndex);
 			EXPECT_VECTOR_NEAR(SupportPt, FVec3(6, 0, -6), 0);
 		}
 
@@ -1075,16 +1079,17 @@ namespace ChaosTest {
 		{//support
 
 			const FVec3 DirX(1, 0, 0);
+			int32 VertexIndex = INDEX_NONE;
 			FVec3 SupportPt = Utilities::CastHelper(Untransformed, Identity, [&](const auto& Concrete, const auto& FullTM)
 			{
-				FVec3 SupportLocal = Concrete.Support(FullTM.InverseTransformVectorNoScale(DirX), 1);
+				FVec3 SupportLocal = Concrete.Support(FullTM.InverseTransformVectorNoScale(DirX), 1, VertexIndex);
 				return FullTM.TransformPosition(SupportLocal);
 			});
 			EXPECT_VECTOR_NEAR_DEFAULT(SupportPt, FVec3(9, 5, 5));
 
 			SupportPt = Utilities::CastHelper(Translated, Identity, [&](const auto& Concrete, const auto& FullTM)
 			{
-				FVec3 SupportLocal = Concrete.Support(FullTM.InverseTransformVectorNoScale(DirX), 1);
+				FVec3 SupportLocal = Concrete.Support(FullTM.InverseTransformVectorNoScale(DirX), 1, VertexIndex);
 				return FullTM.TransformPosition(SupportLocal);
 			});
 			EXPECT_VECTOR_NEAR_DEFAULT(SupportPt, FVec3(13, 5, 5));
@@ -1092,21 +1097,21 @@ namespace ChaosTest {
 			const FVec3 DirZ(0, 0, -1);
 			SupportPt = Utilities::CastHelper(Translated, Identity, [&](const auto& Concrete, const auto& FullTM)
 			{
-				FVec3 SupportLocal = Concrete.Support(FullTM.InverseTransformVectorNoScale(DirZ), 1);
+				FVec3 SupportLocal = Concrete.Support(FullTM.InverseTransformVectorNoScale(DirZ), 1, VertexIndex);
 				return FullTM.TransformPosition(SupportLocal);
 			});
 			EXPECT_VECTOR_NEAR_DEFAULT(SupportPt, FVec3(12, 5, -6));
 
 			SupportPt = Utilities::CastHelper(Rotated, Identity, [&](const auto& Concrete, const auto& FullTM)
 			{
-				FVec3 SupportLocal = Concrete.Support(FullTM.InverseTransformVectorNoScale(DirZ), 1);
+				FVec3 SupportLocal = Concrete.Support(FullTM.InverseTransformVectorNoScale(DirZ), 1, VertexIndex);
 				return FullTM.TransformPosition(SupportLocal);
 			});
 			EXPECT_VECTOR_NEAR_DEFAULT(SupportPt, FVec3(-5, 8, -6)); // @todo why -5?
 
 			SupportPt = Utilities::CastHelper(Transformed, Identity, [&](const auto& Concrete, const auto& FullTM)
 			{
-				FVec3 SupportLocal = Concrete.Support(FullTM.InverseTransformVectorNoScale(DirZ), 1);
+				FVec3 SupportLocal = Concrete.Support(FullTM.InverseTransformVectorNoScale(DirZ), 1, VertexIndex);
 				FVec3 TransformedPt = FullTM.TransformPosition(SupportLocal);
 				return TransformedPt;
 			});
@@ -1875,54 +1880,55 @@ namespace ChaosTest {
 
 		//support
 		{
+			int32 VertexIndex = INDEX_NONE;
 			const FVec3 DirX(1, 0, 0);
-			FVec3 SupportPt = Unscaled.Support(DirX, 1);
+			FVec3 SupportPt = Unscaled.Support(DirX, 1, VertexIndex);
 			EXPECT_FLOAT_EQ(SupportPt[0], 9);
 			EXPECT_FLOAT_EQ(SupportPt[1], 0);
 			EXPECT_FLOAT_EQ(SupportPt[2], 0);
 
-			SupportPt = UnscaledThickened.Support(DirX, 1);
+			SupportPt = UnscaledThickened.Support(DirX, 1, VertexIndex);
 			EXPECT_FLOAT_EQ(SupportPt[0], 9);
 			EXPECT_FLOAT_EQ(SupportPt[1], 0);
 			EXPECT_FLOAT_EQ(SupportPt[2], 0);
 
-			SupportPt = UniformScale.Support(DirX, 1);
+			SupportPt = UniformScale.Support(DirX, 1, VertexIndex);
 			EXPECT_FLOAT_EQ(SupportPt[0], 17);
 			EXPECT_FLOAT_EQ(SupportPt[1], 0);
 			EXPECT_FLOAT_EQ(SupportPt[2], 0);
 
-			SupportPt = UniformScaleThickened.Support(DirX, 1);
+			SupportPt = UniformScaleThickened.Support(DirX, 1, VertexIndex);
 			EXPECT_FLOAT_EQ(SupportPt[0], 17);
 			EXPECT_FLOAT_EQ(SupportPt[1], 0);
 			EXPECT_FLOAT_EQ(SupportPt[2], 0);
 
 			const FVec3 DirZ(0, 0, -1);
-			SupportPt = UniformScale.Support(DirZ, 1);
+			SupportPt = UniformScale.Support(DirZ, 1, VertexIndex);
 			EXPECT_FLOAT_EQ(SupportPt[0], 6);
 			EXPECT_FLOAT_EQ(SupportPt[1], 0);
 			EXPECT_FLOAT_EQ(SupportPt[2], -11);
 
-			SupportPt = UniformScaleThickened.Support(DirZ, 1);
+			SupportPt = UniformScaleThickened.Support(DirZ, 1, VertexIndex);
 			EXPECT_FLOAT_EQ(SupportPt[0], 6);
 			EXPECT_FLOAT_EQ(SupportPt[1], 0);
 			EXPECT_FLOAT_EQ(SupportPt[2], -11);
 
-			SupportPt = NonUniformScale.Support(DirX, 1);
+			SupportPt = NonUniformScale.Support(DirX, 1, VertexIndex);
 			EXPECT_FLOAT_EQ(SupportPt[0], 17);
 			EXPECT_FLOAT_EQ(SupportPt[1], 0);
 			EXPECT_FLOAT_EQ(SupportPt[2], 0);
 
-			SupportPt = NonUniformScaleThickened.Support(DirX, 1);
+			SupportPt = NonUniformScaleThickened.Support(DirX, 1, VertexIndex);
 			EXPECT_FLOAT_EQ(SupportPt[0], 17);
 			EXPECT_FLOAT_EQ(SupportPt[1], 0);
 			EXPECT_FLOAT_EQ(SupportPt[2], 0);
 
-			SupportPt = NonUniformScale.Support(DirZ, 1);
+			SupportPt = NonUniformScale.Support(DirZ, 1, VertexIndex);
 			EXPECT_FLOAT_EQ(SupportPt[0], 6);
 			EXPECT_FLOAT_EQ(SupportPt[1], 0);
 			EXPECT_FLOAT_EQ(SupportPt[2], -6);
 
-			SupportPt = NonUniformScaleThickened.Support(DirZ, 1);
+			SupportPt = NonUniformScaleThickened.Support(DirZ, 1, VertexIndex);
 			EXPECT_FLOAT_EQ(SupportPt[0], 6);
 			EXPECT_FLOAT_EQ(SupportPt[1], 0);
 			EXPECT_FLOAT_EQ(SupportPt[2], -6);
diff --git a/Engine/Source/Runtime/Experimental/Chaos/Private/Chaos/Box.cpp b/Engine/Source/Runtime/Experimental/Chaos/Private/Chaos/Box.cpp
--- a/Engine/Source/Runtime/Experimental/Chaos/Private/Chaos/Box.cpp
+++ b/Engine/Source/Runtime/Experimental/Chaos/Private/Chaos/Box.cpp
@@ -9,12 +9,12 @@ namespace Chaos
 	{
 		TArray<TArray<int32>> PlaneVertices
 		{
-			{6, 7, 3, 2},	// X
-			{1, 5, 6, 2},	// Y
-			{7, 6, 5, 4},	// Z,
-			{1, 0, 4, 5},	// -X
-			{0, 3, 7, 4},	// -Y
-			{0, 1, 2, 3},	// -Z,
+			{7, 5, 4, 6},	// X
+			{2, 3, 7, 6},	// Y
+			{5, 7, 3, 1},	// Z,
+			{2, 0, 1, 3},	// -X
+			{0, 4, 5, 1},	// -Y
+			{0, 2, 6, 4},	// -Z,
 		};
 
 		return FConvexHalfEdgeStructureDataS16::MakePlaneVertices(PlaneVertices, 8);
@@ -22,24 +22,24 @@ namespace Chaos
 
 	template<typename T, int D> TArray<FVec3> TBox<T, D>::SNormals =
 	{
-		FVec3(1,0,0),	// X
-		FVec3(0,1,0),	// Y
-		FVec3(0,0,1),	// Z
-		FVec3(-1,0,0),	// -X
-		FVec3(0,-1,0),	// -Y
-		FVec3(0,0,-1),	// -Z
+		FVec3(1,0,0),	 // X
+		FVec3(0,1,0),	 // Y
+		FVec3(0,0,1),	 // Z
+		FVec3(-1,0,0), // -X
+		FVec3(0,-1,0), // -Y
+		FVec3(0,0,-1), // -Z
 	};
-
+	
 	template<typename T, int D> TArray<FVec3> TBox<T, D>::SVertices =
 	{
 		FVec3(-1,-1,-1),
-		FVec3(-1,1,-1),
-		FVec3(1,1,-1),
-		FVec3(1,-1,-1),
 		FVec3(-1,-1,1),
+		FVec3(-1,1,-1),
 		FVec3(-1,1,1),
-		FVec3(1,1,1),
+		FVec3(1,-1,-1),
 		FVec3(1,-1,1),
+		FVec3(1,1,-1),
+		FVec3(1,1,1),
 	};
 
 	template<typename T, int D> FConvexHalfEdgeStructureDataS16 TBox<T, D>::SStructureData = MakeBoxStructureData();
diff --git a/Engine/Source/Runtime/Experimental/Chaos/Private/Chaos/CollisionOneShotManifolds.cpp b/Engine/Source/Runtime/Experimental/Chaos/Private/Chaos/CollisionOneShotManifolds.cpp
--- a/Engine/Source/Runtime/Experimental/Chaos/Private/Chaos/CollisionOneShotManifolds.cpp
+++ b/Engine/Source/Runtime/Experimental/Chaos/Private/Chaos/CollisionOneShotManifolds.cpp
@@ -619,7 +619,7 @@ namespace Chaos
 		// Use GJK to find the closest points (or shallowest penetrating points) on two convex shapes usingthe specified margin
 		// @todo(chaos): dedupe from GJKContactPoint in CollisionResolution.cpp
 		template <typename GeometryA, typename GeometryB>
-		FContactPoint GJKContactPointMargin(const GeometryA& A, const GeometryB& B, const FRigidTransform3& ATM, const FRigidTransform3& BTM, FReal MarginA, FReal MarginB, FGJKSimplexData& InOutGjkWarmStartData, FReal& OutMaxMarginDelta)
+		FContactPoint GJKContactPointMargin(const GeometryA& A, const GeometryB& B, const FRigidTransform3& ATM, const FRigidTransform3& BTM, FReal MarginA, FReal MarginB, FGJKSimplexData& InOutGjkWarmStartData, FReal& OutMaxMarginDelta, int32& VertexIndexA, int32& VertexIndexB)
 		{
 			SCOPE_CYCLE_COUNTER_MANIFOLD_GJK();
 
@@ -635,7 +635,7 @@ namespace Chaos
 			const TGJKCoreShape<GeometryB> BWithMargin(B, MarginB);
 			const FRigidTransform3 BToATM = BTM.GetRelativeTransformNoScale(ATM);
 
-			if (GJKPenetrationWarmStartable<true>(AWithMargin, BWithMargin, BToATM, FReal(0), FReal(0), Penetration, ClosestA, ClosestB, NormalA, NormalB, InOutGjkWarmStartData, OutMaxMarginDelta, Epsilon))
+			if (GJKPenetrationWarmStartable<true>(AWithMargin, BWithMargin, BToATM, FReal(0), FReal(0), Penetration, ClosestA, ClosestB, NormalA, NormalB, VertexIndexA, VertexIndexB, InOutGjkWarmStartData, OutMaxMarginDelta, Epsilon))
 			{
 				Contact.ShapeContactPoints[0] = ClosestA;
 				Contact.ShapeContactPoints[1] = ClosestB;
@@ -657,7 +657,8 @@ namespace Chaos
 			const ConvexImplicitType& Convex,
 			const FVec3 X,
 			const FVec3 N,
-			const FReal InMaxDistance)
+			const FReal InMaxDistance,
+			const int32 VertexIndex)
 		{
 			// Handle InMaxDistance = 0. We expect that the X is actually on the surface in this case, so the search distance just needs to be some reasonable tolerance.
 			// @todo(chaos): this should probable be dependent on the size of the objects...
@@ -666,35 +667,38 @@ namespace Chaos
 
 			int32 BestPlaneIndex = INDEX_NONE;
 			FReal BestPlaneDot = 1.0f;
-
-			// Finds the most opposing plane (to N) that is within the specified distance (MaxDistance) of the position (X)
-			const int32 NumPlanes = Convex.NumPlanes();
-			for (int32 PlaneIndex = 0; PlaneIndex < NumPlanes; ++PlaneIndex)
+			
 			{
-				const TPlaneConcrete<FReal, 3> Plane = Convex.GetPlane(PlaneIndex);
-				const FReal PlaneNormalDotN = FVec3::DotProduct(N, Plane.Normal());
-
-				// Ignore planes that do not oppose N
-				if (PlaneNormalDotN > -SMALL_NUMBER)
-				{
-					continue;
-				}
-
-				// Reject planes farther than MaxDistance
-				const FReal PlaneDistance = Plane.SignedDistance(X);
-				if (FMath::Abs(PlaneDistance) > MaxDistance)
-				{
-					continue;
-				}
-
-				// Keep the most opposing plane
-				if (PlaneNormalDotN < BestPlaneDot)
+				int32 PlaneIndices[3] = {INDEX_NONE, INDEX_NONE, INDEX_NONE};
+				const int32 NumPlanes = Convex.GetVertexPlanes3(VertexIndex, PlaneIndices[0], PlaneIndices[1], PlaneIndices[2]);
+			
+				for (int32 PlaneIndex = 0; PlaneIndex < NumPlanes; ++PlaneIndex)
 				{
-					BestPlaneDot = PlaneNormalDotN;
-					BestPlaneIndex = PlaneIndex;
+					const TPlaneConcrete<FReal, 3> Plane = Convex.GetPlane(PlaneIndices[PlaneIndex]);
+					const FReal PlaneNormalDotN = FVec3::DotProduct(N, Plane.Normal());
+			
+					// Ignore planes that do not oppose N
+					if (PlaneNormalDotN > -SMALL_NUMBER)
+					{
+						continue;
+					}
+			
+					// Reject planes farther than MaxDistance
+					const FReal PlaneDistance = Plane.SignedDistance(X);
+					if (FMath::Abs(PlaneDistance) > MaxDistance)
+					{
+						continue;
+					}
+			
+					// Keep the most opposing plane
+					if (PlaneNormalDotN < BestPlaneDot)
+					{
+						BestPlaneDot = PlaneNormalDotN;
+						BestPlaneIndex = PlaneIndices[PlaneIndex];
+					}
 				}
 			}
-
+			
 			// Malformed convexes or half-spaces or capsules could have all planes rejected above.
 			// If that happens, select the most opposing plane including those that
 			// may point the same direction as N. 
@@ -751,7 +755,8 @@ namespace Chaos
 			// Find the deepest penetration. This is used to determine the planes and points to use for the manifold
 			// MaxMarginDelta is an upper bound on the distance from the contact on the rounded core shape to the actual shape surface. 
 			FReal MaxMarginDelta = FReal(0);
-			FContactPoint GJKContactPoint = GJKContactPointMargin(Convex1, Convex2, Convex1Transform, Convex2Transform, Margin1, Margin2, Constraint.GetGJKWarmStartData(), MaxMarginDelta);
+			int32 VertexIndexA = INDEX_NONE, VertexIndexB = INDEX_NONE;
+			FContactPoint GJKContactPoint = GJKContactPointMargin(Convex1, Convex2, Convex1Transform, Convex2Transform, Margin1, Margin2, Constraint.GetGJKWarmStartData(), MaxMarginDelta, VertexIndexA, VertexIndexB);
 
 			// GJK is using margins and rounded corner, so if we have a corner-to-corner contact it will under-report the actual distance by an amount that depends on how
 			// "pointy" the edge/corner is - this error is bounded by MaxMarginDelta.
@@ -763,13 +768,13 @@ namespace Chaos
 
 			// @todo(chaos): get the vertex index from GJK and use to to get the plane
 			const FVec3 SeparationDirectionLocalConvex1 = Convex1Transform.InverseTransformVectorNoScale(GJKContactPoint.Normal);
-			const int32 MostOpposingPlaneIndexConvex1 = SelectContactPlane(Convex1, GJKContactPoint.ShapeContactPoints[0], SeparationDirectionLocalConvex1, Margin1);
+			const int32 MostOpposingPlaneIndexConvex1 = SelectContactPlane(Convex1, GJKContactPoint.ShapeContactPoints[0], SeparationDirectionLocalConvex1, Margin1, VertexIndexA);
 			const TPlaneConcrete<FReal, 3> BestPlaneConvex1 = Convex1.GetPlane(MostOpposingPlaneIndexConvex1);
 			const FReal BestPlaneDotNormalConvex1 = !bConvex1IsCapsule ? FVec3::DotProduct(-SeparationDirectionLocalConvex1, BestPlaneConvex1.Normal()) : -FLT_MAX;
 
 			// Now for Convex2
 			const FVec3 SeparationDirectionLocalConvex2 = Convex2Transform.InverseTransformVectorNoScale(GJKContactPoint.Normal);
-			const int32 MostOpposingPlaneIndexConvex2 = SelectContactPlane(Convex2, GJKContactPoint.ShapeContactPoints[1], -SeparationDirectionLocalConvex2, Margin2);
+			const int32 MostOpposingPlaneIndexConvex2 = SelectContactPlane(Convex2, GJKContactPoint.ShapeContactPoints[1], -SeparationDirectionLocalConvex2, Margin2, VertexIndexB);
 			const TPlaneConcrete<FReal, 3> BestPlaneConvex2 = Convex2.GetPlane(MostOpposingPlaneIndexConvex2);
 			const FReal BestPlaneDotNormalConvex2 = FVec3::DotProduct(SeparationDirectionLocalConvex2, BestPlaneConvex2.Normal());
 
diff --git a/Engine/Source/Runtime/Experimental/Chaos/Public/Chaos/AABB.h b/Engine/Source/Runtime/Experimental/Chaos/Public/Chaos/AABB.h
--- a/Engine/Source/Runtime/Experimental/Chaos/Public/Chaos/AABB.h
+++ b/Engine/Source/Runtime/Experimental/Chaos/Public/Chaos/AABB.h
@@ -297,14 +297,25 @@ namespace Chaos
 
 			return BestNormal;
 		}
-
-		FORCEINLINE_DEBUGGABLE TVector<T, d> Support(const TVector<T, d>& Direction, const T Thickness) const
+		
+		FORCEINLINE_DEBUGGABLE TVector<T, d> Support(const TVector<T, d>& Direction, const T Thickness, int32& VertexIndex) const
 		{
 			TVector<T, d> ChosenPt;
+			FIntVector ChosenAxis;
 			for (int Axis = 0; Axis < d; ++Axis)
 			{
-				ChosenPt[Axis] = Direction[Axis] < 0 ? MMin[Axis] : MMax[Axis];
+				if(Direction[Axis] < 0)
+				{
+					ChosenPt[Axis] = MMin[Axis];
+					ChosenAxis[Axis] = 0;
+				}
+				else
+				{
+					ChosenPt[Axis] = MMax[Axis];
+					ChosenAxis[Axis] = 1;
+				}
 			}
+			VertexIndex =  ChosenAxis[0] * 4 + ChosenAxis[1] * 2 + ChosenAxis[2];
 
 			if (Thickness != (T)0)
 			{
@@ -325,14 +336,24 @@ namespace Chaos
 		}
 
 		// Support vertex in the specified direction, assuming each face has been moved inwards by InMargin
-		FORCEINLINE_DEBUGGABLE FVec3 SupportCore(const FVec3& Direction, const FReal InMargin, FReal* OutSupportDelta) const
+		FORCEINLINE_DEBUGGABLE FVec3 SupportCore(const FVec3& Direction, const FReal InMargin, FReal* OutSupportDelta, int32& VertexIndex) const
 		{
 			FVec3 ChosenPt;
+			FIntVector ChosenAxis;
 			for (int Axis = 0; Axis < d; ++Axis)
 			{
-				ChosenPt[Axis] = Direction[Axis] < 0 ? MMin[Axis] + InMargin : MMax[Axis] - InMargin;
+				if(Direction[Axis] < 0)
+				{
+					ChosenPt[Axis] = MMin[Axis] + InMargin;
+					ChosenAxis[Axis] = 0;
+				}
+				else
+				{
+					ChosenPt[Axis] = MMax[Axis] - InMargin;
+					ChosenAxis[Axis] = 1;
+				}
 			}
-
+			VertexIndex = ChosenAxis[0] * 4 + ChosenAxis[1] * 2 + ChosenAxis[2];
 			// Maximum distance between the Core+Margin position and the original outer vertex
 			constexpr FReal RootThreeMinusOne = FReal(1.7320508075688772935274463415059 - 1.0);
 			if (OutSupportDelta != nullptr)
@@ -343,16 +364,27 @@ namespace Chaos
 			return ChosenPt;
 		}
 
-		FORCEINLINE_DEBUGGABLE TVector<T, d> SupportCoreScaled(const TVector<T, d>& Direction, const T InMargin, const TVector<T, d>& Scale, T* OutSupportDelta) const
+		FORCEINLINE_DEBUGGABLE TVector<T, d> SupportCoreScaled(const TVector<T, d>& Direction, const T InMargin, const TVector<T, d>& Scale, T* OutSupportDelta, int32& VertexIndex) const
 		{
 			const TVector<T, d> ScaledDirection = Direction * Scale;
 
 			TVector<T, d> ChosenPt;
+			FIntVector ChosenAxis;
 			for (int Axis = 0; Axis < d; ++Axis)
 			{
-				ChosenPt[Axis] = ScaledDirection[Axis] < 0 ? Scale[Axis] * MMin[Axis] + InMargin : Scale[Axis] * MMax[Axis] - InMargin;
+				if(ScaledDirection[Axis] < 0)
+				{
+					ChosenPt[Axis] = Scale[Axis] * MMin[Axis] + InMargin;
+					ChosenAxis[Axis] = 0;
+				}
+				else
+				{
+					ChosenPt[Axis] = Scale[Axis] *  MMax[Axis] - InMargin;
+					ChosenAxis[Axis] = 1;
+				}
 			}
-
+			VertexIndex = ChosenAxis[0] * 4 + ChosenAxis[1] * 2 + ChosenAxis[2];
+			
 			constexpr T RootThreeMinusOne = T(1.7320508075688772935274463415059 - 1.0);
 			if (OutSupportDelta != nullptr)
 			{
diff --git a/Engine/Source/Runtime/Experimental/Chaos/Public/Chaos/Box.h b/Engine/Source/Runtime/Experimental/Chaos/Public/Chaos/Box.h
--- a/Engine/Source/Runtime/Experimental/Chaos/Public/Chaos/Box.h
+++ b/Engine/Source/Runtime/Experimental/Chaos/Public/Chaos/Box.h
@@ -305,6 +305,13 @@ namespace Chaos
 		{
 			return SStructureData.FindVertexPlanes(VertexIndex, OutVertexPlanes, MaxVertexPlanes);
 		}
+		
+		// Get up to the 3  plane indices that belong to a vertex
+		// Returns the number of planes found.
+		int32 GetVertexPlanes3(int32 VertexIndex, int32& PlaneIndex0, int32& PlaneIndex1, int32& PlaneIndex2) const
+		{
+			return SStructureData.GetVertexPlanes3(VertexIndex, PlaneIndex0, PlaneIndex1, PlaneIndex2);
+		}
 
 		// The number of vertices that make up the corners of the specified face
 		int32 NumPlaneVertices(int32 PlaneIndex) const
@@ -350,20 +357,20 @@ namespace Chaos
 		}
 
 		// Returns a position on the shape
-		FORCEINLINE_DEBUGGABLE TVector<T, d> Support(const TVector<T, d>& Direction, const T Thickness) const
+		FORCEINLINE_DEBUGGABLE TVector<T, d> Support(const TVector<T, d>& Direction, const T Thickness, int32& VertexIndex) const
 		{
-			return AABB.Support(Direction, Thickness);
+			return AABB.Support(Direction, Thickness, VertexIndex);
 		}
 
 		// Returns a position on the core shape excluding the margin
-		FORCEINLINE_DEBUGGABLE TVector<T, d> SupportCore(const TVector<T, d>& Direction, const FReal InMargin, FReal* OutSupportDelta) const
+		FORCEINLINE_DEBUGGABLE TVector<T, d> SupportCore(const TVector<T, d>& Direction, const FReal InMargin, FReal* OutSupportDelta, int32& VertexIndex) const
 		{
-			return AABB.SupportCore(Direction, InMargin, OutSupportDelta);
+			return AABB.SupportCore(Direction, InMargin, OutSupportDelta, VertexIndex);
 		}
 
-		FORCEINLINE_DEBUGGABLE TVector<T, d> SupportCoreScaled(const TVector<T, d>& Direction, const FReal InMargin, const TVector<T, d>& Scale, T* OutSupportDelta) const
+		FORCEINLINE_DEBUGGABLE TVector<T, d> SupportCoreScaled(const TVector<T, d>& Direction, const FReal InMargin, const TVector<T, d>& Scale, T* OutSupportDelta, int32& VertexIndex) const
 		{
-			return AABB.SupportCoreScaled(Direction, InMargin, Scale, OutSupportDelta);
+			return AABB.SupportCoreScaled(Direction, InMargin, Scale, OutSupportDelta, VertexIndex);
 			// @todo(chaos): Needs to operate in scaled space as margin is not non-uniform scalable
 			//const FReal InvScale = 1.0f / Scale[0];
 			//const FReal NetMargin = InvScale * InMargin;
diff --git a/Engine/Source/Runtime/Experimental/Chaos/Public/Chaos/Capsule.h b/Engine/Source/Runtime/Experimental/Chaos/Public/Chaos/Capsule.h
--- a/Engine/Source/Runtime/Experimental/Chaos/Public/Chaos/Capsule.h
+++ b/Engine/Source/Runtime/Experimental/Chaos/Public/Chaos/Capsule.h
@@ -339,21 +339,21 @@ namespace Chaos
 			return RaycastFast(GetRadius(), GetHeight(), GetAxis(), GetX1(), GetX2(), StartPoint, Dir, Length, Thickness, OutTime, OutPosition, OutNormal, OutFaceIndex);
 		}
 
-		FORCEINLINE FVec3 Support(const FVec3& Direction, const FReal Thickness) const
+		FORCEINLINE FVec3 Support(const FVec3& Direction, const FReal Thickness, int32& VertexIndex) const
 		{
-			return MSegment.Support(Direction, GetRadius() + Thickness);
+			return MSegment.Support(Direction, GetRadius() + Thickness, VertexIndex);
 		}
 
-		FORCEINLINE FVec3 SupportCore(const FVec3& Direction, const FReal InMargin, FReal* OutSupportDelta) const
+		FORCEINLINE FVec3 SupportCore(const FVec3& Direction, const FReal InMargin, FReal* OutSupportDelta, int32& VertexIndex) const
 		{
 			// NOTE: Ignores InMargin, assumes Radius
-			return MSegment.SupportCore(Direction);
+			return MSegment.SupportCore(Direction, VertexIndex);
 		}
 
-		FORCEINLINE FVec3 SupportCoreScaled(const FVec3& Direction, const FReal InMargin, const FVec3& Scale, FReal* OutSupportDelta) const
+		FORCEINLINE FVec3 SupportCoreScaled(const FVec3& Direction, const FReal InMargin, const FVec3& Scale, FReal* OutSupportDelta, int32& VertexIndex) const
 		{
 			// NOTE: Ignores InMargin, assumes Radius
-			return SupportCore(Scale * Direction, GetMargin(), OutSupportDelta) * Scale;
+			return SupportCore(Scale * Direction, GetMargin(), OutSupportDelta, VertexIndex) * Scale;
 		}
 
 		FORCEINLINE void SerializeImp(FArchive& Ar)
@@ -518,6 +518,20 @@ namespace Chaos
 			return TPlaneConcrete<FReal, 3>(FVec3(0), FVec3(0));
 		}
 
+		// Get an array of all the plane indices that belong to a vertex (up to MaxVertexPlanes).
+		// Returns the number of planes found.
+		int32 FindVertexPlanes(int32 VertexIndex, int32* OutVertexPlanes, int32 MaxVertexPlanes) const
+		{
+			return 0; 
+		}
+		
+		// Get up to the 3  plane indices that belong to a vertex
+		// Returns the number of planes found.
+		int32 GetVertexPlanes3(int32 VertexIndex, int32& PlaneIndex0, int32& PlaneIndex1, int32& PlaneIndex2) const
+		{
+			return 0;
+		}
+
 		// Capsules have no planes
 		// Used for manifold generation
 		int32 NumPlanes() const { return 0; }
diff --git a/Engine/Source/Runtime/Experimental/Chaos/Public/Chaos/Convex.h b/Engine/Source/Runtime/Experimental/Chaos/Public/Chaos/Convex.h
--- a/Engine/Source/Runtime/Experimental/Chaos/Public/Chaos/Convex.h
+++ b/Engine/Source/Runtime/Experimental/Chaos/Public/Chaos/Convex.h
@@ -609,9 +609,10 @@ namespace Chaos
 
 	public:
 		// Return support point on the core shape (the convex shape with all planes moved inwards by margin).
-		FVec3 SupportCore(const FVec3& Direction, const FReal InMargin, FReal* OutSupportDelta) const
+		FVec3 SupportCore(const FVec3& Direction, const FReal InMargin, FReal* OutSupportDelta, int32& VertexIndex) const
 		{
 			const int32 SupportVertexIndex = GetSupportVertex(Direction);
+			VertexIndex = SupportVertexIndex;
 			if (SupportVertexIndex != INDEX_NONE)
 			{
 				return GetMarginAdjustedVertex(SupportVertexIndex, InMargin, OutSupportDelta);
@@ -621,12 +622,12 @@ namespace Chaos
 
 		// SupportCore with non-uniform scale support. This is required for the margin in scaled
 		// space to by uniform. Note in this version all the inputs are in outer container's (scaled shape) space
-		FVec3 SupportCoreScaled(const FVec3& Direction, FReal InMargin, const FVec3& Scale, FReal* OutSupportDelta) const
+		FVec3 SupportCoreScaled(const FVec3& Direction, FReal InMargin, const FVec3& Scale, FReal* OutSupportDelta, int32& VertexIndex) const
 		{
 			// Find the supporting vertex index
 			const FVec3 DirectionScaled = Scale * Direction;	// does not need to be normalized
 			const int32 SupportVertexIndex = GetSupportVertex(DirectionScaled);
-
+			VertexIndex = SupportVertexIndex;
 			// Adjust the vertex position based on margin
 			if (SupportVertexIndex != INDEX_NONE)
 			{
@@ -639,9 +640,10 @@ namespace Chaos
 
 		// Return support point on the shape
 		// @todo(chaos): do we need to support thickness?
-		FORCEINLINE FVec3 Support(const FVec3& Direction, const FReal Thickness) const
+		FORCEINLINE FVec3 Support(const FVec3& Direction, const FReal Thickness, int32& VertexIndex) const
 		{
 			const int32 MaxVIdx = GetSupportVertex(Direction);
+			VertexIndex = MaxVIdx;
 			if (MaxVIdx != INDEX_NONE)
 			{
 				if (Thickness != 0.0f)
@@ -653,9 +655,9 @@ namespace Chaos
 			return FVec3(0);
 		}
 
-		FORCEINLINE FVec3 SupportScaled(const FVec3& Direction, const FReal Thickness, const FVec3& Scale) const
+		FORCEINLINE FVec3 SupportScaled(const FVec3& Direction, const FReal Thickness, const FVec3& Scale, int32& VertexIndex) const
 		{
-			FVec3 SupportPoint = Support(Direction * Scale, 0.0f) * Scale;
+			FVec3 SupportPoint = Support(Direction * Scale, 0.0f, VertexIndex) * Scale;
 			if (Thickness > 0.0f)
 			{
 				SupportPoint += Thickness * Direction.GetSafeNormal();
diff --git a/Engine/Source/Runtime/Experimental/Chaos/Public/Chaos/GJK.h b/Engine/Source/Runtime/Experimental/Chaos/Public/Chaos/GJK.h
--- a/Engine/Source/Runtime/Experimental/Chaos/Public/Chaos/GJK.h
+++ b/Engine/Source/Runtime/Experimental/Chaos/Public/Chaos/GJK.h
@@ -45,6 +45,7 @@ namespace Chaos
 		const T ThicknessB = B.GetMargin() + InThicknessB;
 		const T Inflation = ThicknessA + ThicknessB + static_cast<T>(1e-3);
 		const T Inflation2 = Inflation * Inflation;
+		int32 VertexIndexA = INDEX_NONE, VertexIndexB = INDEX_NONE;
 		do
 		{
 			if (!ensure(NumIterations++ < 32))	//todo: take this out
@@ -52,9 +53,9 @@ namespace Chaos
 				break;	//if taking too long just stop. This should never happen
 			}
 			const TVector<T, 3> NegV = -V;
-			const TVector<T, 3> SupportA = A.SupportCore(NegV, A.GetMargin(), nullptr);
+			const TVector<T, 3> SupportA = A.SupportCore(NegV, A.GetMargin(), nullptr, VertexIndexA);
 			const TVector<T, 3> VInB = AToBRotation * V;
-			const TVector<T, 3> SupportBLocal = B.SupportCore(VInB, B.GetMargin(), nullptr);
+			const TVector<T, 3> SupportBLocal = B.SupportCore(VInB, B.GetMargin(), nullptr, VertexIndexB);
 			const TVector<T, 3> SupportB = BToATM.TransformPositionNoScale(SupportBLocal);
 			const TVector<T, 3> W = SupportA - SupportB;
 
@@ -184,6 +185,8 @@ namespace Chaos
 	 * @param OutClosestB The closest point on B, in B-local space
 	 * @param OutNormalA The contact normal pointing away from A in A-local space
 	 * @param OutNormalB The contact normal pointing away from A in B-local space
+	 * @param OutVertexA The closest vertex on A
+	 * @param OutVertexB The closest vertex on B
 	 * @param SimplexData In/out simplex data used to initialize and update GJK. Can be stored to improve convergence on subsequent calls for "small" changes in relative rotation.
 	 * @param OutMaxContactDelta The maximum error in the contact position as a result of using a margin. This is the difference between the core shape + margin and the outer shape on the supporting vertices
 	 * @param Epsilon The separation distance below which GJK aborts and switches to EPA
@@ -199,25 +202,28 @@ namespace Chaos
 	 *
 	*/
 	template <bool bFindSeparationDistance = false, typename T, typename TGeometryA, typename TGeometryB>
-	bool GJKPenetrationWarmStartable(const TGeometryA& A, const TGeometryB& B, const TRigidTransform<T, 3>& BToATM, const T InThicknessA, const T InThicknessB, T& OutPenetration, TVec3<T>& OutClosestA, TVec3<T>& OutClosestB, TVec3<T>& OutNormalA, TVec3<T>& OutNormalB, TGJKSimplexData<T>& InOutSimplexData, T& OutMaxSupportDelta, const T Epsilon = T(1.e-3))
+	bool GJKPenetrationWarmStartable(const TGeometryA& A, const TGeometryB& B, const TRigidTransform<T, 3>& BToATM, const T InThicknessA, const T InThicknessB, T& OutPenetration, TVec3<T>& OutClosestA, TVec3<T>& OutClosestB, TVec3<T>& OutNormalA, TVec3<T>& OutNormalB, int32& OutVertexA, int32& OutVertexB, TGJKSimplexData<T>& InOutSimplexData, T& OutMaxSupportDelta, const T Epsilon = T(1.e-3))
 	{
 		T SupportDeltaA = 0;
 		T SupportDeltaB = 0;
 		T MaxSupportDelta = 0;
 
+		int32& VertexIndexA = OutVertexA;
+		int32& VertexIndexB = OutVertexB;
+
 		// Return the support vertex in A-local space for a vector V in A-local space
-		auto SupportAFunc = [&A, &SupportDeltaA](const TVec3<T>& V)
+		auto SupportAFunc = [&A, &SupportDeltaA, &VertexIndexA](const TVec3<T>& V)
 		{
-			return A.SupportCore(V, A.GetMargin(), &SupportDeltaA);
+			return A.SupportCore(V, A.GetMargin(), &SupportDeltaA, VertexIndexA);
 		};
 
 		const TRotation<T, 3> AToBRotation = BToATM.GetRotation().Inverse();
 
 		// Return the support point on B, in B-local space, for a vector V in A-local space
-		auto SupportBFunc = [&B, &AToBRotation, &SupportDeltaB](const TVec3<T>& V)
+		auto SupportBFunc = [&B, &AToBRotation, &SupportDeltaB, &VertexIndexB](const TVec3<T>& V)
 		{
 			const TVec3<T> VInB = AToBRotation * V;
-			return B.SupportCore(VInB, B.GetMargin(), &SupportDeltaB);
+			return B.SupportCore(VInB, B.GetMargin(), &SupportDeltaB, VertexIndexB);
 		};
 
 		// V and Simplex are in A-local space
@@ -301,11 +307,12 @@ namespace Chaos
 				VertsA.Add(InOutSimplexData.As[i]);
 				VertsB.Add(BToATM.TransformPositionNoScale(InOutSimplexData.Bs[i]));
 			}
-
+			
 			auto SupportBInAFunc = [&B, &BToATM, &AToBRotation, &SupportDeltaB](const TVec3<T>& V)
 			{
 				const TVec3<T> VInB = AToBRotation * V;
-				const TVec3<T> SupportBLocal = B.SupportCore(VInB, B.GetMargin(), &SupportDeltaB);
+				int32 VertexIndex = INDEX_NONE;
+				const TVec3<T> SupportBLocal = B.SupportCore(VInB, B.GetMargin(), &SupportDeltaB, VertexIndex);
 				return BToATM.TransformPositionNoScale(SupportBLocal);
 			};
 
@@ -380,17 +387,15 @@ namespace Chaos
 
 		auto SupportAFunc = [&A, &VertexIndexA](const TVec3<T>& V)
 		{
-			VertexIndexA = INDEX_NONE;
-			return A.SupportCore(V, A.GetMargin(), nullptr);
+			return A.SupportCore(V, A.GetMargin(), nullptr, VertexIndexA);
 		};
 
 		const TRotation<T, 3> AToBRotation = BToATM.GetRotation().Inverse();
 
 		auto SupportBFunc = [&B, &BToATM, &AToBRotation, &VertexIndexB](const TVec3<T>& V)
 		{
-			VertexIndexB = INDEX_NONE;
 			const TVector<T, 3> VInB = AToBRotation * V;
-			const TVector<T, 3> SupportBLocal = B.SupportCore(VInB, B.GetMargin(), nullptr);
+			const TVector<T, 3> SupportBLocal = B.SupportCore(VInB, B.GetMargin(), nullptr, VertexIndexB);
 			return BToATM.TransformPositionNoScale(SupportBLocal);
 		};
 
@@ -589,6 +594,8 @@ namespace Chaos
 		ensure(RayLength > 0);
 		check(A.IsConvex() && B.IsConvex());
 		const TVector<T, 3> StartPoint = StartTM.GetLocation();
+		int32 VertexIndexA = INDEX_NONE;
+		int32 VertexIndexB = INDEX_NONE;
 
 		TVector<T, 3> Simplex[4] = { TVector<T,3>(0), TVector<T,3>(0), TVector<T,3>(0), TVector<T,3>(0) };
 		TVector<T, 3> As[4] = { TVector<T,3>(0), TVector<T,3>(0), TVector<T,3>(0), TVector<T,3>(0) };
@@ -599,11 +606,11 @@ namespace Chaos
 		FSimplex SimplexIDs;
 		const TRotation<T, 3> BToARotation = StartTM.GetRotation();
 		const TRotation<T, 3> AToBRotation = BToARotation.Inverse();
-		TVector<T, 3> SupportA = A.Support(InitialDir, ThicknessA);	//todo: use Thickness on quadratic geometry
+		TVector<T, 3> SupportA = A.Support(InitialDir, ThicknessA, VertexIndexA);	//todo: use Thickness on quadratic geometry
 		As[0] = SupportA;
 
 		const TVector<T, 3> InitialDirInB = AToBRotation * (-InitialDir);
-		const TVector<T, 3> InitialSupportBLocal = B.Support(InitialDirInB, ThicknessB);
+		const TVector<T, 3> InitialSupportBLocal = B.Support(InitialDirInB, ThicknessB, VertexIndexB);
 		TVector<T, 3> SupportB = BToARotation * InitialSupportBLocal;
 		Bs[0] = SupportB;
 
@@ -625,9 +632,9 @@ namespace Chaos
 				break;	//if taking too long just stop. This should never happen
 			}
 
-			SupportA = A.Support(V, ThicknessA);	//todo: add thickness to quadratic geometry to avoid quadratic vs quadratic when possible
+			SupportA = A.Support(V, ThicknessA, VertexIndexA);	//todo: add thickness to quadratic geometry to avoid quadratic vs quadratic when possible
 			const TVector<T, 3> VInB = AToBRotation * (-V);
-			const TVector<T, 3> SupportBLocal = B.Support(VInB, ThicknessB);
+			const TVector<T, 3> SupportBLocal = B.Support(VInB, ThicknessB, VertexIndexB);
 			SupportB = BToARotation * SupportBLocal;
 			const TVector<T, 3> P = SupportA - SupportB;
 			const TVector<T, 3> W = X - P;
@@ -746,20 +753,23 @@ namespace Chaos
 
 		auto SupportAFunc = [&A, MarginA](const TVec3<T>& V)
 		{
-			return A.SupportCore(V, MarginA, nullptr);
+			int32 VertexIndex = INDEX_NONE;
+			return A.SupportCore(V, MarginA, nullptr, VertexIndex);
 		};
 
 		auto SupportBFunc = [&B, MarginB, &AToBRotation, &BToARotation](const TVec3<T>& V)
 		{
+			int32 VertexIndex = INDEX_NONE;
 			const TVector<T, 3> VInB = AToBRotation * V;
-			const TVector<T, 3> SupportBLocal = B.SupportCore(VInB, MarginB, nullptr);
+			const TVector<T, 3> SupportBLocal = B.SupportCore(VInB, MarginB, nullptr, VertexIndex);
 			return BToARotation * SupportBLocal;
 		};
 		
 		auto SupportBAtOriginFunc = [&B, MarginB, &StartTM, &AToBRotation](const TVec3<T>& Dir)
 		{
+			int32 VertexIndex = INDEX_NONE;
 			const TVector<T, 3> DirInB = AToBRotation * Dir;
-			const TVector<T, 3> SupportBLocal = B.SupportCore(DirInB, MarginB, nullptr);
+			const TVector<T, 3> SupportBLocal = B.SupportCore(DirInB, MarginB, nullptr, VertexIndex);
 			return StartTM.TransformPositionNoScale(SupportBLocal);
 		};
 
@@ -1011,10 +1021,11 @@ namespace Chaos
 	template <typename T, typename TGeometryA, typename TGeometryB>
 	TVector<T, 3> GJKDistanceInitialV(const TGeometryA& A, T MarginA, const TGeometryB& B, T MarginB, const TRigidTransform<T, 3>& BToATM)
 	{
+		int32 VertexIndexA = INDEX_NONE, VertexIndexB = INDEX_NONE;
 		const TVec3<T> V = -BToATM.GetTranslation();
-		const TVector<T, 3> SupportA = A.SupportCore(-V, MarginA, nullptr);
+		const TVector<T, 3> SupportA = A.SupportCore(-V, MarginA, nullptr, VertexIndexA);
 		const TVector<T, 3> VInB = BToATM.GetRotation().Inverse() * V;
-		const TVector<T, 3> SupportBLocal = B.SupportCore(VInB, MarginB, nullptr);
+		const TVector<T, 3> SupportBLocal = B.SupportCore(VInB, MarginB, nullptr, VertexIndexB);
 		const TVector<T, 3> SupportB = BToATM.TransformPositionNoScale(SupportBLocal);
 		return SupportA - SupportB;
 	}
@@ -1076,6 +1087,7 @@ namespace Chaos
 		// Select an initial vector in Minkowski(A - B)
 		TVector<T, 3> V = GJKDistanceInitialV(A, AMargin, B, BMargin, BToATM);
 		T VLen = V.Size();
+		int32 VertexIndexA = INDEX_NONE, VertexIndexB = INDEX_NONE;
 
 		int32 It = 0;
 		while (VLen > Epsilon)
@@ -1083,9 +1095,9 @@ namespace Chaos
 			// Find a new point in A-B that is closer to the origin
 			// NOTE: we do not use support thickness here. Thickness is used when separating objects
 			// so that GJK can find a solution, but that can be added in a later step.
-			const TVector<T, 3> SupportA = A.SupportCore(-V, AMargin, nullptr);
+			const TVector<T, 3> SupportA = A.SupportCore(-V, AMargin, nullptr, VertexIndexA);
 			const TVector<T, 3> VInB = AToBRotation * V;
-			const TVector<T, 3> SupportBLocal = B.SupportCore(VInB, BMargin, nullptr);
+			const TVector<T, 3> SupportBLocal = B.SupportCore(VInB, BMargin, nullptr, VertexIndexB);
 			const TVector<T, 3> SupportB = BToATM.TransformPositionNoScale(SupportBLocal);
 			const TVector<T, 3> W = SupportA - SupportB;
 
diff --git a/Engine/Source/Runtime/Experimental/Chaos/Public/Chaos/GJKShape.h b/Engine/Source/Runtime/Experimental/Chaos/Public/Chaos/GJKShape.h
--- a/Engine/Source/Runtime/Experimental/Chaos/Public/Chaos/GJKShape.h
+++ b/Engine/Source/Runtime/Experimental/Chaos/Public/Chaos/GJKShape.h
@@ -30,9 +30,9 @@ namespace Chaos
 			return 0.0f;
 		}
 
-		FVec3 SupportCore(const FVec3 Dir, const FReal InMargin, FReal* OutSupportDelta) const
+		FVec3 SupportCore(const FVec3 Dir, const FReal InMargin, FReal* OutSupportDelta, int32& VertexIndex) const
 		{
-			return Shape.Support(Dir, InMargin);
+			return Shape.Support(Dir, InMargin, VertexIndex);
 		}
 
 		bool IsConvex() const
@@ -74,9 +74,9 @@ namespace Chaos
 			return Margin;
 		}
 
-		FVec3 SupportCore(const FVec3 Dir, const FReal InMargin, FReal* OutSupportDelta) const
+		FVec3 SupportCore(const FVec3 Dir, const FReal InMargin, FReal* OutSupportDelta, int32& VertexIndex) const
 		{
-			return Shape.SupportCore(Dir, InMargin, OutSupportDelta);
+			return Shape.SupportCore(Dir, InMargin, OutSupportDelta, VertexIndex);
 		}
 
 		bool IsConvex() const
diff --git a/Engine/Source/Runtime/Experimental/Chaos/Public/Chaos/ImplicitObjectScaled.h b/Engine/Source/Runtime/Experimental/Chaos/Public/Chaos/ImplicitObjectScaled.h
--- a/Engine/Source/Runtime/Experimental/Chaos/Public/Chaos/ImplicitObjectScaled.h
+++ b/Engine/Source/Runtime/Experimental/Chaos/Public/Chaos/ImplicitObjectScaled.h
@@ -160,21 +160,21 @@ public:
 	}
 
 	// The support position from the specified direction
-	FORCEINLINE TVector<T, d> Support(const TVector<T, d>& Direction, const T Thickness) const
+	FORCEINLINE TVector<T, d> Support(const TVector<T, d>& Direction, const T Thickness, int32& VertexIndex) const
 	{
-		return MObject->Support(Direction, Thickness); 
+		return MObject->Support(Direction, Thickness, VertexIndex); 
 	}
 
 	// this shouldn't be called, but is required until we remove the explicit function implementations in CollisionResolution.cpp
-	FORCEINLINE TVector<T, d> SupportScaled(const TVector<T, d>& Direction, const T Thickness, const FVec3& Scale) const
+	FORCEINLINE TVector<T, d> SupportScaled(const TVector<T, d>& Direction, const T Thickness, const FVec3& Scale, int32& VertexIndex) const
 	{
-		return MObject->SupportScaled(Direction, Thickness, Scale);
+		return MObject->SupportScaled(Direction, Thickness, Scale, VertexIndex);
 	}
 
 	// The support position from the specified direction, if the shape is reduced by the margin
-	FORCEINLINE TVector<T, d> SupportCore(const TVector<T, d>& Direction, const FReal InMargin, FReal* OutSupportDelta) const
+	FORCEINLINE TVector<T, d> SupportCore(const TVector<T, d>& Direction, const FReal InMargin, FReal* OutSupportDelta, int32& VertexIndex) const
 	{
-		return MObject->SupportCore(Direction, InMargin, OutSupportDelta);
+		return MObject->SupportCore(Direction, InMargin, OutSupportDelta, VertexIndex);
 	}
 
 	virtual const TAABB<T, d> BoundingBox() const override 
@@ -250,6 +250,13 @@ public:
 		return MObject->FindVertexPlanes(VertexIndex, OutVertexPlanes, MaxVertexPlanes);
 	}
 
+	// Get up to the 3  plane indices that belong to a vertex
+	// Returns the number of planes found.
+	int32 GetVertexPlanes3(int32 VertexIndex, int32& PlaneIndex0, int32& PlaneIndex1, int32& PlaneIndex2) const
+	{
+		return MObject->GetVertexPlanes3(VertexIndex, PlaneIndex0, PlaneIndex1, PlaneIndex2);
+	}
+
 	// The number of vertices that make up the corners of the specified face
 	int32 NumPlaneVertices(int32 PlaneIndex) const
 	{
@@ -694,6 +701,13 @@ public:
 	{
 		return MObject->FindVertexPlanes(VertexIndex, OutVertexPlanes, MaxVertexPlanes);
 	}
+	
+	// Get up to the 3  plane indices that belong to a vertex
+	// Returns the number of planes found.
+	int32 GetVertexPlanes3(int32 VertexIndex, int32& PlaneIndex0, int32& PlaneIndex1, int32& PlaneIndex2) const
+	{
+		return MObject->GetVertexPlanes3(VertexIndex, PlaneIndex0, PlaneIndex1, PlaneIndex2);
+	}
 
 	// The number of vertices that make up the corners of the specified face
 	int32 NumPlaneVertices(int32 PlaneIndex) const
@@ -820,7 +834,7 @@ public:
 	}
 
 
-	FORCEINLINE_DEBUGGABLE TVector<T, d> Support(const TVector<T, d>& Direction, const T Thickness) const
+	FORCEINLINE_DEBUGGABLE TVector<T, d> Support(const TVector<T, d>& Direction, const T Thickness, int32& VertexIndex) const
 	{
 		// Support_obj(dir) = pt => for all x in obj, pt \dot dir >= x \dot dir
 		// We want Support_objScaled(dir) = Support_obj(dir') where dir' is some modification of dir so we can use the unscaled support function
@@ -828,13 +842,13 @@ public:
 		// But this is the same as pt \dot dir >= dir^T Ax = (dir^TA) x = (A^T dir)^T x
 		//So let dir' = A^T dir.
 		//Since we only support scaling on the principal axes A is a diagonal (and therefore symmetric) matrix and so a simple component wise multiplication is sufficient
-		const TVector<T, d> UnthickenedPt = MObject->Support(Direction * MScale, 0.0f) * MScale;
+		const TVector<T, d> UnthickenedPt = MObject->Support(Direction * MScale, 0.0f, VertexIndex) * MScale;
 		return Thickness > 0 ? TVector<T, d>(UnthickenedPt + Direction.GetSafeNormal() * Thickness) : UnthickenedPt;
 	}
 
-	FORCEINLINE_DEBUGGABLE TVector<T, d> SupportCore(const TVector<T, d>& Direction, const FReal InMargin, FReal* OutSupportDelta) const
+	FORCEINLINE_DEBUGGABLE TVector<T, d> SupportCore(const TVector<T, d>& Direction, const FReal InMargin, FReal* OutSupportDelta, int32& VertexIndex) const
 	{
-		return MObject->SupportCoreScaled(Direction, InMargin, MScale, OutSupportDelta);
+		return MObject->SupportCoreScaled(Direction, InMargin, MScale, OutSupportDelta, VertexIndex);
 	}
 
 	void SetScale(const FVec3& Scale)
diff --git a/Engine/Source/Runtime/Experimental/Chaos/Public/Chaos/Segment.h b/Engine/Source/Runtime/Experimental/Chaos/Public/Chaos/Segment.h
--- a/Engine/Source/Runtime/Experimental/Chaos/Public/Chaos/Segment.h
+++ b/Engine/Source/Runtime/Experimental/Chaos/Public/Chaos/Segment.h
@@ -26,10 +26,12 @@ namespace Chaos
 
 		FORCEINLINE FReal GetLength() const { return MLength; }
 
-		TVec3<T> Support(const TVec3<T>& Direction, const T Thickness) const
+		TVec3<T> Support(const TVec3<T>& Direction, const T Thickness, int32& VertexIndex) const
 		{
 			const T Dot = TVec3<T>::DotProduct(Direction, MAxis);
-			const TVec3<T> FarthestCap = Dot >= 0 ? GetX2() : GetX1();	//orthogonal we choose either
+			const bool bIsSecond = Dot >= 0;
+			const TVec3<T> FarthestCap = bIsSecond? GetX2() : GetX1();	//orthogonal we choose either
+			VertexIndex = bIsSecond ? 1 : 0;
 			//We want N / ||N|| and to avoid inf
 			//So we want N / ||N|| < 1 / eps => N eps < ||N||, but this is clearly true for all eps < 1 and N > 0
 			T SizeSqr = Direction.SizeSquared();
@@ -41,10 +43,12 @@ namespace Chaos
 			return FarthestCap + (NormalizedDirection * Thickness);
 		}
 
-		FORCEINLINE_DEBUGGABLE TVec3<T> SupportCore(const TVec3<T>& Direction) const
+		FORCEINLINE_DEBUGGABLE TVec3<T> SupportCore(const TVec3<T>& Direction, int32& VertexIndex) const
 		{
 			const T Dot = TVec3<T>::DotProduct(Direction, MAxis);
-			const TVec3<T> FarthestCap = Dot >= 0 ? GetX2() : GetX1();	//orthogonal we choose either
+			const bool bIsSecond = Dot >= 0;
+			const TVec3<T> FarthestCap = bIsSecond ? GetX2() : GetX1();	//orthogonal we choose either
+			VertexIndex = bIsSecond  ? 1 : 0;
 			return FarthestCap;
 		}
 
diff --git a/Engine/Source/Runtime/Experimental/Chaos/Public/Chaos/Sphere.h b/Engine/Source/Runtime/Experimental/Chaos/Public/Chaos/Sphere.h
--- a/Engine/Source/Runtime/Experimental/Chaos/Public/Chaos/Sphere.h
+++ b/Engine/Source/Runtime/Experimental/Chaos/Public/Chaos/Sphere.h
@@ -187,11 +187,12 @@ namespace Chaos
 			return MakePair(TVector<T, d>(Root2 * Direction + StartPoint), true);
 		}
 
-		TVector<T, d> Support(const TVector<T, d>& Direction, const T Thickness) const
+		TVector<T, d> Support(const TVector<T, d>& Direction, const T Thickness, int32& VertexIndex) const
 		{
 			//We want N / ||N|| and to avoid inf
 			//So we want N / ||N|| < 1 / eps => N eps < ||N||, but this is clearly true for all eps < 1 and N > 0
 			T SizeSqr = Direction.SizeSquared();
+			VertexIndex = 0;
 			if (SizeSqr <= TNumericLimits<T>::Min())
 			{
 				return Center;
@@ -201,14 +202,16 @@ namespace Chaos
 			return Center + Normalized * (GetRadius() + Thickness);
 		}
 
-		FORCEINLINE const TVector<T, d>& SupportCore(const TVector<T, d>& Direction, const FReal InMargin, FReal* OutSupportDelta) const
+		FORCEINLINE const TVector<T, d>& SupportCore(const TVector<T, d>& Direction, const FReal InMargin, FReal* OutSupportDelta, int32& VertexIndex) const
 		{
+			VertexIndex = 0;
 			// Note: ignores InMargin, assumed Radius
 			return Center;
 		}
 
-		FORCEINLINE TVector<T, d> SupportCoreScaled(const TVector<T, d>& Direction, const FReal InMargin, const TVector<T, d>& Scale, FReal* OutSupportDelta) const
+		FORCEINLINE TVector<T, d> SupportCoreScaled(const TVector<T, d>& Direction, const FReal InMargin, const TVector<T, d>& Scale, FReal* OutSupportDelta, int32& VertexIndex) const
 		{
+			VertexIndex = 0;
 			// Note: ignores InMargin, assumed Radius
 			return Center * Scale;
 		}
diff --git a/Engine/Source/Runtime/Experimental/Chaos/Public/Chaos/Triangle.h b/Engine/Source/Runtime/Experimental/Chaos/Public/Chaos/Triangle.h
--- a/Engine/Source/Runtime/Experimental/Chaos/Public/Chaos/Triangle.h
+++ b/Engine/Source/Runtime/Experimental/Chaos/Public/Chaos/Triangle.h
@@ -119,6 +119,25 @@ namespace Chaos
 			return Result;
 		}
 
+		// Get an array of all the plane indices that belong to a vertex (up to MaxVertexPlanes).
+		// Returns the number of planes found.
+		int32 FindVertexPlanes(int32 VertexIndex, int32* OutVertexPlanes, int32 MaxVertexPlanes) const
+		{
+			if(MaxVertexPlanes > 0)
+			{
+				OutVertexPlanes[0] = 0;
+			}
+			return 1; 
+		}
+		
+		// Get up to the 3  plane indices that belong to a vertex
+		// Returns the number of planes found.
+		int32 GetVertexPlanes3(int32 VertexIndex, int32& PlaneIndex0, int32& PlaneIndex1, int32& PlaneIndex2) const
+		{
+			PlaneIndex0 = 0;
+			return 1;
+		}
+		
 		// Get the index of the plane that most opposes the normal
 		int32 GetMostOpposingPlane(const FVec3& Normal) const
 		{
@@ -143,7 +162,7 @@ namespace Chaos
 			return FVec3::DotProduct((InSamplePoint - ClosestPoint), OutNormal);
 		}
 
-		FORCEINLINE FVec3 Support(const FVec3& Direction, const FReal Thickness) const
+		FORCEINLINE FVec3 Support(const FVec3& Direction, const FReal Thickness, int32& VertexIndex) const
 		{
 			const FReal DotA = FVec3::DotProduct(A, Direction);
 			const FReal DotB = FVec3::DotProduct(B, Direction);
@@ -151,6 +170,7 @@ namespace Chaos
 
 			if(DotA >= DotB && DotA >= DotC)
 			{
+				VertexIndex = 0;
 				if(Thickness != 0)
 				{
 					return A + Direction.GetUnsafeNormal() * Thickness;
@@ -159,13 +179,14 @@ namespace Chaos
 			}
 			else if(DotB >= DotA && DotB >= DotC)
 			{
+				VertexIndex = 1;
 				if(Thickness != 0)
 				{
 					return B + Direction.GetUnsafeNormal() * Thickness;
 				}
 				return B;
 			}
-
+			VertexIndex = 2;
 			if(Thickness != 0)
 			{
 				return C + Direction.GetUnsafeNormal() * Thickness;
@@ -173,7 +194,7 @@ namespace Chaos
 			return C;
 		}
 
-		FORCEINLINE_DEBUGGABLE FVec3 SupportCore(const FVec3& Direction, const FReal InMargin, FReal* OutSupportDelta) const
+		FORCEINLINE_DEBUGGABLE FVec3 SupportCore(const FVec3& Direction, const FReal InMargin, FReal* OutSupportDelta,int32& VertexIndex) const
 		{
 			// Note: assumes margin == 0
 			const FReal DotA = FVec3::DotProduct(A, Direction);
@@ -182,20 +203,22 @@ namespace Chaos
 
 			if (DotA >= DotB && DotA >= DotC)
 			{
+				VertexIndex = 0;
 				return A;
 			}
 			else if (DotB >= DotA && DotB >= DotC)
 			{
+				VertexIndex = 1;
 				return B;
 			}
-
+			VertexIndex = 2;
 			return C;
 		}
 
-		FORCEINLINE FVec3 SupportCoreScaled(const FVec3& Direction, FReal InMargin, const FVec3& Scale, FReal* OutSupportDelta) const
+		FORCEINLINE FVec3 SupportCoreScaled(const FVec3& Direction, FReal InMargin, const FVec3& Scale, FReal* OutSupportDelta, int32& VertexIndex) const
 		{
 			// Note: ignores InMargin, assumed 0 (triangles cannot have a margin as they are zero thickness)
-			return SupportCore(Direction * Scale, 0.0f, OutSupportDelta) * Scale;
+			return SupportCore(Direction * Scale, 0.0f, OutSupportDelta, VertexIndex) * Scale;
 		}
 
 		FORCEINLINE FReal GetMargin() const { return 0; }
@@ -301,9 +324,9 @@ namespace Chaos
 			return Bounds;
 		}
 
-		virtual TVec3<T> Support(const TVec3<T>& Direction, const T Thickness) const override
+		virtual TVec3<T> Support(const TVec3<T>& Direction, const T Thickness, int32& VertexIndex) const override
 		{
-			return Tri.Support(Direction, Thickness);
+			return Tri.Support(Direction, Thickness, VertexIndex);
 		}
 
 		virtual bool Raycast(const TVec3<T>& StartPoint, const TVec3<T>& Dir, const T Length, const T Thickness, T& OutTime, TVec3<T>& OutPosition, TVec3<T>& OutNormal, int32& OutFaceIndex) const override
