Chaos optimization :  use dedicated double sided ray triangle test for heighfield
measured performance gain of 47% in direct line tracing against heightfields

diff --git a/Engine/Source/Programs/HeadlessChaos/Private/HeadlessChaosTestHeightfield.cpp b/Engine/Source/Programs/HeadlessChaos/Private/HeadlessChaosTestHeightfield.cpp
--- a/Engine/Source/Programs/HeadlessChaos/Private/HeadlessChaosTestHeightfield.cpp
+++ b/Engine/Source/Programs/HeadlessChaos/Private/HeadlessChaosTestHeightfield.cpp
@@ -36,7 +36,9 @@ namespace ChaosTest {
 				const FVec3 A(0,0,0);
 				const FVec3 B(Scale[0],0,CountToWorldScale * Scale[2]);
 				const FVec3 C(0,Scale[1],Columns*CountToWorldScale * Scale[2]);
-				const FVec3 ExpectedNormal = FVec3::CrossProduct((B-A),(C-A)).GetUnsafeNormal();
+				
+				FVec3 ExpectedNormal = FVec3::CrossProduct((B - A), (C - A));
+				ExpectedNormal.SafeNormalize();
 				return ExpectedNormal;
 			};
 
@@ -65,19 +67,42 @@ namespace ChaosTest {
 						EXPECT_VECTOR_NEAR(Position,FVec3(Col*Scale[0],Row * Scale[1],Heights[Row*Columns+Col] * Scale[2]),1e-2);
 						EXPECT_VECTOR_NEAR(Normal,ExpectedNormal,1e-2);
 
+						const FVec3 RayStart = Start + FVec3(0.2 * Scale[0], 0.1 * Scale[1], 0);
+						const FVec3 RayDir = FVec3(0, 0, -1);
+						const FReal RayLength = 2000 * Scale[2];
+						const FVec3 RayEnd = RayStart + (RayDir * RayLength);
+
 						//offset in from border ever so slightly to get a clear face
-						const bool bResult = Heightfield.Raycast(Start + FVec3(0.2 * Scale[0],0.1 * Scale[1],0),FVec3(0,0,-1),2000*Scale[2],0,TOI,Position,Normal,FaceIdx);
+						const bool bResult = Heightfield.Raycast(RayStart, RayDir, RayLength,0,TOI,Position,Normal,FaceIdx);
 						if(Col + 1 == Columns || Row + 1 == Rows)
 						{
 							EXPECT_FALSE(bResult);	//went past edge so no hit
 							//Went past column so do not increment expected face idx
 						} else
 						{
+							check(bResult == true);
 							EXPECT_TRUE(bResult);
 							EXPECT_EQ(FaceIdx,ExpectedFaceIdx);	//each quad has two triangles, so for each column we pass two faces
 
 							ExpectedFaceIdx += 2;	//We hit the first triangle in the quad. Since we are going 1 quad at a time we skip 2
 						}
+
+						// reverse the ray to test double sided 
+						FVec3 ReversePosition;
+						FReal ReverseTOI;
+						const bool bReverseResult = Heightfield.Raycast(RayEnd, -RayDir, RayLength, 0, ReverseTOI, ReversePosition, Normal, FaceIdx);
+						if (Col + 1 == Columns || Row + 1 == Rows)
+						{
+							EXPECT_FALSE(bReverseResult);	//went past edge so no hit
+							//Went past column so do not increment expected face idx
+						}
+						else
+						{
+							check(bReverseResult == true);
+							EXPECT_TRUE(bReverseResult);
+							EXPECT_NEAR(RayLength, (TOI + ReverseTOI), SMALL_NUMBER);
+							EXPECT_VECTOR_NEAR(Position, ReversePosition, SMALL_NUMBER);
+						}
 					}
 				}
 			};
@@ -108,18 +133,41 @@ namespace ChaosTest {
 					for(int32 Col = 0; Col < Columns; ++Col)
 					{
 						const FVec3 Start(-Scale[0],Row * Scale[1],Heights[Row*Columns + Col] * Scale[2] + 0.01 * Scale[2]);
-						const bool bResult = Heightfield.Raycast(Start,FVec3(1,0,0),2000*Scale[0],0,TOI,Position,Normal,FaceIdx);
+
+						const FVec3 RayStart = Start;
+						const FVec3 RayDir = FVec3(1, 0, 0);
+						const FReal RayLength = 2000 * Scale[0];
+						const FVec3 RayEnd = RayStart + (RayDir * RayLength);
+
+						const bool bResult = Heightfield.Raycast(RayStart, RayDir, RayLength,0,TOI,Position,Normal,FaceIdx);
 						if(Col + 1 == Columns)
 						{
 							EXPECT_FALSE(bResult);
 							//No more columns so we shot over the final edge
 						} else
 						{
+							check(bResult);
 							EXPECT_TRUE(bResult);
 							EXPECT_NEAR(TOI,(Scale[0] * (1 + Col)),1e-1);
 							EXPECT_VECTOR_NEAR(Position,(Start + FVec3{TOI,0,0}),1e-2);
 							EXPECT_VECTOR_NEAR(Normal,ExpectedNormal,1e-1);
 						}
+
+						// reverse ray to test double sided
+						FVec3 ReversePosition;
+						FReal ReverseTOI;
+						const bool bReverseResult = Heightfield.Raycast(RayEnd, -RayDir, RayLength, 0, ReverseTOI, ReversePosition, Normal, FaceIdx);
+						if (Col + 1 == Columns)
+						{
+							EXPECT_FALSE(bReverseResult);
+							//No more columns so we shot over the final edge
+						}
+						else
+						{
+							check(bReverseResult);
+							EXPECT_TRUE(bReverseResult);
+							EXPECT_VECTOR_NEAR(Normal, (ExpectedNormal*-1), 1e-1);
+						}
 					}
 				}
 			};
@@ -149,18 +197,41 @@ namespace ChaosTest {
 					for(int32 Col = 0; Col < Columns; ++Col)
 					{
 						const FVec3 Start(Col * Scale[0],-Scale[1],Heights[Row*Columns + Col] * Scale[2] + 0.01 * Scale[2]);
-						const bool bResult = Heightfield.Raycast(Start,FVec3(0,1,0),2000*Scale[0],0,TOI,Position,Normal,FaceIdx);
+
+						const FVec3 RayStart = Start;
+						const FVec3 RayDir = FVec3(0, 1, 0);
+						const FReal RayLength = 2000 * Scale[0];
+						const FVec3 RayEnd = RayStart + (RayDir * RayLength);
+
+						const bool bResult = Heightfield.Raycast(RayStart, RayDir, RayLength,0,TOI,Position,Normal,FaceIdx);
 						if(Row + 1 == Rows)
 						{
 							EXPECT_FALSE(bResult);
 							//No more columns so we shot over the final edge
 						} else
 						{
+							check(bResult == true);
 							EXPECT_TRUE(bResult);
 							EXPECT_NEAR(TOI,(Scale[1] * (1 + Row)),1e-1);
 							EXPECT_VECTOR_NEAR(Position,(Start + FVec3{0,TOI,0}),1e-2);
 							EXPECT_VECTOR_NEAR(Normal,ExpectedNormal,1e-1);
 						}
+
+						// reverse ray to test double sided
+						FVec3 ReversePosition;
+						FReal ReverseTOI;
+						const bool bReverseResult = Heightfield.Raycast(RayEnd, -RayDir, RayLength, 0, ReverseTOI, Position, Normal, FaceIdx);
+						if (Row + 1 == Rows)
+						{
+							EXPECT_FALSE(bReverseResult);
+							//No more columns so we shot over the final edge
+						}
+						else
+						{
+							check(bReverseResult == true);
+							EXPECT_TRUE(bReverseResult);
+							EXPECT_VECTOR_NEAR(Normal, (ExpectedNormal * -1), 1e-1);
+						}
 					}
 				}
 			};
@@ -266,6 +337,7 @@ namespace ChaosTest {
 					const FVec3 Dir = (End - Start).GetUnsafeNormal();
 
 					const bool bResult = Heightfield.Raycast(Start,Dir,2000,0,TOI,Position,Normal,FaceIdx);
+					check(bResult);
 					EXPECT_TRUE(bResult);
 					EXPECT_VECTOR_NEAR(Normal,ExpectedNormal,1e-1);
 					EXPECT_VECTOR_NEAR(Position,End,1e-1);
diff --git a/Engine/Source/Runtime/Experimental/Chaos/Private/Chaos/HeightField.cpp b/Engine/Source/Runtime/Experimental/Chaos/Private/Chaos/HeightField.cpp
--- a/Engine/Source/Runtime/Experimental/Chaos/Private/Chaos/HeightField.cpp
+++ b/Engine/Source/Runtime/Experimental/Chaos/Private/Chaos/HeightField.cpp
@@ -48,8 +48,125 @@ namespace Chaos
 			Sweep
 		};
 
-		template<ERaycastType SQType>
-		bool Visit(int32 Payload, FReal& CurrentLength)
+		/**
+		* Ray / triangle  intersection
+		* this provides a double sided test
+		* note : this method assumes that the triangle formed by A,B and C is well formed
+		*/
+		FORCEINLINE_DEBUGGABLE bool RayTriangleIntersection(
+			const FVec3& RayStart, const FVec3& RayDir, FReal RayLength,
+			const FVec3& A, const FVec3& B, const FVec3& C,
+			FReal& OutT, FVec3& OutN
+		)
+		{
+			const FVec3 AB = B - A; // edge 1
+			const FVec3 AC = C - A; // edge 2
+			const FVec3 Normal = FVec3::CrossProduct(AB, AC);
+			const FVec3 NegRayDir = -RayDir;
+
+			FReal Den = FVec3::DotProduct(NegRayDir, Normal);
+			if (FMath::Abs(Den) < SMALL_NUMBER)
+			{
+				// ray is parallel or away to the triangle plane it is a miss
+				return false;
+			}
+
+			const FReal InvDen = (FReal)1 / Den;
+
+			// let's compute the time to intersection
+			const FVec3 RayToA = RayStart - A;
+			const FReal Time = FVec3::DotProduct(RayToA, Normal) * InvDen;
+			if (Time < (FReal)0 || Time > RayLength)
+			{
+				return false;
+			}
+
+			// now compute baricentric coordinates
+			const FVec3 RayToACrossNegDir = FVec3::CrossProduct(NegRayDir, RayToA);
+			constexpr FReal Epsilon = SMALL_NUMBER;
+			const FReal UU = FVec3::DotProduct(AC, RayToACrossNegDir) * InvDen;
+			if (UU < -Epsilon || UU >(1 + Epsilon))
+			{
+				return false; // outside of the triangle
+			}
+			const FReal VV = -FVec3::DotProduct(AB, RayToACrossNegDir) * InvDen;
+			if (VV < -Epsilon || (VV + UU) >(1 + Epsilon))
+			{
+				return false; // outside of the triangle
+			}
+
+			// point is within the triangle, let's compute 
+			OutT = Time;
+			OutN = Normal.GetSafeNormal();
+			OutN *= FMath::Sign(Den);
+			return true;
+		}
+
+
+		FORCEINLINE_DEBUGGABLE bool VisitRaycast(int32 Payload, FReal& CurrentLength)
+		{
+			const int32 SubX = Payload % (GeomData->NumCols - 1);
+			const int32 SubY = Payload / (GeomData->NumCols - 1);
+
+			const int32 FullIndex = Payload + SubY;
+
+			const FReal Radius = Thickness + SMALL_NUMBER;
+			const FReal Radius2 = Radius * Radius;
+			bool bIntersection = false;
+
+			// return if the triangle was hit or not 
+			auto TestTriangle = [&](int32 FaceIndex, const FVec3& A, const FVec3& B, const FVec3& C) -> bool
+			{
+				FReal Time;
+				FVec3 Normal;
+				if (RayTriangleIntersection(Start, Dir, CurrentLength, A, B, C, Time, Normal))
+				{
+					if (Time < OutTime)
+					{
+						bool bHole = false;
+
+						const int32 CellIndex = FaceIndex / 2;
+						if (GeomData->MaterialIndices.IsValidIndex(CellIndex))
+						{
+							bHole = GeomData->MaterialIndices[CellIndex] == TNumericLimits<uint8>::Max();
+						}
+
+						if (!bHole)
+						{
+							OutPosition = Start + (Dir * Time);
+							OutNormal = Normal;
+							OutTime = Time;
+							OutFaceIndex = FaceIndex;
+							CurrentLength = Time;
+							return true;
+						}
+					}
+				}
+
+				return false;
+			};
+
+			FVec3 Points[4];
+			FAABB3 CellBounds;
+			GeomData->GetPointsAndBoundsScaled(FullIndex, Points, CellBounds);
+			CellBounds.Thicken(Thickness);
+
+			// Check cell bounds
+			//todo: can do it without raycast
+			FReal TOI;
+			FVec3 HitPoint;
+			bool bHit = false;
+			if (CellBounds.RaycastFast(Start, Dir, InvDir, bParallel, CurrentLength, 1.0f / CurrentLength, TOI, HitPoint))
+			{
+				// Test both triangles that are in this cell, as we could hit both in any order
+				bHit |= TestTriangle(Payload * 2, Points[0], Points[1], Points[3]);
+				bHit |= TestTriangle(Payload * 2 + 1, Points[0], Points[3], Points[2]);
+			}
+			const bool bShouldContinueVisiting = !bHit;
+			return bShouldContinueVisiting;
+		}
+
+		bool VisitSweep(int32 Payload, FReal& CurrentLength)
 		{
 			const int32 SubX = Payload % (GeomData->NumCols - 1);
 			const int32 SubY = Payload / (GeomData->NumCols - 1);
@@ -106,7 +223,7 @@ namespace Chaos
 					}
 				}
 
-				if(SQType == ERaycastType::Sweep && !bIntersection)
+				if(!bIntersection)
 				{
 					//sphere is not immediately touching the triangle, but it could start intersecting the perimeter as it sweeps by
 					FVec3 BorderPositions[3];
@@ -203,16 +320,6 @@ namespace Chaos
 			return bShouldContinueVisiting;
 		}
 
-		bool VisitRaycast(int32 Payload, FReal& CurLength)
-		{
-			return Visit<ERaycastType::Raycast>(Payload, CurLength);
-		}
-
-		bool VisitSweep(int32 Payload, FReal& CurLength)
-		{
-			return Visit<ERaycastType::Sweep>(Payload, CurLength);
-		}
-
 		FReal OutTime;
 		FVec3 OutPosition;
 		FVec3 OutNormal;
