Fix for bug where constraint was operating on disabled particles.

diff --git a/Engine/Source/Programs/HeadlessChaos/Private/HeadlessChaosTestJoints.cpp b/Engine/Source/Programs/HeadlessChaos/Private/HeadlessChaosTestJoints.cpp
--- a/Engine/Source/Programs/HeadlessChaos/Private/HeadlessChaosTestJoints.cpp
+++ b/Engine/Source/Programs/HeadlessChaos/Private/HeadlessChaosTestJoints.cpp
@@ -585,6 +585,67 @@ namespace ChaosTest {
 		}
 	}
 
+	// check that joints don't simulate if a constrained particle is disabled
+	template <typename TEvolution>
+	void JointConstraint_DisableOneConstrainedParticle()
+	{
+
+		const int32 NumIterations = 1;
+		const FReal Gravity = 980;
+
+		FJointConstraintsTest<TEvolution> Test(NumIterations, Gravity);
+
+		Test.ParticlePositions = {
+			{ (FReal)0, (FReal)0, (FReal)1000 },
+			{ (FReal)500, (FReal)0, (FReal)1000 },
+		};
+		Test.ParticleSizes =
+		{
+			{ (FReal)100, (FReal)100, (FReal)100 },
+			{ (FReal)100, (FReal)100, (FReal)100 },
+		};
+		Test.ParticleMasses =
+		{
+			(FReal)1,
+			(FReal)1,
+		};
+
+		Test.JointPositions =
+		{
+			{ (FReal)250, (FReal)0, (FReal)1000 },
+		};
+		Test.JointParticleIndices =
+		{
+			{ 0, 1 },
+		};
+
+		Test.Create();
+
+		const int32 Box1Id = 0;
+		const int32 Box2Id = 1;
+		const FReal ExpectedDistance = (Test.ParticlePositions[1] - Test.ParticlePositions[0]).Size();
+		const FVec3 Box2LocalSpaceJointPosition = Test.JointPositions[0] - Test.ParticlePositions[1];
+
+		// box 1 disabled
+		Test.Evolution.DisableParticle(Test.GetParticle(Box1Id));
+
+		const FReal Dt = 0.01f;
+		for (int32 i = 0; i < 100; ++i)
+		{
+			Test.Evolution.AdvanceOneTimeStep(Dt);
+			Test.Evolution.EndFrame(Dt);
+
+			// box 1 not simulating so would expect to not have moved
+			EXPECT_LT(FMath::Abs(Test.ParticlePositions[0].X - Test.GetParticle(Box1Id)->X().X), (FReal)0.1);
+			EXPECT_LT(FMath::Abs(Test.ParticlePositions[0].Y - Test.GetParticle(Box1Id)->X().Y), (FReal)0.1);
+			EXPECT_LT(FMath::Abs(Test.ParticlePositions[0].Z - Test.GetParticle(Box1Id)->X().Z), (FReal)0.09);
+
+			// box 2 should fall under gravity & not have moved in X or Y, constraint should not 'Apply' if other particle is disabled
+			EXPECT_LT(FMath::Abs(Test.ParticlePositions[1].X - Test.GetParticle(Box2Id)->X().X), (FReal)0.1);
+			EXPECT_LT(FMath::Abs(Test.ParticlePositions[1].Y - Test.GetParticle(Box2Id)->X().Y), (FReal)0.1);
+			EXPECT_GT(FMath::Abs(Test.ParticlePositions[1].Z - Test.GetParticle(Box2Id)->X().Z), (FReal)0.09);
+		}
+	}
 
 
 	GTEST_TEST(AllEvolutions, JointTests_TestSingleConstraint) {
@@ -615,6 +676,9 @@ namespace ChaosTest {
 		DynamicSpringConstraint<FPBDRigidsEvolutionGBF>();
 	}
 
+	GTEST_TEST(AllEvolutions, JointConstraint_TestDisableOneConstrainedParticle) {
+		JointConstraint_DisableOneConstrainedParticle<FPBDRigidsEvolutionGBF>();
+	}
 
 	// Check that constraints end up in the same island when graph is fully connected
 	GTEST_TEST(JointTests, TestJointConstraintGraph_Connected)
@@ -681,5 +745,7 @@ namespace ChaosTest {
 		EXPECT_EQ(Joints[0]->GetConstraintLevel(), 0);
 		EXPECT_EQ(Joints[1]->GetConstraintLevel(), 0);
 	}
+
+
 }
 
diff --git a/Engine/Source/Runtime/Experimental/Chaos/Private/Chaos/PBDJointConstraints.cpp b/Engine/Source/Runtime/Experimental/Chaos/Private/Chaos/PBDJointConstraints.cpp
--- a/Engine/Source/Runtime/Experimental/Chaos/Private/Chaos/PBDJointConstraints.cpp
+++ b/Engine/Source/Runtime/Experimental/Chaos/Private/Chaos/PBDJointConstraints.cpp
@@ -943,11 +943,43 @@ namespace Chaos
 		return FMath::Clamp(IterationStiffness, 0.0f, 1.0f);
 	}
 
+
+	bool FPBDJointConstraints::CanEvaluate(const int32 ConstraintIndex) const
+	{
+		if (!IsConstraintEnabled(ConstraintIndex))
+		{
+			return false;
+		}
+
+		int32 Index0, Index1;
+		GetConstrainedParticleIndices(ConstraintIndex, Index0, Index1);
+		const FGenericParticleHandle Particle0 = FGenericParticleHandle(ConstraintParticles[ConstraintIndex][Index0]);
+		const FGenericParticleHandle Particle1 = FGenericParticleHandle(ConstraintParticles[ConstraintIndex][Index1]);
+
+		// check for valid and enabled particles
+		if (Particle0->Handle() == nullptr || Particle0->Disabled()
+			|| Particle1->Handle() == nullptr || Particle1->Disabled())
+		{
+			return false;
+		}
+
+		// check valid particle and solver state
+		const FPBDJointSolver& Solver = ConstraintSolvers[ConstraintIndex];
+		if ((Particle0->Sleeping() && Particle1->Sleeping())
+			|| (Particle0->IsKinematic() && Particle1->Sleeping())
+			|| (Particle0->Sleeping() && Particle1->IsKinematic())
+			|| (FMath::IsNearlyZero(Solver.InvM(0)) && FMath::IsNearlyZero(Solver.InvM(1))))
+		{
+			return false;
+		}
+		return true;
+	}
+
 	// This position solver iterates over each of the inner constraints (position, twist, swing) and solves them independently.
 	// This will converge slowly in some cases, particularly where resolving angular constraints violates position constraints and vice versa.
 	bool FPBDJointConstraints::ApplyPhase1Single(const FReal Dt, const int32 ConstraintIndex, const int32 NumPairIts, const int32 It, const int32 NumIts)
 	{
-		if (!IsConstraintEnabled(ConstraintIndex))
+		if (!CanEvaluate(ConstraintIndex))
 		{
 			return false;
 		}
@@ -1013,7 +1045,7 @@ namespace Chaos
 
 	bool FPBDJointConstraints::ApplyPhase2Single(const FReal Dt, const int32 ConstraintIndex, const int32 NumPairIts, const int32 It, const int32 NumIts)
 	{
-		if (!IsConstraintEnabled(ConstraintIndex))
+		if (!CanEvaluate(ConstraintIndex))
 		{
 			return false;
 		}
