JointBreak_UnderLinearThreshold3

diff --git a/Engine/Source/Programs/HeadlessChaos/Private/HeadlessChaosTestJointBreak.cpp b/Engine/Source/Programs/HeadlessChaos/Private/HeadlessChaosTestJointBreak.cpp
--- a/Engine/Source/Programs/HeadlessChaos/Private/HeadlessChaosTestJointBreak.cpp
+++ b/Engine/Source/Programs/HeadlessChaos/Private/HeadlessChaosTestJointBreak.cpp
@@ -193,7 +193,7 @@ namespace ChaosTest {
 
 
 	// 1 Kinematic Body with 2 Dynamic bodies hanging from it by a breakable constraint.
-	// Constraint break force is less than M x G, so joint should break.
+	// Constraint break force is less than M x G, so joint should not break.
 	template <typename TEvolution>
 	void JointBreak_UnderLinearThreshold3()
 	{
@@ -225,7 +225,7 @@ namespace ChaosTest {
 			Test.Evolution.EndFrame(Dt);
 		}
 
-		// Constraint should have broken
+		// Constraint should not have broken
 		EXPECT_FALSE(bBrokenCallbackCalled);
 		EXPECT_TRUE(Test.Evolution.GetJointConstraints().IsConstraintEnabled(0));
 		EXPECT_TRUE(Test.Evolution.GetJointConstraints().IsConstraintEnabled(1));
diff --git a/Engine/Source/Runtime/Experimental/Chaos/Private/Chaos/PBDJointConstraintUtilities.cpp b/Engine/Source/Runtime/Experimental/Chaos/Private/Chaos/PBDJointConstraintUtilities.cpp
--- a/Engine/Source/Runtime/Experimental/Chaos/Private/Chaos/PBDJointConstraintUtilities.cpp
+++ b/Engine/Source/Runtime/Experimental/Chaos/Private/Chaos/PBDJointConstraintUtilities.cpp
@@ -700,11 +700,21 @@ namespace Chaos
 			OutInvMParent = (FReal)1 / MParent;
 			OutInvIParent = FVec3((FReal)1 / IParent.X, (FReal)1 / IParent.Y, (FReal)1 / IParent.Z);
 		}
+		else
+		{
+			OutInvMParent = 0.f;
+			OutInvIParent = FVec3(0.f);
+		}
 		if (InInvMChild > 0)
 		{
 			OutInvMChild = (FReal)1 / MChild;
 			OutInvIChild = FVec3((FReal)1 / IChild.X, (FReal)1 / IChild.Y, (FReal)1 / IChild.Z);
 		}
+		else
+		{
+			OutInvMChild = 0.f;
+			OutInvIChild = FVec3(0.f);
+		}
 	}
 
 	void FPBDJointUtilities::ConditionInverseMassAndInertia(
