Chaos : fix unit test compilation issue

diff --git a/Engine/Source/Programs/HeadlessChaos/Private/HeadlessChaosTestGJK.cpp b/Engine/Source/Programs/HeadlessChaos/Private/HeadlessChaosTestGJK.cpp
--- a/Engine/Source/Programs/HeadlessChaos/Private/HeadlessChaosTestGJK.cpp
+++ b/Engine/Source/Programs/HeadlessChaos/Private/HeadlessChaosTestGJK.cpp
@@ -1916,7 +1916,7 @@ namespace ChaosTest
 
 		// Deep (EPA) With Margin
 		WarmStartData = FGJKSimplexData();
-		TGJKCoreShape MarginBox(Box, FReal(10));
+		TGJKCoreShape<FAABB3> MarginBox(Box, FReal(10));
 		GJKPenetrationWarmStartable(MarginBox, MarginBox, BTM.GetRelativeTransformNoScale(ATM), Penetration, ClosestA, ClosestB, NormalA, NormalB,  VertexIndexA, VertexIndexB, WarmStartData, SupportDelta);
 		EXPECT_NEAR(Penetration, 40.0f, KINDA_SMALL_NUMBER);
 		EXPECT_NEAR(ClosestA.Z, (FReal)50.0f, (FReal)KINDA_SMALL_NUMBER);
