CIS fix, HeadlessChaos file missing from changelist 18615358

diff --git a/Engine/Source/Programs/HeadlessChaos/Private/GeometryCollection/GeometryCollectionTestEvents.cpp b/Engine/Source/Programs/HeadlessChaos/Private/GeometryCollection/GeometryCollectionTestEvents.cpp
--- a/Engine/Source/Programs/HeadlessChaos/Private/GeometryCollection/GeometryCollectionTestEvents.cpp
+++ b/Engine/Source/Programs/HeadlessChaos/Private/GeometryCollection/GeometryCollectionTestEvents.cpp
@@ -95,7 +95,7 @@ namespace GeometryCollectionTest
 		EventTestData TestData;
 		EventTestData* TestDataPtr = &TestData;
 		EventManager.template RegisterEvent<EventTestData>(CustomEvent1, [TestDataPtr]
-		(const auto* Solver, EventTestData& MyData)
+		(const auto* Solver, EventTestData& MyData, bool ResetData)
 		{
 			MyData = *TestDataPtr;
 		});
@@ -104,7 +104,7 @@ namespace GeometryCollectionTest
 		TArray<EventTestData> TestArrayData;
 		TArray<EventTestData>* TestDataPtr2 = &TestArrayData;
 		EventManager.template RegisterEvent<TArray<EventTestData>>(CustomEvent2, [TestDataPtr2]
-		(const auto* Solver, TArray<EventTestData>& MyData)
+		(const auto* Solver, TArray<EventTestData>& MyData, bool ResetData)
 		{
 			MyData = *TestDataPtr2;
 		});
