Fix CIS - Chaos unit tests

diff --git a/Engine/Source/Programs/HeadlessChaos/Private/HeadlessChaosTestCollisions.cpp b/Engine/Source/Programs/HeadlessChaos/Private/HeadlessChaosTestCollisions.cpp
--- a/Engine/Source/Programs/HeadlessChaos/Private/HeadlessChaosTestCollisions.cpp
+++ b/Engine/Source/Programs/HeadlessChaos/Private/HeadlessChaosTestCollisions.cpp
@@ -73,11 +73,11 @@ namespace ChaosTest {
 
 		EXPECT_EQ(Constraint.GetParticle0(), Box2);
 		EXPECT_EQ(Constraint.GetParticle1(), Box1);
-		EXPECT_TRUE(Constraint.GetNormal().operator==(FVec3(0, 0, 1)));
+		EXPECT_TRUE(Constraint.CalculateWorldContactNormal().operator==(FVec3(0, 0, 1)));
 
 		// The contact point is the average of the surface contact points, so it should be half of Phi inside particle0
 		EXPECT_NEAR(Constraint.GetPhi(), -0.1f, KINDA_SMALL_NUMBER);
-		EXPECT_NEAR(Constraint.GetLocation().Z, 1.45f, KINDA_SMALL_NUMBER);
+		EXPECT_NEAR(Constraint.CalculateWorldContactLocation().Z, 1.45f, KINDA_SMALL_NUMBER);
 	}
 
 	void LevelsetConstraintGJK()
@@ -120,10 +120,10 @@ namespace ChaosTest {
 		
 		EXPECT_EQ(Constraint.GetParticle0(), Box2);
 		EXPECT_EQ(Constraint.GetParticle1(), Box1);
-		EXPECT_TRUE(Constraint.GetNormal().operator==(FVec3(0, 0, 1)));
+		EXPECT_TRUE(Constraint.CalculateWorldContactNormal().operator==(FVec3(0, 0, 1)));
 
 		// The contact point is the average of the surface contact points, so it should be half of Phi inside particle0
-		Chaos::FReal Distance = ChaosTest::SignedDistance(*Constraint.GetParticle0(), Constraint.GetLocation());
+		Chaos::FReal Distance = ChaosTest::SignedDistance(*Constraint.GetParticle0(), Constraint.CalculateWorldContactLocation());
 		EXPECT_NEAR(Distance, 0.5f * Constraint.GetPhi(), KINDA_SMALL_NUMBER);
 	}
 	
@@ -176,7 +176,7 @@ namespace ChaosTest {
 
 		EXPECT_EQ(Constraint.GetParticle0(), Box);
 		EXPECT_EQ(Constraint.GetParticle1(), Floor);
-		EXPECT_TRUE(Constraint.GetNormal().operator==(FVec3(0, 0, 1)));
+		EXPECT_TRUE(Constraint.CalculateWorldContactNormal().operator==(FVec3(0, 0, 1)));
 		EXPECT_TRUE(FMath::Abs(Constraint.GetPhi() - FReal(-0.5) ) < SMALL_THRESHOLD );
 
 		{
@@ -259,7 +259,7 @@ namespace ChaosTest {
 
 		EXPECT_EQ(Constraint->GetParticle0(), Box);
 		EXPECT_EQ(Constraint->GetParticle1(), Floor);
-		EXPECT_TRUE(Constraint->GetNormal().operator==(FVec3(0, 0, 1)));
+		EXPECT_TRUE(Constraint->CalculateWorldContactNormal().operator==(FVec3(0, 0, 1)));
 		EXPECT_TRUE(FMath::Abs(Constraint->GetPhi() - FReal(-1.0)) < SMALL_THRESHOLD);
 
 		{
@@ -335,7 +335,7 @@ namespace ChaosTest {
 		FPBDCollisionConstraint& Constraint = Collisions.GetConstraint(0);
 		EXPECT_EQ(Constraint.GetParticle0(), Box);
 		EXPECT_EQ(Constraint.GetParticle1(), Floor);
-		EXPECT_TRUE(Constraint.GetNormal().operator==(FVec3(0, 0, 1)));
+		EXPECT_TRUE(Constraint.CalculateWorldContactNormal().operator==(FVec3(0, 0, 1)));
 		EXPECT_TRUE(FMath::Abs(Constraint.GetPhi() - FReal(-0.5)) < SMALL_THRESHOLD);
 
 		{
@@ -413,7 +413,7 @@ namespace ChaosTest {
 		Collisions.UpdateLevelsetConstraint(Constraint);
 		EXPECT_EQ(Constraint.GetParticle0(), Box);
 		EXPECT_EQ(Constraint.GetParticle1(), Floor);
-		EXPECT_TRUE(Constraint.GetNormal().operator==(FVec3(0, 0, 1)));
+		EXPECT_TRUE(Constraint.CalculateWorldContactNormal().operator==(FVec3(0, 0, 1)));
 		EXPECT_TRUE(FMath::Abs(Constraint.GetPhi() - FReal(-0.5)) < SMALL_THRESHOLD);
 
 		{
@@ -500,7 +500,7 @@ namespace ChaosTest {
 
 		EXPECT_EQ(Constraint.GetParticle0(), DynamicCube);
 		EXPECT_EQ(Constraint.GetParticle1(), StaticCube);
-		EXPECT_NEAR(Constraint.GetNormal().Z, FReal(1), KINDA_SMALL_NUMBER);
+		EXPECT_NEAR(Constraint.CalculateWorldContactNormal().Z, FReal(1), KINDA_SMALL_NUMBER);
 		{
 			INVARIANT_XR_START(DynamicCube);
 			Collisions.Apply(Dt, 0, 1);
@@ -567,7 +567,7 @@ namespace ChaosTest {
 
 		EXPECT_EQ(Constraint.GetParticle0(), Box2);
 		EXPECT_EQ(Constraint.GetParticle1(), StaticBox);
-		EXPECT_TRUE(Constraint.GetNormal().Equals(FVector(0.0, 0.0, 1.0f)));
+		EXPECT_TRUE(Constraint.CalculateWorldContactNormal().Equals(FVector(0.0, 0.0, 1.0f)));
 		EXPECT_TRUE(FMath::Abs(Constraint.GetPhi() - FReal(-0.4)) < SMALL_THRESHOLD);
 
 
diff --git a/Engine/Source/Programs/HeadlessChaos/Private/HeadlessChaosTestShapePadding.cpp b/Engine/Source/Programs/HeadlessChaos/Private/HeadlessChaosTestShapePadding.cpp
--- a/Engine/Source/Programs/HeadlessChaos/Private/HeadlessChaosTestShapePadding.cpp
+++ b/Engine/Source/Programs/HeadlessChaos/Private/HeadlessChaosTestShapePadding.cpp
@@ -105,9 +105,9 @@ namespace ChaosTest {
 		Collisions::UpdateConstraintFromGeometry<ECollisionUpdateType::Deepest>(*Constraint, FParticleUtilitiesPQ::GetActorWorldTransform(Box0), FParticleUtilitiesPQ::GetActorWorldTransform(Box1), 1 / 30.0f);
 
 		EXPECT_NEAR(Constraint->GetPhi(), ExpectedPhi, Tolerance);
-		EXPECT_NEAR(Constraint->GetNormal().X, ExpectedNormal.X, Tolerance);
-		EXPECT_NEAR(Constraint->GetNormal().Y, ExpectedNormal.Y, Tolerance);
-		EXPECT_NEAR(Constraint->GetNormal().Z, ExpectedNormal.Z, Tolerance);
+		EXPECT_NEAR(Constraint->CalculateWorldContactNormal().X, ExpectedNormal.X, Tolerance);
+		EXPECT_NEAR(Constraint->CalculateWorldContactNormal().Y, ExpectedNormal.Y, Tolerance);
+		EXPECT_NEAR(Constraint->CalculateWorldContactNormal().Z, ExpectedNormal.Z, Tolerance);
 	}
 
 	TEST(CollisionTests, TestBoxBoxCollisionMargin)
@@ -214,9 +214,9 @@ namespace ChaosTest {
 		Collisions::UpdateConstraintFromGeometry<ECollisionUpdateType::Deepest>(*Constraint, FParticleUtilitiesPQ::GetActorWorldTransform(Box0), FParticleUtilitiesPQ::GetActorWorldTransform(Box1), 1 / 30.0f);
 
 		EXPECT_NEAR(Constraint->GetPhi(), ExpectedPhi, Tolerance);
-		EXPECT_NEAR(Constraint->GetNormal().X, ExpectedNormal.X, Tolerance);
-		EXPECT_NEAR(Constraint->GetNormal().Y, ExpectedNormal.Y, Tolerance);
-		EXPECT_NEAR(Constraint->GetNormal().Z, ExpectedNormal.Z, Tolerance);
+		EXPECT_NEAR(Constraint->CalculateWorldContactNormal().X, ExpectedNormal.X, Tolerance);
+		EXPECT_NEAR(Constraint->CalculateWorldContactNormal().Y, ExpectedNormal.Y, Tolerance);
+		EXPECT_NEAR(Constraint->CalculateWorldContactNormal().Z, ExpectedNormal.Z, Tolerance);
 	}
 
 
