HeadlessChaos - Compile fixes after material container changes

diff --git a/Engine/Source/Programs/HeadlessChaos/Private/HeadlessChaosTestMinEvolution.cpp b/Engine/Source/Programs/HeadlessChaos/Private/HeadlessChaosTestMinEvolution.cpp
--- a/Engine/Source/Programs/HeadlessChaos/Private/HeadlessChaosTestMinEvolution.cpp
+++ b/Engine/Source/Programs/HeadlessChaos/Private/HeadlessChaosTestMinEvolution.cpp
@@ -40,7 +40,7 @@ namespace ChaosTest
 		TArrayCollectionArray<TUniquePtr<Chaos::FChaosPhysicsMaterial>> PerParticleMaterials;
 		TArrayCollectionArray<FVec3> ParticlePrevXs;
 		TArrayCollectionArray<FRotation3> ParticlePrevRs;
-		FCollisionConstraints Collisions(ParticlesContainer, CollidedParticles, ParticleMaterials, PerParticleMaterials);
+		FCollisionConstraints Collisions(ParticlesContainer, CollidedParticles, ParticleMaterials, PerParticleMaterials, nullptr);
 		FBasicBroadPhase BroadPhase(&ActivePotentiallyCollidingPairs, nullptr, nullptr);
 		FNarrowPhase NarrowPhase(0, 0, Collisions.GetConstraintAllocator());
 		FCollisionDetector CollisionDetector(BroadPhase, NarrowPhase, Collisions);
diff --git a/Engine/Source/Programs/HeadlessChaos/Public/HeadlessChaosCollisionConstraints.h b/Engine/Source/Programs/HeadlessChaos/Public/HeadlessChaosCollisionConstraints.h
--- a/Engine/Source/Programs/HeadlessChaos/Public/HeadlessChaosCollisionConstraints.h
+++ b/Engine/Source/Programs/HeadlessChaos/Public/HeadlessChaosCollisionConstraints.h
@@ -38,7 +38,7 @@ public:
 	FPBDCollisionConstraintAccessor()
 		: EmptyParticles(UniqueIndices)
 		, SpatialAcceleration(EmptyParticles.GetNonDisabledView())
-		, CollisionConstraints(EmptyParticles, EmptyCollided, EmptyPhysicsMaterials, EmptyUniquePhysicsMaterials, 1, 1)
+		, CollisionConstraints(EmptyParticles, EmptyCollided, EmptyPhysicsMaterials, EmptyUniquePhysicsMaterials, nullptr, 1, 1)
 		, NarrowPhase((FReal)1, (FReal)0, CollisionConstraints.GetConstraintAllocator())
 		, BroadPhase(EmptyParticles)
 		, CollisionDetector(BroadPhase, NarrowPhase, CollisionConstraints)
@@ -52,7 +52,7 @@ public:
 		const int32 PushOutIterations, const int32 PushOutPairIterations) 
 		: EmptyParticles(UniqueIndices)
 		, SpatialAcceleration(InParticles.GetNonDisabledView())
-		, CollisionConstraints(InParticles, Collided, PerParticleMaterials, PerParticleUniqueMaterials, 1, 1)
+		, CollisionConstraints(InParticles, Collided, PerParticleMaterials, PerParticleUniqueMaterials, nullptr, 1, 1)
 		, NarrowPhase((FReal)1, (FReal)0, CollisionConstraints.GetConstraintAllocator())
 		, BroadPhase(InParticles)
 		, CollisionDetector(BroadPhase, NarrowPhase, CollisionConstraints)
