Chaos Cloth - Fix the Headless Chaos cloth test broken by CL 18842837.

diff --git a/Engine/Source/Programs/HeadlessChaos/Private/HeadlessChaosTestCloth.cpp b/Engine/Source/Programs/HeadlessChaos/Private/HeadlessChaosTestCloth.cpp
--- a/Engine/Source/Programs/HeadlessChaos/Private/HeadlessChaosTestCloth.cpp
+++ b/Engine/Source/Programs/HeadlessChaos/Private/HeadlessChaosTestCloth.cpp
@@ -5,6 +5,7 @@
 #include "HeadlessChaos.h"
 #include "HeadlessChaosTestUtility.h"
 
+#include "Chaos/PBDSoftsEvolutionFwd.h"
 #include "Chaos/PBDEvolution.h"
 #include "Chaos/GeometryParticles.h"
 #include "Chaos/PBDAxialSpringConstraints.h"
@@ -19,19 +20,17 @@ namespace ChaosTest {
 	using namespace Chaos;
 	DEFINE_LOG_CATEGORY_STATIC(LogChaosTestCloth, Verbose, All);
 
-	using FPBDEvolution = FPBDEvolution;
-
-	TUniquePtr<FPBDEvolution> InitPBDEvolution(
+	TUniquePtr<Softs::FPBDEvolution> InitPBDEvolution(
 		const int32 NumIterations=1,
-		const FReal CollisionThickness=KINDA_SMALL_NUMBER,
-		const FReal SelfCollisionThickness=KINDA_SMALL_NUMBER,
-		const FReal Friction=0.0,
-		const FReal Damping=0.04)
+		const Softs::FSolverReal CollisionThickness=KINDA_SMALL_NUMBER,
+		const Softs::FSolverReal SelfCollisionThickness=KINDA_SMALL_NUMBER,
+		const Softs::FSolverReal Friction=0.0,
+		const Softs::FSolverReal Damping=0.04)
 	{
-		Chaos::FPBDParticles Particles;
-		Chaos::FKinematicGeometryClothParticles RigidParticles;
-		TUniquePtr<FPBDEvolution> Evolution(
-			new FPBDEvolution(
+		Chaos::Softs::FSolverParticles Particles;
+		Chaos::Softs::FSolverRigidParticles RigidParticles;
+		TUniquePtr<Softs::FPBDEvolution> Evolution(
+			new Softs::FPBDEvolution(
 				MoveTemp(Particles),
 				MoveTemp(RigidParticles),
 				{},
@@ -46,9 +45,9 @@ namespace ChaosTest {
 	template <class TEvolutionPtr>
 	void InitSingleParticle(
 		TEvolutionPtr& Evolution,
-		const FVec3& Position = FVec3(0),
-		const FVec3& Velocity = FVec3(0),
-		const FReal Mass = 1.0)
+		const Softs::FSolverVec3& Position = Softs::FSolverVec3(0),
+		const Softs::FSolverVec3& Velocity = Softs::FSolverVec3(0),
+		const Softs::FSolverReal Mass = 1.0)
 	{
 		auto& Particles = Evolution->Particles();
 		const uint32 Idx = Particles.Size();
@@ -63,22 +62,23 @@ namespace ChaosTest {
 	void InitTriMesh_EquilateralTri(
 		FTriangleMesh& TriMesh,
 		TEvolutionPtr& Evolution, 
-		const FVec3& XOffset=FVec3(0))
+		const Softs::FSolverVec3& XOffset=Softs::FSolverVec3(0))
 	{
 		auto& Particles = Evolution->Particles();
 		const uint32 InitialNumParticles = Particles.Size();
+
 		FTriangleMesh::InitEquilateralTriangleYZ(TriMesh, Particles);
 
 		// Initialize particles.  Use 1/3 area of connected triangles for particle mass.
 		for (uint32 i = InitialNumParticles; i < Particles.Size(); i++)
 		{
 			Particles.X(i) += XOffset;
-			Particles.V(i) = Chaos::FVec3(0);
+			Particles.V(i) = Chaos::Softs::FSolverVec3(0);
 			Particles.M(i) = 0;
 		}
 		for (const Chaos::TVec3<int32>& Tri : TriMesh.GetElements())
 		{
-			const FReal TriArea = 0.5 * Chaos::FVec3::CrossProduct(
+			const Softs::FSolverReal TriArea = 0.5 * Chaos::Softs::FSolverVec3::CrossProduct(
 				Particles.X(Tri[1]) - Particles.X(Tri[0]),
 				Particles.X(Tri[2]) - Particles.X(Tri[0])).Size();
 			for (int32 i = 0; i < 3; i++)
@@ -97,7 +97,7 @@ namespace ChaosTest {
 	void AddEdgeLengthConstraint(
 		TEvolutionPtr& Evolution,
 		const TArray<Chaos::TVec3<int32>>& Topology,
-		const FReal Stiffness)
+		const Softs::FSolverReal Stiffness)
 	{
 		check(Stiffness >= 0. && Stiffness <= 1.);
 		// TODO: Use Add AddConstraintRuleRange
@@ -114,7 +114,7 @@ namespace ChaosTest {
 	void AddAxialConstraint(
 		TEvolutionPtr& Evolution,
 		TArray<Chaos::TVec3<int32>>&& Topology,
-		const FReal Stiffness)
+		const Softs::FSolverReal Stiffness)
 	{
 		check(Stiffness >= 0. && Stiffness <= 1.);
 		// TODO: Use Add AddConstraintRuleRange
@@ -138,7 +138,7 @@ namespace ChaosTest {
 		Evolution->SetIterations(NumTimeStepsPerFrame);
 
 		check(FPS > 0);
-		const FReal Dt = 1.0 / FPS;
+		const Softs::FSolverReal Dt = 1.0 / FPS;
 		for (uint32 i = 0; i < NumFrames; i++)
 		{
 			Evolution->AdvanceOneTimeStep(Dt);
@@ -146,9 +146,9 @@ namespace ChaosTest {
 	}
 
 	template <class TParticleContainer>
-	TArray<Chaos::FVec3> CopyPoints(const TParticleContainer& Particles)
+	TArray<Chaos::Softs::FSolverVec3> CopyPoints(const TParticleContainer& Particles)
 	{
-		TArray<Chaos::FVec3> Points;
+		TArray<Chaos::Softs::FSolverVec3> Points;
 		Points.SetNum(Particles.Size());
 
 		for (uint32 i = 0; i < Particles.Size(); i++)
@@ -158,27 +158,27 @@ namespace ChaosTest {
 	}
 
 	template <class TParticleContainer>
-	void Reset(TParticleContainer& Particles, const TArray<Chaos::FVec3>& Points)
+	void Reset(TParticleContainer& Particles, const TArray<Chaos::Softs::FSolverVec3>& Points)
 	{
 		for (uint32 i = 0; i < Particles.Size(); i++)
 		{
 			Particles.X(i) = Points[i];
-			Particles.V(i) = Chaos::FVec3(0);
+			Particles.V(i) = Chaos::Softs::FSolverVec3(0);
 		}
 	}
 
-	TArray<FVec3> GetDifference(const TArray<FVec3>& A, const TArray<FVec3>& B)
+	TArray<Softs::FSolverVec3> GetDifference(const TArray<Softs::FSolverVec3>& A, const TArray<Softs::FSolverVec3>& B)
 	{
-		TArray<FVec3> C;
+		TArray<Softs::FSolverVec3> C;
 		C.SetNum(A.Num());
 		for (int32 i = 0; i < A.Num(); i++)
 			C[i] = A[i] - B[i];
 		return C;
 	}
 
-	TArray<FReal> GetMagnitude(const TArray<Chaos::FVec3>& V)
+	TArray<Softs::FSolverReal> GetMagnitude(const TArray<Chaos::Softs::FSolverVec3>& V)
 	{
-		TArray<FReal> M;
+		TArray<Softs::FSolverReal> M;
 		M.SetNum(V.Num());
 		for (int32 i = 0; i < V.Num(); i++)
 			M[i] = V[i].Size();
@@ -204,29 +204,29 @@ namespace ChaosTest {
 
 	template <class TV>
 	bool RunDropTest(
-		TUniquePtr<FPBDEvolution>& Evolution,
-		const FReal GravMag,
+		TUniquePtr<Softs::FPBDEvolution>& Evolution,
+		const Softs::FSolverReal GravMag,
 		const TV& GravDir,
 		const TArray<TV>& InitialPoints,
 		const int32 SubFrameSteps,
-		const FReal DistTolerance,
+		const Softs::FSolverReal DistTolerance,
 		const char* TestID)
 	{
-		const FReal PreTime = Evolution->GetTime();
+		const Softs::FSolverReal PreTime = Evolution->GetTime();
 		AdvanceTime(Evolution, 24, SubFrameSteps); // 1 second
-		const FReal PostTime = Evolution->GetTime();
+		const Softs::FSolverReal PostTime = Evolution->GetTime();
 		EXPECT_NEAR(PostTime-PreTime, 1.0, KINDA_SMALL_NUMBER)
 			<< TestID
 			<< "Evolution advanced time by " << (PostTime - PreTime)
 			<< " seconds, expected 1.0 seconds.";
 
-		const TArray<FVec3> PostPoints = CopyPoints(Evolution->Particles());
-		const TArray<FVec3> Diff = GetDifference(PostPoints, InitialPoints);
-		const TArray<FReal> ScalarDiff = GetMagnitude(Diff);
+		const TArray<Softs::FSolverVec3> PostPoints = CopyPoints(Evolution->Particles());
+		const TArray<Softs::FSolverVec3> Diff = GetDifference(PostPoints, InitialPoints);
+		const TArray<Softs::FSolverReal> ScalarDiff = GetMagnitude(Diff);
 
 		// All points did the same thing
 		int32 Idx = 0;
-		EXPECT_TRUE(AllSame(ScalarDiff, Idx, (FReal)0.1))
+		EXPECT_TRUE(AllSame(ScalarDiff, Idx, (Softs::FSolverReal)0.1))
 			<< TestID
 			<< "Points fell different distances - Index 0: "
 			<< ScalarDiff[0] << " != Index " 
@@ -234,15 +234,15 @@ namespace ChaosTest {
 			<< ScalarDiff[Idx] << " +/- 0.1.";
 
 		// Fell the right amount
-		EXPECT_NEAR(ScalarDiff[0], (FReal)0.5 * GravMag, DistTolerance)
+		EXPECT_NEAR(ScalarDiff[0], (Softs::FSolverReal)0.5 * GravMag, DistTolerance)
 			<< TestID
 			<< "Points fell by " << ScalarDiff[0] 
-			<< ", expected " << ((FReal)0.5 * GravMag)
+			<< ", expected " << ((Softs::FSolverReal)0.5 * GravMag)
 			<< " +/- " << DistTolerance << "."; 
 
 		// Fell the right direction
-		const FReal DirDot = Chaos::FVec3::DotProduct(GravDir, Diff[0].GetSafeNormal());
-		EXPECT_NEAR(DirDot, (FReal)1.0, (FReal)KINDA_SMALL_NUMBER)
+		const Softs::FSolverReal DirDot = Chaos::Softs::FSolverVec3::DotProduct(GravDir, Diff[0].GetSafeNormal());
+		EXPECT_NEAR(DirDot, (Softs::FSolverReal)1.0, (Softs::FSolverReal)KINDA_SMALL_NUMBER)
 			<< TestID
 			<< "Points fell in different directions.";
 
@@ -251,23 +251,23 @@ namespace ChaosTest {
 
 	void DeformableGravity()
 	{
-		const FReal DistTol = 0.0002;
+		const Softs::FSolverReal DistTol = 0.0002;
 
 		//
 		// Initialize solver and gravity
 		//
 
-		TUniquePtr<FPBDEvolution> Evolution = InitPBDEvolution();
+		TUniquePtr<Softs::FPBDEvolution> Evolution = InitPBDEvolution();
 
-		const FVec3 GravDir(0, 0, -1);
-		const FReal GravMag = 980.665;
+		const Softs::FSolverVec3 GravDir(0, 0, -1);
+		const Softs::FSolverReal GravMag = 980.665;
 
 		//
 		// Drop a single particle
 		//
 
 		InitSingleParticle(Evolution);
-		TArray<FVec3> InitialPoints = CopyPoints(Evolution->Particles());
+		TArray<Softs::FSolverVec3> InitialPoints = CopyPoints(Evolution->Particles());
 
 		RunDropTest(Evolution, GravMag, GravDir, InitialPoints, 1, DistTol, "Single point falling under gravity, iters: 1 - ");
 		Reset(Evolution->Particles(), InitialPoints);
@@ -305,7 +305,7 @@ namespace ChaosTest {
 
 	void EdgeConstraints()
 	{
-		TUniquePtr<FPBDEvolution> Evolution = InitPBDEvolution();
+		TUniquePtr<Softs::FPBDEvolution> Evolution = InitPBDEvolution();
 		FTriangleMesh TriMesh;
 		auto& Particles = Evolution->Particles();
 		Particles.AddParticles(2145);
diff --git a/Engine/Source/Runtime/Experimental/Chaos/Public/Chaos/TriangleMesh.h b/Engine/Source/Runtime/Experimental/Chaos/Public/Chaos/TriangleMesh.h
--- a/Engine/Source/Runtime/Experimental/Chaos/Public/Chaos/TriangleMesh.h
+++ b/Engine/Source/Runtime/Experimental/Chaos/Public/Chaos/TriangleMesh.h
@@ -230,14 +230,15 @@ namespace Chaos
 		CHAOS_API void RemoveDuplicateElements();
 		CHAOS_API void RemoveDegenerateElements();
 
-		static FORCEINLINE void InitEquilateralTriangleXY(FTriangleMesh& TriMesh, FParticles& Particles)
+		template <typename T>
+		static FORCEINLINE void InitEquilateralTriangleXY(FTriangleMesh& TriMesh, TParticles<T, 3>& Particles)
 		{
 			const int32 Idx = Particles.Size();
 			Particles.AddParticles(3);
 			// Left handed
-			Particles.X(Idx + 0) = FVec3((FReal)0., (FReal)0.8083, (FReal)0.);
-			Particles.X(Idx + 1) = FVec3((FReal)0.7, (FReal)-0.4041, (FReal)0.);
-			Particles.X(Idx + 2) = FVec3((FReal)-0.7, (FReal)-0.4041, (FReal)0.);
+			Particles.X(Idx + 0) = FVec3((T)0., (T)0.8083, (T)0.);
+			Particles.X(Idx + 1) = FVec3((T)0.7, (T)-0.4041, (T)0.);
+			Particles.X(Idx + 2) = FVec3((T)-0.7, (T)-0.4041, (T)0.);
 
 			TArray<TVec3<int32>> Elements;
 			Elements.SetNum(1);
@@ -245,14 +246,15 @@ namespace Chaos
 
 			TriMesh.Init(MoveTemp(Elements));
 		}
-		static FORCEINLINE void InitEquilateralTriangleYZ(FTriangleMesh& TriMesh, FParticles& Particles)
+		template <typename T>
+		static FORCEINLINE void InitEquilateralTriangleYZ(FTriangleMesh& TriMesh, TParticles<T, 3>& Particles)
 		{
 			const int32 Idx = Particles.Size();
 			Particles.AddParticles(3);
 			// Left handed
-			Particles.X(Idx + 0) = FVec3((FReal)0., (FReal)0., (FReal)0.8083);
-			Particles.X(Idx + 1) = FVec3((FReal)0., (FReal)0.7, (FReal)-0.4041);
-			Particles.X(Idx + 2) = FVec3((FReal)0., (FReal)-0.7, (FReal)-0.4041);
+			Particles.X(Idx + 0) = FVec3((T)0., (T)0., (T)0.8083);
+			Particles.X(Idx + 1) = FVec3((T)0., (T)0.7, (T)-0.4041);
+			Particles.X(Idx + 2) = FVec3((T)0., (T)-0.7, (T)-0.4041);
 
 			TArray<TVec3<int32>> Elements;
 			Elements.SetNum(1);
