Add Triangle Register, a triangle version for SIMD

diff --git a/Engine/Source/Programs/HeadlessChaos/Private/HeadlessChaosTestEPA.cpp b/Engine/Source/Programs/HeadlessChaos/Private/HeadlessChaosTestEPA.cpp
--- a/Engine/Source/Programs/HeadlessChaos/Private/HeadlessChaosTestEPA.cpp
+++ b/Engine/Source/Programs/HeadlessChaos/Private/HeadlessChaosTestEPA.cpp
@@ -12,6 +12,7 @@
 #include "Chaos/ImplicitObjectScaled.h"
 #include "Chaos/TriangleMeshImplicitObject.h"
 #include "Chaos/Triangle.h"
+#include "Chaos/TriangleRegister.h"
 
 namespace ChaosTest
 {
@@ -705,6 +706,10 @@ namespace ChaosTest
 		// Sphere sweep against triangle, fails when it should hit. Raycast added as well for verification purposes.
 		{
 			const FTriangle Triangle({ 0.000000000, 0.000000000, 0.000000000 }, { 128.000000, 0.000000000, -114.064575 }, { 128.000000, 128.000000, 2.35327148 });
+			const FTriangleRegister TriangleReg(
+				MakeVectorRegisterFloat( 0.000000000f, 0.000000000f, 0.000000000f, 0.0f ), 
+				MakeVectorRegisterFloat(128.000000f, 0.000000000f, -114.064575, 0.0f), 
+				MakeVectorRegisterFloat(128.000000, 128.000000, 2.35327148, 0.0f));
 			const TSphere<FReal, 3> Sphere({ 0.0, 0.0, 0.0 }, 4);
 			const TRigidTransform<FReal, 3> Transform({ 174.592773, -161.781250, -68.0469971 }, FQuat::Identity);
 			const FVec3 Dir(-0.406315684, 0.913382649, -0.0252906363);
@@ -716,7 +721,7 @@ namespace ChaosTest
 			FVec3 OutPosition;
 			FVec3 OutNormal;
 
-			bool bSweepResult = GJKRaycast2(Triangle, Sphere, Transform, Dir, Length, OutTime, OutPosition, OutNormal, Thickness, bComputeMTD);
+			bool bSweepResult = GJKRaycast2(TriangleReg, Sphere, Transform, Dir, Length, OutTime, OutPosition, OutNormal, Thickness, bComputeMTD);
 
 
 			// Do a raycast w/ same inputs instead of sweep against triangle to verify sweep should be a hit.
@@ -800,6 +805,10 @@ namespace ChaosTest
 			// Triangle v Box
 			{
 				FTriangle Triangle(FVec3(0.000000000,0.000000000,0.000000000),FVec3(128.000000,0.000000000,35.9375000),FVec3(128.000000,128.000000,134.381042));
+				FTriangleRegister TriangleReg(
+					MakeVectorRegisterFloat(0.000000000f, 0.000000000f, 0.000000000f, 0.0f), 
+					MakeVectorRegisterFloat(128.000000f, 0.000000000f, 35.9375000f, 0.0f), 
+					MakeVectorRegisterFloat(128.000000f, 128.000000f, 134.381042f, 0.0f));
 				TBox<FReal,3> Box(FVec3(-50.0000000,-60.0000000,-30.0000000),FVec3 (50.0000000,60.0000000,30.0000000));
 
 				const TRigidTransform<FReal,3> Transform({127.898438,35.0742188,109.781067},FQuat(0.374886870,-0.0289460570,0.313643545,0.871922970),FVec3(1.0));
@@ -812,7 +821,7 @@ namespace ChaosTest
 				FReal OutTime = -1;
 				FVec3 LocalPosition(-1);
 				FVec3 LocalNormal(-1);
-				GJKRaycast2(Triangle,Box,Transform,Dir,Length,OutTime,LocalPosition,LocalNormal,Thickness,bComputeMTD);
+				GJKRaycast2(TriangleReg,Box,Transform,Dir,Length,OutTime,LocalPosition,LocalNormal,Thickness,bComputeMTD);
 			}
 		}
 
@@ -1263,6 +1272,12 @@ namespace ChaosTest
 				{-91.0660172, 839.028320, 118.413063}
 				});
 
+			FTriangleRegister TriangleReg({
+				MakeVectorRegisterFloat(-306.119476f, 1674.38647f, 117.138489f, 0.0f),
+				MakeVectorRegisterFloat(-491.015747f, 1526.35803f, 116.067123f, 0.0f),
+				MakeVectorRegisterFloat(-91.0660172f, 839.028320f, 118.413063f, 0.0f)
+				});
+
 
 			FVec3 ExpectedNormal = FVec3::CrossProduct(Triangle[1] - Triangle[0], Triangle[2] - Triangle[0]);
 			ExpectedNormal.Normalize();
@@ -1287,7 +1302,7 @@ namespace ChaosTest
 			int32 FaceIndex = -1;
 
 			// This is local to trimesh, world scale.
-			bool bResult = GJKRaycast2<FReal>(Triangle, ScaledCapsule, StartTM, Dir, Length, OutTime, Position, Normal, Thickness, bComputeMTD);
+			bool bResult = GJKRaycast2<FReal>(TriangleReg, ScaledCapsule, StartTM, Dir, Length, OutTime, Position, Normal, Thickness, bComputeMTD);
 
 			// Compare results against GJKPenetration, sweep is initial overlap, so this should be the same.
 			FVec3 Normal2, ClosestA, ClosestB;
diff --git a/Engine/Source/Programs/HeadlessChaos/Private/HeadlessChaosTestGJK.cpp b/Engine/Source/Programs/HeadlessChaos/Private/HeadlessChaosTestGJK.cpp
--- a/Engine/Source/Programs/HeadlessChaos/Private/HeadlessChaosTestGJK.cpp
+++ b/Engine/Source/Programs/HeadlessChaos/Private/HeadlessChaosTestGJK.cpp
@@ -11,6 +11,7 @@
 #include "Chaos/ImplicitObjectScaled.h"
 #include "Chaos/Collision/PBDCollisionConstraint.h"
 #include "Chaos/Triangle.h"
+#include "Chaos/TriangleRegister.h"
 
 namespace ChaosTest
 {
@@ -1505,6 +1506,10 @@ namespace ChaosTest
 			const TImplicitObjectScaled<FConvex> AConvScaled(AConv,FVec3(1.0,1.0,1.0));
 
 			FTriangle A(ConvexParticles[0],ConvexParticles[1],ConvexParticles[2]);
+			FTriangleRegister AReg(
+				MakeVectorRegisterFloat(ConvexParticles[0].X, ConvexParticles[0].Y, ConvexParticles[0].Z, 0.0f),
+				MakeVectorRegisterFloat(ConvexParticles[1].X, ConvexParticles[1].Y, ConvexParticles[1].Z, 0.0f),
+				MakeVectorRegisterFloat(ConvexParticles[2].X, ConvexParticles[2].Y, ConvexParticles[2].Z, 0.0f));
 
 			const FVec3 Pt0(0.0,0.0,-29.6999969);
 			FVec3 Pt1 = Pt0;
@@ -1526,7 +1531,7 @@ namespace ChaosTest
 			FReal Time2;
 			FVec3 Position2,Normal2;
 			EXPECT_TRUE(GJKRaycast2<FReal>(AConvScaled,B,BToATM2,LocalDir,Length+100,Time2,Position2,Normal2,0,true,SearchDir,0));
-			EXPECT_TRUE(GJKRaycast2<FReal>(A,B,BToATM2,LocalDir,Length+100,Time2,Position2,Normal2,0,true,SearchDir,0));
+			EXPECT_TRUE(GJKRaycast2<FReal>(AReg,B,BToATM2,LocalDir,Length+100,Time2,Position2,Normal2,0,true,SearchDir,0));
 
 			EXPECT_NEAR(Time+100,Time2, 1.0f); // TODO: Investigate: This used to be 0
 			EXPECT_VECTOR_NEAR(Normal,Normal2,1e-3); // TODO: Investigate: This used to be 1e-4
diff --git a/Engine/Source/Runtime/Experimental/Chaos/Private/Chaos/HeightField.cpp b/Engine/Source/Runtime/Experimental/Chaos/Private/Chaos/HeightField.cpp
--- a/Engine/Source/Runtime/Experimental/Chaos/Private/Chaos/HeightField.cpp
+++ b/Engine/Source/Runtime/Experimental/Chaos/Private/Chaos/HeightField.cpp
@@ -14,6 +14,7 @@
 #include "HAL/IConsoleManager.h"
 #include "HAL/PlatformMisc.h"
 #include "Chaos/Triangle.h"
+#include "Chaos/TriangleRegister.h"
 
 namespace Chaos
 {
@@ -378,8 +379,11 @@ namespace Chaos
 				}
 
 				//Convert into local space of A to get better precision
+				VectorRegister4Float AReg = MakeVectorRegisterFloatFromDouble(MakeVectorRegister(A.X, B.Y, A.Z, 0.0));
+				VectorRegister4Float BReg = MakeVectorRegisterFloatFromDouble(MakeVectorRegister(B.X, B.Y, B.Z, 0.0));
+				VectorRegister4Float CReg = MakeVectorRegisterFloatFromDouble(MakeVectorRegister(C.X, C.Y, C.Z, 0.0));
 
-				FTriangle Triangle(FVec3(0), B-A, C-A);
+				FTriangleRegister Triangle(VectorZero(), VectorSubtract(BReg, AReg), VectorSubtract(CReg, AReg));
 
 				FReal Time;
 				FVec3 LocalHitPosition;
@@ -1524,6 +1528,11 @@ namespace Chaos
 			const FVec3 Offset = FVec3::CrossProduct(AB, AC);
 			const FVec3 TriNormal = Offset.GetUnsafeNormal();
 			FTriangle TriangleConvex(A, B, C);
+			FTriangleRegister TriangleConvexReg(
+				MakeVectorRegisterFloatFromDouble(MakeVectorRegister(A.X, A.Y, A.Z, 0.0f)),
+				MakeVectorRegisterFloatFromDouble(MakeVectorRegister(B.X, B.Y, B.Z, 0.0f)),
+				MakeVectorRegisterFloatFromDouble(MakeVectorRegister(C.X, C.Y, C.Z, 0.0f)));
+
 
 			FReal Penetration;
 			FVec3 ClosestA, ClosestB, Normal;
@@ -1540,7 +1549,7 @@ namespace Chaos
 				const FReal ApproximateDistToObject = FVec3::DistSquared(QueryTM.GetLocation(), A);
 				const FReal SweepLength = ApproximateSizeOfObject + ApproximateDistToObject;
 				const FRigidTransform3 QueryStartTM(QueryTM.GetLocation() + TriNormal * SweepLength, QueryTM.GetRotation());
-				if (GJKRaycast2(TriangleConvex, QueryGeom, QueryStartTM, -TriNormal, SweepLength, Penetration, ClosestB, Normal, (FReal)0., true))
+				if (GJKRaycast2(TriangleConvexReg, QueryGeom, QueryStartTM, -TriNormal, SweepLength, Penetration, ClosestB, Normal, (FReal)0., true))
 				{
 					LocalContactLocation = ClosestB;
 					LocalContactNormal = TriNormal;
diff --git a/Engine/Source/Runtime/Experimental/Chaos/Private/Chaos/TriangleMeshImplicitObject.cpp b/Engine/Source/Runtime/Experimental/Chaos/Private/Chaos/TriangleMeshImplicitObject.cpp
--- a/Engine/Source/Runtime/Experimental/Chaos/Private/Chaos/TriangleMeshImplicitObject.cpp
+++ b/Engine/Source/Runtime/Experimental/Chaos/Private/Chaos/TriangleMeshImplicitObject.cpp
@@ -6,6 +6,7 @@
 #include "Chaos/Capsule.h"
 #include "Chaos/GJK.h"
 #include "Chaos/Triangle.h"
+#include "Chaos/TriangleRegister.h"
 #include "Chaos/Convex.h"
 #include "Chaos/ImplicitObjectScaled.h"
 #include "Chaos/GeometryQueries.h"
@@ -889,6 +890,8 @@ struct FTriangleMeshSweepVisitor
 	, OutTime(TNumericLimits<FReal>::Max())
 	, TriMeshScale(InTriMeshScale)
 	{
+		VectorScaledDirNormalized = MakeVectorRegisterFloatFromDouble(MakeVectorRegister(InScaledDirNormalized.X, InScaledDirNormalized.Y, InScaledDirNormalized.Z, 0.0));
+		VectorCullsBackFaceSweepsCode = MakeVectorRegisterFloatFromDouble(VectorLoadFloat1(&InCullsBackFaceSweepsCode));
 	}
 
 	const void* GetQueryData() const { return nullptr; }
@@ -920,19 +923,33 @@ struct FTriangleMeshSweepVisitor
 		FVec3 HitPosition;
 		FVec3 HitNormal;
 
-		FVec3 A, B, C;
-		TriangleMeshTransformVertsHelper(TriMeshScale,TriIdx,TriMesh.MParticles,Elements,A,B,C);
-		FTriangle Tri(A, B, C);
+		const VectorRegister4Float TriMeshScaleVector = MakeVectorRegisterFloatFromDouble(MakeVectorRegister(TriMeshScale.X, TriMeshScale.Y, TriMeshScale.Z, 0.0));
+
+		const TParticles<FRealSingle, 3>& Particles = TriMesh.MParticles;
+
+		const TVector<FRealSingle, 3>& AVec = Particles.X(Elements[TriIdx][0]);
+		const TVector<FRealSingle, 3>& BVec = Particles.X(Elements[TriIdx][1]);
+		const TVector<FRealSingle, 3>& CVec = Particles.X(Elements[TriIdx][2]);
+
+		VectorRegister4Float A = MakeVectorRegister(AVec.X, AVec.Y, AVec.Z, 0.0f);
+		VectorRegister4Float B = MakeVectorRegister(BVec.X, BVec.Y, BVec.Z, 0.0f);
+		VectorRegister4Float C = MakeVectorRegister(CVec.X, CVec.Y, CVec.Z, 0.0f);
+
+		A = VectorMultiply(A, TriMeshScaleVector);
+		B = VectorMultiply(B, TriMeshScaleVector);
+		C = VectorMultiply(C, TriMeshScaleVector);
+
+		FTriangleRegister Tri(A, B, C);
 
 		if(CullsBackFaceSweepsCode != 0)
 		{
-			const FVec3 TriNormal = FVec3::CrossProduct(B - A, C - A);
-			if (FVec3::DotProduct(TriNormal, ScaledDirNormalized) * CullsBackFaceSweepsCode > 0)
+			const VectorRegister4Float TriNormal = VectorCross(VectorSubtract(B, A), VectorSubtract(C, A));
+			const VectorRegister4Float ReturnTrue = VectorCompareGT(VectorMultiply(VectorCross(TriNormal, VectorScaledDirNormalized), VectorCullsBackFaceSweepsCode), VectorZero());
+			if (VectorMaskBits(ReturnTrue))
 			{
 				return true;
 			}
 		}
-		
 
 		const auto& WorldScaleQueryGeom = ScaleGeomIntoWorldHelper(QueryGeom, TriMeshScale);
 
@@ -968,9 +985,11 @@ struct FTriangleMeshSweepVisitor
 	const FReal Thickness;
 	const bool bComputeMTD;
 	const FReal CullsBackFaceSweepsCode; // 0: no culling, 1/-1: winding order
+	VectorRegister4Float VectorCullsBackFaceSweepsCode; // 0: no culling, 1/-1: winding order
 
 	// Cache these values for Scaled Triangle Mesh, as they are needed for transformation when sweeping against triangles.
 	FVec3 ScaledDirNormalized;
+	VectorRegister4Float VectorScaledDirNormalized;
 	FReal LengthScale;
 	FRigidTransform3 ScaledStartTM;
 
diff --git a/Engine/Source/Runtime/Experimental/Chaos/Public/Chaos/Triangle.h b/Engine/Source/Runtime/Experimental/Chaos/Public/Chaos/Triangle.h
--- a/Engine/Source/Runtime/Experimental/Chaos/Public/Chaos/Triangle.h
+++ b/Engine/Source/Runtime/Experimental/Chaos/Public/Chaos/Triangle.h
@@ -8,8 +8,6 @@
 #include "Plane.h"
 #include "Chaos/Plane.h"
 
-#include "Math/UnrealMathVectorConstants.h"
-
 namespace Chaos
 {
 	class FTriangle
@@ -24,9 +22,6 @@ namespace Chaos
 			, B(InB)
 			, C(InC)
 		{
-			VectorA = MakeVectorRegisterFloat(static_cast<FRealSingle>(A[0]), static_cast<FRealSingle>(A[1]), static_cast<FRealSingle>(A[2]), static_cast<FRealSingle>(0.0f));
-			VectorB = MakeVectorRegisterFloat(static_cast<FRealSingle>(B[0]), static_cast<FRealSingle>(B[1]), static_cast<FRealSingle>(B[2]), static_cast<FRealSingle>(0.0f));
-			VectorC = MakeVectorRegisterFloat(static_cast<FRealSingle>(C[0]), static_cast<FRealSingle>(C[1]), static_cast<FRealSingle>(C[2]), static_cast<FRealSingle>(0.0f));
 		}
 
 		FORCEINLINE FVec3& operator[](uint32 InIndex)
@@ -220,27 +215,6 @@ namespace Chaos
 			return C;
 		}
 
-		FORCEINLINE VectorRegister4Float SupportCoreSimd(const VectorRegister4Float& Direction, const FReal InMargin) const
-		{
-			// Note: assumes margin == 0
-			VectorRegister4Float DotA = VectorDot4(VectorA, Direction);
-			VectorRegister4Float DotB = VectorDot4(VectorB, Direction);
-			VectorRegister4Float DotC = VectorDot4(VectorC, Direction);
-
-			VectorRegister4Float DotAGEDotB = VectorCompareGE(DotA, DotB);
-			VectorRegister4Float DotAGEDotC = VectorCompareGE(DotA, DotC);
-			VectorRegister4Float DotAG = VectorBitwiseAnd(DotAGEDotB, DotAGEDotC);
-
-			VectorRegister4Float DotBGEDotA = VectorCompareGE(DotB, DotA);
-			VectorRegister4Float DotBGEDotC = VectorCompareGE(DotB, DotC);
-			VectorRegister4Float DotBG = VectorBitwiseAnd(DotBGEDotA, DotBGEDotC);
-
-			VectorRegister4Float Result = VectorSelect(DotAG, VectorA, VectorC);
-			Result = VectorSelect(DotBG, VectorB, Result);
-
-			return Result;
-		}
-
 		FORCEINLINE FVec3 SupportCoreScaled(const FVec3& Direction, FReal InMargin, const FVec3& Scale, FReal* OutSupportDelta, int32& VertexIndex) const
 		{
 			// Note: ignores InMargin, assumed 0 (triangles cannot have a margin as they are zero thickness)
@@ -280,10 +254,6 @@ namespace Chaos
 		FVec3 A;
 		FVec3 B;
 		FVec3 C;
-		VectorRegister4Float VectorA;
-		VectorRegister4Float VectorB;
-		VectorRegister4Float VectorC;
-
 	};
 
 	inline FChaosArchive& operator<<(FChaosArchive& Ar, FTriangle& Value)
diff --git a/Engine/Source/Runtime/Experimental/Chaos/Public/Chaos/TriangleRegister.h b/Engine/Source/Runtime/Experimental/Chaos/Public/Chaos/TriangleRegister.h
new file mode 100644
--- /dev/null
+++ b/Engine/Source/Runtime/Experimental/Chaos/Public/Chaos/TriangleRegister.h
@@ -0,0 +1,55 @@
+// Copyright Epic Games, Inc. All Rights Reserved.
+
+#pragma once
+#include "Math/VectorRegister.h"
+
+namespace Chaos
+{
+	class FTriangleRegister
+	{
+	public:
+		FTriangleRegister()
+		{
+		}
+
+		FTriangleRegister(const VectorRegister4Float& InA, const VectorRegister4Float& InB, const VectorRegister4Float& InC)
+			: A(InA)
+			, B(InB)
+			, C(InC)
+		{
+		}
+
+		FORCEINLINE VectorRegister4Float SupportCoreSimd(const VectorRegister4Float& Direction, const FReal InMargin) const
+		{
+			// Note: assumes margin == 0
+			VectorRegister4Float DotA = VectorDot4(A, Direction);
+			VectorRegister4Float DotB = VectorDot4(B, Direction);
+			VectorRegister4Float DotC = VectorDot4(C, Direction);
+
+			VectorRegister4Float DotAGEDotB = VectorCompareGE(DotA, DotB);
+			VectorRegister4Float DotAGEDotC = VectorCompareGE(DotA, DotC);
+			VectorRegister4Float DotAG = VectorBitwiseAnd(DotAGEDotB, DotAGEDotC);
+
+			VectorRegister4Float DotBGEDotA = VectorCompareGE(DotB, DotA);
+			VectorRegister4Float DotBGEDotC = VectorCompareGE(DotB, DotC);
+			VectorRegister4Float DotBG = VectorBitwiseAnd(DotBGEDotA, DotBGEDotC);
+
+			VectorRegister4Float Result = VectorSelect(DotAG, A, C);
+			Result = VectorSelect(DotBG, B, Result);
+
+			return Result;
+		}
+
+		FORCEINLINE FReal GetMargin() const { return 0; }
+		FORCEINLINE FReal GetRadius() const { return 0; }
+
+
+	private:
+
+		VectorRegister4Float A;
+		VectorRegister4Float B;
+		VectorRegister4Float C;
+
+	};
+
+}
\ No newline at end of file
