Chaos : Fix AABBVectorized not properly accounting for ray boundaries causing false positives

diff --git a/Engine/Source/Programs/HeadlessChaos/Private/HeadlessChaosApp.cpp b/Engine/Source/Programs/HeadlessChaos/Private/HeadlessChaosApp.cpp
--- a/Engine/Source/Programs/HeadlessChaos/Private/HeadlessChaosApp.cpp
+++ b/Engine/Source/Programs/HeadlessChaos/Private/HeadlessChaosApp.cpp
@@ -176,6 +176,7 @@ TEST(RaycastTests, Raycast) {
 	ChaosTest::CapsuleRaycastMissWithEndPointOnBounds();
 	ChaosTest::TriangleRaycast();
 	ChaosTest::BoxRaycast();
+	ChaosTest::VectorizedAABBRaycast();
 	ChaosTest::ScaledRaycast();
 	//ChaosTest::TransformedRaycast();
 	//ChaosTest::UnionRaycast();
diff --git a/Engine/Source/Programs/HeadlessChaos/Private/HeadlessChaosTestRaycast.cpp b/Engine/Source/Programs/HeadlessChaos/Private/HeadlessChaosTestRaycast.cpp
--- a/Engine/Source/Programs/HeadlessChaos/Private/HeadlessChaosTestRaycast.cpp
+++ b/Engine/Source/Programs/HeadlessChaos/Private/HeadlessChaosTestRaycast.cpp
@@ -522,6 +522,26 @@ namespace ChaosTest
 		EXPECT_EQ(Time, 0);
 	}
 
+	void VectorizedAABBRaycast()
+	{
+		VectorRegister4Float Time;
+		VectorRegister4Float Position;
+
+		// AABB with 12cm thickness
+		FAABBVectorized AABB(MakeVectorRegisterFloat(565429.188-12.0, -17180.4355-12.0, -95264.4219-12.0, 0.f), MakeVectorRegisterFloat(568988.312+12.0, -13372.2793+12.0, -93649.4609+12.0, 0.f));
+
+		constexpr float Length = 371.331360;
+		const VectorRegister4Float RayStart = MakeVectorRegisterFloat(565223.812, -13919.9111, -93982.5078, 0.f);
+		const VectorRegister4Float RayDir = MakeVectorRegisterFloat(0.0622759163, -0.997566581, -0.0313483514, 0.f);
+		const VectorRegister4Float RayInvDir = MakeVectorRegisterFloat(16.0575714, -1.00243938, -31.8996048, 0.f);
+		const VectorRegister4Float RayLength = MakeVectorRegisterFloat(Length, Length, Length, Length);
+		const VectorRegister4Float RayInvLength = VectorDivide(MakeVectorRegisterFloatConstant(1.f, 1.f, 1.f, 1.f), RayLength);
+		constexpr bool bParallel[3] = { false, false, false };
+		bool bHit = AABB.RaycastFast(RayStart, RayDir, RayInvDir, bParallel, RayLength, RayInvLength, Time, Position);
+		EXPECT_FALSE(bHit);
+
+	}
+	
 	void ScaledRaycast()
 	{
 		// Note: Spheres cannot be thickened by adding a margin to a wrapper type (such as TImplicitObjectScaled) 
diff --git a/Engine/Source/Programs/HeadlessChaos/Public/HeadlessChaosTestRaycast.h b/Engine/Source/Programs/HeadlessChaos/Public/HeadlessChaosTestRaycast.h
--- a/Engine/Source/Programs/HeadlessChaos/Public/HeadlessChaosTestRaycast.h
+++ b/Engine/Source/Programs/HeadlessChaos/Public/HeadlessChaosTestRaycast.h
@@ -25,6 +25,8 @@ namespace ChaosTest {
 
 	void BoxRaycast();
 
+	void VectorizedAABBRaycast();
+
 	void ScaledRaycast();
 
 	void TransformedRaycast();
diff --git a/Engine/Source/Runtime/Experimental/Chaos/Public/Chaos/AABBVectorized.h b/Engine/Source/Runtime/Experimental/Chaos/Public/Chaos/AABBVectorized.h
--- a/Engine/Source/Runtime/Experimental/Chaos/Public/Chaos/AABBVectorized.h
+++ b/Engine/Source/Runtime/Experimental/Chaos/Public/Chaos/AABBVectorized.h
@@ -58,9 +58,11 @@ public:
 
 		VectorRegister4Float LatestStartTime = VectorMax(Time1, VectorSwizzle(Time1, 1, 2, 0, 3));
 		LatestStartTime = VectorMax(LatestStartTime, VectorSwizzle(Time1, 2, 0, 1, 3));
+		LatestStartTime = VectorMax(LatestStartTime, VectorZero());
 
 		VectorRegister4Float EarliestEndTime = VectorMin(Time2, VectorSwizzle(Time2, 1, 2, 0, 3));
 		EarliestEndTime = VectorMin(EarliestEndTime, VectorSwizzle(Time2, 2, 0, 1, 3));
+		EarliestEndTime = VectorMin(EarliestEndTime, Length);
 
 		//Outside of slab before entering another
 		IsFalse = VectorCompareGT(LatestStartTime, EarliestEndTime);
