Store local inertia tensor/inverse tensor diagonals in vec3 instead of matrix, reducing memory.

diff --git a/Engine/Source/Programs/HeadlessChaos/Private/HeadlessChaosTestCCD.cpp b/Engine/Source/Programs/HeadlessChaos/Private/HeadlessChaosTestCCD.cpp
--- a/Engine/Source/Programs/HeadlessChaos/Private/HeadlessChaosTestCCD.cpp
+++ b/Engine/Source/Programs/HeadlessChaos/Private/HeadlessChaosTestCCD.cpp
@@ -49,9 +49,9 @@ namespace ChaosTest
 
 		Evolution.SetPhysicsMaterial(Dynamic, MakeSerializable(PhysicsMaterial));
 
-		const FReal Mass = 100000.0f;;
-		Dynamic->I() = FMatrix33(Mass, Mass, Mass);
-		Dynamic->InvI() = FMatrix33(1.0f / Mass, 1.0f / Mass, 1.0f / Mass);
+		const FRealSingle Mass = 100000.0f;;
+		Dynamic->I() = TVec3<FRealSingle>(Mass, Mass, Mass);
+		Dynamic->InvI() = TVec3<FRealSingle>(1.0f / Mass, 1.0f / Mass, 1.0f / Mass);
 
 		// Positions and velocities
 		Static->X() = FVec3(0, 0, 0);
@@ -111,8 +111,8 @@ namespace ChaosTest
 		Evolution.SetPhysicsMaterial(Dynamic, MakeSerializable(PhysicsMaterial));
 
 		const FReal Mass = 100000.0f;;
-		Dynamic->I() = FMatrix33(Mass, Mass, Mass);
-		Dynamic->InvI() = FMatrix33(1.0f / Mass, 1.0f / Mass, 1.0f / Mass);
+		Dynamic->I() = TVec3<FRealSingle>(Mass, Mass, Mass);
+		Dynamic->InvI() = TVec3<FRealSingle>(1.0f / Mass, 1.0f / Mass, 1.0f / Mass);
 
 		// Positions and velocities
 		Static->X() = FVec3(0, 0, 0);
@@ -174,8 +174,8 @@ namespace ChaosTest
 		Evolution.SetPhysicsMaterial(Dynamic, MakeSerializable(PhysicsMaterial));
 
 		const FReal Mass = 100000.0f;;
-		Dynamic->I() = FMatrix33(Mass, Mass, Mass);
-		Dynamic->InvI() = FMatrix33(1.0f / Mass, 1.0f / Mass, 1.0f / Mass);
+		Dynamic->I() = TVec3<FRealSingle>(Mass);
+		Dynamic->InvI() = TVec3<FRealSingle>(1.0f / Mass);
 
 		// Positions and velocities
 		Static->X() = FVec3(0, 0, 0);
@@ -256,8 +256,8 @@ namespace ChaosTest
 		Evolution.SetPhysicsMaterial(Dynamic, MakeSerializable(PhysicsMaterial));
 
 		const FReal Mass = 100000.0f;
-		Dynamic->I() = FMatrix33(Mass, Mass, Mass);
-		Dynamic->InvI() = FMatrix33(1.0f / Mass, 1.0f / Mass, 1.0f / Mass);
+		Dynamic->I() = TVec3<FRealSingle>(Mass);
+		Dynamic->InvI() = TVec3<FRealSingle>(1.0f / Mass);
 
 		// Positions and velocities
 		ContainerFaces[0]->X() = FVec3(ContainerBoxHalfSize, 0, 0);
diff --git a/Engine/Source/Programs/HeadlessChaos/Private/HeadlessChaosTestConstraints.cpp b/Engine/Source/Programs/HeadlessChaos/Private/HeadlessChaosTestConstraints.cpp
--- a/Engine/Source/Programs/HeadlessChaos/Private/HeadlessChaosTestConstraints.cpp
+++ b/Engine/Source/Programs/HeadlessChaos/Private/HeadlessChaosTestConstraints.cpp
@@ -161,8 +161,8 @@ namespace ChaosTest {
 			auto* DynamicParticle = Evolution.CreateDynamicParticles(1)[0];
 			DynamicParticle->X() = FVec3(0, 10, 10);
 			//	DynamicParticle->R() = FRotation3::FromAxisAngle(FVec3(0,1,0), PI); // upside down
-			DynamicParticle->I() = FMatrix33(100.0f, 100.0f, 100.0f);
-			DynamicParticle->InvI() = FMatrix33(1.0f / 100.0f, 1.0f / 100.0f, 1.0f / 100.0f);
+			DynamicParticle->I() = TVec3<FRealSingle>(100.0f);
+			DynamicParticle->InvI() = TVec3<FRealSingle>(1.0f / 100.0f);
 
 			FPBDSuspensionConstraints SuspensionConstraints;
 			FVec3 SuspensionLocalLocationA(FVec3(0, 0, 0));
@@ -213,8 +213,8 @@ namespace ChaosTest {
 			DynamicParticle->X() = FVec3(50, 10, 10);
 
 			// minimize rotation using high inertia
-			DynamicParticle->I() = FMatrix33(100000.0f, 100000.0f, 100000.0f);
-			DynamicParticle->InvI() = FMatrix33(1.0f / 100000.0f, 1.0f / 100000.0f, 1.0f / 100000.0f);
+			DynamicParticle->I() = TVec3<FRealSingle>(100000.0f);
+			DynamicParticle->InvI() = TVec3<FRealSingle>(1.0f / 100000.0f);
 
 			FPBDSuspensionConstraints SuspensionConstraints;
 
@@ -276,8 +276,8 @@ namespace ChaosTest {
 		DynamicParticle->X() = FVec3(0, 0, 10);
 		DynamicParticle->M() = Mass;
 		DynamicParticle->InvM() = 1.0f / Mass;
-		DynamicParticle->I() = FMatrix33(100000.0f, 100000.0f, 100000.0f);
-		DynamicParticle->InvI() = FMatrix33(1.0f / 100000.0f, 1.0f / 100000.0f, 1.0f / 100000.0f);
+		DynamicParticle->I() = TVec3<FRealSingle>(100000.0f);
+		DynamicParticle->InvI() = TVec3<FRealSingle>(1.0f / 100000.0f);
 
 		// Suspension setup
 		FPBDSuspensionSettings SuspensionSettings;
diff --git a/Engine/Source/Programs/HeadlessChaos/Private/HeadlessChaosTestMinEvolution.cpp b/Engine/Source/Programs/HeadlessChaos/Private/HeadlessChaosTestMinEvolution.cpp
--- a/Engine/Source/Programs/HeadlessChaos/Private/HeadlessChaosTestMinEvolution.cpp
+++ b/Engine/Source/Programs/HeadlessChaos/Private/HeadlessChaosTestMinEvolution.cpp
@@ -71,17 +71,17 @@ namespace ChaosTest
 		// @todo(ccaulfield) this needs to be easier
 		Particles[0]->X() = FVec3(-50, 0, 0);
 		Particles[0]->M() = 1.0f;
-		Particles[0]->I() = FMatrix33(100.0f, 100.0f, 100.0f);
+		Particles[0]->I() = TVec3<FRealSingle>(100.0f, 100.0f, 100.0f);
 		Particles[0]->InvM() = 1.0f;
-		Particles[0]->InvI() = FMatrix33(1.0f / 100.0f, 1.0f / 100.0f, 1.0f / 100.0f);
+		Particles[0]->InvI() = TVec3<FRealSingle>(1.0f / 100.0f, 1.0f / 100.0f, 1.0f / 100.0f);
 		Particles[0]->AuxilaryValue(ParticlePrevXs) = Particles[0]->X();
 		Particles[0]->AuxilaryValue(ParticlePrevRs) = Particles[0]->R();
 
 		Particles[1]->X() = FVec3(50, 0, 0);
 		Particles[1]->M() = 1.0f;
-		Particles[1]->I() = FMatrix33(100.0f, 100.0f, 100.0f);
+		Particles[1]->I() = TVec3<FRealSingle>(100.0f, 100.0f, 100.0f);
 		Particles[1]->InvM() = 1.0f;
-		Particles[1]->InvI() = FMatrix33(1.0f / 100.0f, 1.0f / 100.0f, 1.0f / 100.0f);
+		Particles[1]->InvI() = TVec3<FRealSingle>(1.0f / 100.0f, 1.0f / 100.0f, 1.0f / 100.0f);
 		Particles[1]->AuxilaryValue(ParticlePrevXs) = Particles[1]->X();
 		Particles[1]->AuxilaryValue(ParticlePrevRs) = Particles[1]->R();
 
diff --git a/Engine/Source/Programs/HeadlessChaos/Private/HeadlessChaosTestPerf.cpp b/Engine/Source/Programs/HeadlessChaos/Private/HeadlessChaosTestPerf.cpp
--- a/Engine/Source/Programs/HeadlessChaos/Private/HeadlessChaosTestPerf.cpp
+++ b/Engine/Source/Programs/HeadlessChaos/Private/HeadlessChaosTestPerf.cpp
@@ -89,16 +89,16 @@ namespace ChaosTest {
 				auto BoxDynamic = DynamicParticles[DynamicParticleIndex++];
 				BoxDynamic->SetGeometry(MakeSerializable(Box));
 				BoxDynamic->X() = Offset + FVec3(x * Interval, y * Interval, Height);
-				BoxDynamic->I() = FMatrix33(100000., 100000., 100000.);
-				BoxDynamic->InvI() = FMatrix33(1. / 100000., 1. / 100000., 1. / 100000.);
+				BoxDynamic->I() = TVec3<FRealSingle>(100000.);
+				BoxDynamic->InvI() = TVec3<FRealSingle>(1. / 100000.);
 				Evolution.SetPhysicsMaterial(BoxDynamic, Material);
 				ParticleHandles.Add(BoxDynamic);
 
 				auto SphereDynamic = DynamicParticles[DynamicParticleIndex++];
 				SphereDynamic->SetGeometry(MakeSerializable(Sphere));
 				SphereDynamic->X() = Offset + FVec3(x * Interval, y * Interval, Height * (FReal)0.5);
-				SphereDynamic->I() = FMatrix33(100000., 100000., 100000.);
-				SphereDynamic->InvI() = FMatrix33(1. / 100000., 1. / 100000., 1. / 100000.);
+				SphereDynamic->I() = TVec3<FRealSingle>(100000.);
+				SphereDynamic->InvI() = TVec3<FRealSingle>(1. / 100000.);
 				Evolution.SetPhysicsMaterial(SphereDynamic, Material);
 				ParticleHandles.Add(SphereDynamic);
 #if USE_CONSTRAINTS
diff --git a/Engine/Source/Programs/HeadlessChaos/Private/HeadlessChaosTestSim.cpp b/Engine/Source/Programs/HeadlessChaos/Private/HeadlessChaosTestSim.cpp
--- a/Engine/Source/Programs/HeadlessChaos/Private/HeadlessChaosTestSim.cpp
+++ b/Engine/Source/Programs/HeadlessChaos/Private/HeadlessChaosTestSim.cpp
@@ -81,8 +81,8 @@ namespace ChaosTest {
 
 		Static->X() = FVec3(10, 10, 10);
 		Dynamic->X() = FVec3(10, 10, 150);
-		Dynamic->I() = FMatrix33(100000.0f, 100000.0f, 100000.0f);
-		Dynamic->InvI() = FMatrix33(1.0f / 100000.0f, 1.0f / 100000.0f, 1.0f / 100000.0f);
+		Dynamic->I() = TVec3<FRealSingle>(100000.0f);
+		Dynamic->InvI() = TVec3<FRealSingle>(1.0f / 100000.0f);
 
 		// The position of the static has changed and statics don't automatically update bounds, so update explicitly
 		Static->UpdateWorldSpaceState(FRigidTransform3(Static->X(), Static->R()), FVec3(0));
@@ -122,8 +122,8 @@ namespace ChaosTest {
 		Static->X() = FVec3(10, 10, 10);
 		Static->UpdateWorldSpaceState(FRigidTransform3(Static->X(), Static->R()), FVec3(0));
 		Dynamic->X() = FVec3(10, 10, 120);
-		Dynamic->I() = FMatrix33(100000.0f, 100000.0f, 100000.0f);
-		Dynamic->InvI() = FMatrix33(1.0f / 100000.0f, 1.0f / 100000.0f, 1.0f / 100000.0f);
+		Dynamic->I() = TVec3<FRealSingle>(100000.0f);
+		Dynamic->InvI() = TVec3<FRealSingle>(1.0f / 100000.0f);
 
 		::ChaosTest::SetParticleSimDataToCollide({ Static,Dynamic });
 
@@ -168,8 +168,8 @@ namespace ChaosTest {
 		Static->X() = FVec3(10, 10, 10);
 		Static->UpdateWorldSpaceState(FRigidTransform3(Static->X(), Static->R()), FVec3(0));
 		Dynamic->X() = FVec3(10, 10, 300);
-		Dynamic->I() = FMatrix33(1, 1, 1);
-		Dynamic->InvI() = FMatrix33(1, 1, 1);
+		Dynamic->I() = TVec3<FRealSingle>(1);
+		Dynamic->InvI() = TVec3<FRealSingle>(1);
 
 		::ChaosTest::SetParticleSimDataToCollide({ Static,Dynamic });
 
@@ -217,10 +217,10 @@ namespace ChaosTest {
 		Evolution.SetPhysicsMaterial(Dynamic1, MakeSerializable(PhysicsMaterial));
 		Evolution.SetPhysicsMaterial(Dynamic2, MakeSerializable(PhysicsMaterial));
 
-		Dynamic1->I() = FMatrix33(100000.0f, 100000.0f, 100000.0f);
-		Dynamic1->InvI() = FMatrix33(1.0f / 100000.0f, 1.0f / 100000.0f, 1.0f / 100000.0f);
-		Dynamic2->I() = FMatrix33(100000.0f, 100000.0f, 100000.0f);
-		Dynamic2->InvI() = FMatrix33(1.0f / 100000.0f, 1.0f / 100000.0f, 1.0f / 100000.0f);
+		Dynamic1->I() = TVec3<FRealSingle>(100000.0f);
+		Dynamic1->InvI() = TVec3<FRealSingle>(1.0f / 100000.0f);
+		Dynamic2->I() = TVec3<FRealSingle>(100000.0f);
+		Dynamic2->InvI() = TVec3<FRealSingle>(1.0f / 100000.0f);
 
 		::ChaosTest::SetParticleSimDataToCollide({ Static,Dynamic1,Dynamic2 });
 
diff --git a/Engine/Source/Programs/HeadlessChaos/Private/HeadlessChaosTestUtility.cpp b/Engine/Source/Programs/HeadlessChaos/Private/HeadlessChaosTestUtility.cpp
--- a/Engine/Source/Programs/HeadlessChaos/Private/HeadlessChaosTestUtility.cpp
+++ b/Engine/Source/Programs/HeadlessChaos/Private/HeadlessChaosTestUtility.cpp
@@ -34,8 +34,8 @@ namespace ChaosTest {
 
 		InParticles.M(RigidBodyIndex) = 1.0;
 		InParticles.InvM(RigidBodyIndex) = 1.0;
-		InParticles.I(RigidBodyIndex) = FMatrix33(1, 0, 0, 0, 1, 0, 0, 0, 1);
-		InParticles.InvI(RigidBodyIndex) = FMatrix33(1, 0, 0, 0, 1, 0, 0, 0, 1);
+		InParticles.I(RigidBodyIndex) = TVec3<FRealSingle>(1, 1, 1);
+		InParticles.InvI(RigidBodyIndex) = TVec3<FRealSingle>(1, 1, 1);
 		InParticles.SetDynamicGeometry(RigidBodyIndex, MakeUnique<TSphere<FReal, 3>>(FVec3(0), Scale));
 		InParticles.SetObjectState(RigidBodyIndex, EObjectStateType::Dynamic);
 
@@ -56,8 +56,8 @@ namespace ChaosTest {
 
 		InParticles.M(RigidBodyIndex) = 1.0;
 		InParticles.InvM(RigidBodyIndex) = 1.0;
-		InParticles.I(RigidBodyIndex) = FMatrix33(1, 0, 0, 0, 1, 0, 0, 0, 1);
-		InParticles.InvI(RigidBodyIndex) = FMatrix33(1, 0, 0, 0, 1, 0, 0, 0, 1);
+		InParticles.I(RigidBodyIndex) = TVec3<FRealSingle>(1, 1, 1);
+		InParticles.InvI(RigidBodyIndex) = TVec3<FRealSingle>(1, 1, 1);
 		InParticles.SetDynamicGeometry(RigidBodyIndex, MakeUnique<TBox<FReal, 3>>(-Scale / 2.0, Scale / 2.0));
 		InParticles.SetObjectState(RigidBodyIndex, EObjectStateType::Dynamic);
 
@@ -80,8 +80,8 @@ namespace ChaosTest {
 
 			DynamicParticle->M() = 1.0;
 			DynamicParticle->InvM() = 1.0;
-			DynamicParticle->I() = FMatrix33(1, 0, 0, 0, 1, 0, 0, 0, 1);
-			DynamicParticle->InvI() = FMatrix33(1, 0, 0, 0, 1, 0, 0, 0, 1);
+			DynamicParticle->I() = TVec3<FRealSingle>(1, 1, 1);
+			DynamicParticle->InvI() = TVec3<FRealSingle>(1, 1, 1);
 		}
 	}
 
@@ -101,8 +101,8 @@ namespace ChaosTest {
 		FReal ScaleSq = Scale.X * Scale.X;
 		InParticles.M(RigidBodyIndex) = 1.0;
 		InParticles.InvM(RigidBodyIndex) = 1.0;
-		InParticles.I(RigidBodyIndex) = FMatrix33(ScaleSq / 6.0, 0, 0, 0, ScaleSq / 6.0, 0, 0, 0, ScaleSq / 6.0);
-		InParticles.InvI(RigidBodyIndex) = FMatrix33(6.0 / ScaleSq, 0, 0, 0, 6.0 / ScaleSq, 0, 0, 0, 6.0 / ScaleSq);
+		InParticles.I(RigidBodyIndex) = TVec3<FRealSingle>(ScaleSq / 6.0, ScaleSq / 6.0, ScaleSq / 6.0);
+		InParticles.InvI(RigidBodyIndex) = TVec3<FRealSingle>(6.0 / ScaleSq, 6.0 / ScaleSq, 6.0 / ScaleSq);
 		InParticles.SetDynamicGeometry(RigidBodyIndex, MakeUnique<TBox<FReal, 3>>(-Scale / 2.0, Scale / 2.0));
 		InParticles.SetObjectState(RigidBodyIndex, EObjectStateType::Dynamic);
 
@@ -169,8 +169,8 @@ namespace ChaosTest {
 		FReal ScaleSq = MaxScale * MaxScale;
 		Particle->M() = 1.0;
 		Particle->InvM() = 1.0;
-		Particle->I() = FMatrix33(ScaleSq / 6.0, 0, 0, 0, ScaleSq / 6.0, 0, 0, 0, ScaleSq / 6.0);
-		Particle->InvI() = FMatrix33(6.0 / ScaleSq, 0, 0, 0, 6.0 / ScaleSq, 0, 0, 0, 6.0 / ScaleSq);
+		Particle->I() = TVec3<FRealSingle>(ScaleSq / 6.0);
+		Particle->InvI() = TVec3<FRealSingle>(6.0 / ScaleSq);
 
 		Particle->SetDynamicGeometry(MakeUnique<TBox<FReal, 3>>(-Scale / 2.0, Scale / 2.0, Margin));
 
@@ -241,8 +241,8 @@ namespace ChaosTest {
 		FReal ScaleSq = Scale.X * Scale.X;
 		Particle->M() = 1.0;
 		Particle->InvM() = 1.0;
-		Particle->I() = FMatrix33(ScaleSq / 6.0, 0, 0, 0, ScaleSq / 6.0, 0, 0, 0, ScaleSq / 6.0);
-		Particle->InvI() = FMatrix33(6.0 / ScaleSq, 0, 0, 0, 6.0 / ScaleSq, 0, 0, 0, 6.0 / ScaleSq);
+		Particle->I() = TVec3<FRealSingle>(ScaleSq / 6.0);
+		Particle->InvI() = TVec3<FRealSingle>(6.0 / ScaleSq);
 
 		Particle->SetDynamicGeometry(MakeUnique<TSphere<FReal, 3>>(FVec3(0), Scale.X / 2.0));
 
@@ -304,8 +304,8 @@ namespace ChaosTest {
 		FReal ScaleSq = Scale.X * Scale.X;
 		Particle->M() = 1.0;
 		Particle->InvM() = 1.0;
-		Particle->I() = FMatrix33(ScaleSq / 6.0, 0, 0, 0, ScaleSq / 6.0, 0, 0, 0, ScaleSq / 6.0);
-		Particle->InvI() = FMatrix33(6.0 / ScaleSq, 0, 0, 0, 6.0 / ScaleSq, 0, 0, 0, 6.0 / ScaleSq);
+		Particle->I() = TVec3<FRealSingle>(ScaleSq / 6.0);
+		Particle->InvI() = TVec3<FRealSingle>(6.0 / ScaleSq);
 		
 		if (Tapered)
 		{
@@ -474,8 +474,8 @@ namespace ChaosTest {
 		InParticles.W(RigidBodyIndex) = FVec3(0, 0, 0);
 		InParticles.M(RigidBodyIndex) = 1.0;
 		InParticles.InvM(RigidBodyIndex) = 0.0;
-		InParticles.I(RigidBodyIndex) = FMatrix33(1, 0, 0, 0, 1, 0, 0, 0, 1);
-		InParticles.InvI(RigidBodyIndex) = FMatrix33(0, 0, 0, 0, 0, 0, 0, 0, 0);
+		InParticles.I(RigidBodyIndex) = TVec3<FRealSingle>(1);
+		InParticles.InvI(RigidBodyIndex) = TVec3<FRealSingle>(0);
 		InParticles.SetDynamicGeometry(RigidBodyIndex, MakeUnique<TPlane<FReal, 3>>(FVec3(0, 0, 0), FVec3(0, 0, 1)));
 		InParticles.SetObjectState(RigidBodyIndex, EObjectStateType::Kinematic);
 
@@ -587,8 +587,8 @@ namespace ChaosTest {
 
 		InParticles.M() = 1.0;
 		InParticles.InvM() = 1.0;
-		InParticles.I() = FMatrix33(1, 0, 0, 0, 1, 0, 0, 0, 1);
-		InParticles.InvI() = FMatrix33(1, 0, 0, 0, 1, 0, 0, 0, 1);
+		InParticles.I() = TVec3<FRealSingle>(1);
+		InParticles.InvI() = TVec3<FRealSingle>(1);
 		InParticles.SetDynamicGeometry(MakeUnique<FConvex>(Cube, Margin));
 		InParticles.SetObjectStateLowLevel(EObjectStateType::Dynamic);
 
diff --git a/Engine/Source/Programs/HeadlessChaos/Private/HeadlessChaosTestVehicles.cpp b/Engine/Source/Programs/HeadlessChaos/Private/HeadlessChaosTestVehicles.cpp
--- a/Engine/Source/Programs/HeadlessChaos/Private/HeadlessChaosTestVehicles.cpp
+++ b/Engine/Source/Programs/HeadlessChaos/Private/HeadlessChaosTestVehicles.cpp
@@ -1522,8 +1522,8 @@ namespace ChaosTest
 		Dynamic->X() = FVec3(10, 10, 20);
 		Dynamic->M() = BodyMass;
 		Dynamic->InvM() = 1.0f / BodyMass;
-		Dynamic->I() = FMatrix33(100000.0f, 100000.0f, 100000.0f);
-		Dynamic->InvI() = FMatrix33(1.0f / 100000.0f, 1.0f / 100000.0f, 1.0f / 100000.0f);
+		Dynamic->I() = TVec3<FRealSingle>(100000.0f);
+		Dynamic->InvI() = TVec3<FRealSingle>(1.0f / 100000.0f);
 
 		float AccumulatedTime = 0.f;
 		const FReal Dt = 1 / 30.f;
diff --git a/Engine/Source/Programs/HeadlessChaos/Public/HeadlessChaosTestUtility.h b/Engine/Source/Programs/HeadlessChaos/Public/HeadlessChaosTestUtility.h
--- a/Engine/Source/Programs/HeadlessChaos/Public/HeadlessChaosTestUtility.h
+++ b/Engine/Source/Programs/HeadlessChaos/Public/HeadlessChaosTestUtility.h
@@ -166,8 +166,8 @@ namespace ChaosTest {
 	{
 		float Element = Mass * Dimension * Dimension / 6.0f;
 		float InvElement = 1.f / Element;
-		Particle.SetI(FMatrix33(Element, Element, Element));
-		Particle.SetInvI(FMatrix33(InvElement, InvElement, InvElement));
+		Particle.SetI(TVec3<FRealSingle>(Element));
+		Particle.SetInvI(TVec3<FRealSingle>(InvElement));
 	}
 
 
diff --git a/Engine/Source/Runtime/Core/Public/UObject/UE5ReleaseStreamObjectVersion.h b/Engine/Source/Runtime/Core/Public/UObject/UE5ReleaseStreamObjectVersion.h
--- a/Engine/Source/Runtime/Core/Public/UObject/UE5ReleaseStreamObjectVersion.h
+++ b/Engine/Source/Runtime/Core/Public/UObject/UE5ReleaseStreamObjectVersion.h
@@ -114,6 +114,9 @@ struct CORE_API FUE5ReleaseStreamObjectVersion
 		
 		// add inertia tensor and rotation of mass to convex 
 		AddedInertiaTensorAndRotationOfMassAddedToConvex,
+		
+		// Storing inertia tensor as vec3 instead of matrix.
+		ChaosInertiaConvertedToVec3,
 
 		// -----<new versions can be added above this line>-------------------------------------------------
 		VersionPlusOne,
diff --git a/Engine/Source/Runtime/Engine/Private/PhysicsEngine/ImmediatePhysics/ImmediatePhysicsChaos/ImmediatePhysicsActorHandle_Chaos.cpp b/Engine/Source/Runtime/Engine/Private/PhysicsEngine/ImmediatePhysics/ImmediatePhysicsChaos/ImmediatePhysicsActorHandle_Chaos.cpp
--- a/Engine/Source/Runtime/Engine/Private/PhysicsEngine/ImmediatePhysics/ImmediatePhysicsChaos/ImmediatePhysicsActorHandle_Chaos.cpp
+++ b/Engine/Source/Runtime/Engine/Private/PhysicsEngine/ImmediatePhysics/ImmediatePhysicsChaos/ImmediatePhysicsActorHandle_Chaos.cpp
@@ -295,8 +295,8 @@ namespace ImmediatePhysics_Chaos
 					Dynamic->SetInvM(MassInv);
 					Dynamic->SetCenterOfMass(CoMTransform.GetTranslation());
 					Dynamic->SetRotationOfMass(CoMTransform.GetRotation());
-					Dynamic->SetI({ Inertia.X, Inertia.Y, Inertia.Z });
-					Dynamic->SetInvI({ InertiaInv.X, InertiaInv.Y, InertiaInv.Z });
+					Dynamic->SetI(TVec3<FRealSingle>( Inertia.X, Inertia.Y, Inertia.Z ));
+					Dynamic->SetInvI(TVec3<FRealSingle>(InertiaInv.X, InertiaInv.Y, InertiaInv.Z ));
 					if (BodyInstance != nullptr)
 					{
 						Dynamic->SetLinearEtherDrag(BodyInstance->LinearDamping);
@@ -635,25 +635,19 @@ namespace ImmediatePhysics_Chaos
 			{
 				NewInertia = FVector3f( 1.0f / NewInverseInertia.X , 1.0f / NewInverseInertia.Y, 1.0f / NewInverseInertia.Z );
 			}
-			Dynamic->SetI({ NewInertia.X, NewInertia.Y, NewInertia.Z });
-			Dynamic->SetInvI({ (float)NewInverseInertia.X, (float)NewInverseInertia.Y, (float)NewInverseInertia.Z });
+			Dynamic->SetI(TVec3<FRealSingle>(NewInertia.X, NewInertia.Y, NewInertia.Z ));
+			Dynamic->SetInvI(TVec3<FRealSingle>(NewInverseInertia.X, NewInverseInertia.Y, NewInverseInertia.Z ));
 		}
 	}
 
 	FVector FActorHandle::GetInverseInertia() const
 	{
-		using namespace Chaos;
-
-		const FMatrix33& InvI = Handle()->InvI();
-		return { InvI.M[0][0], InvI.M[1][1], InvI.M[2][2] };
+		return Handle()->InvI();
 	}
 
 	FVector FActorHandle::GetInertia() const
 	{
-		using namespace Chaos;
-
-		const FMatrix33& I = Handle()->I();
-		return { I.M[0][0], I.M[1][1], I.M[2][2] };
+		return Handle()->I();
 	}
 
 	void FActorHandle::SetMaxDepenetrationVelocity(FReal NewMaxDepenetrationVelocity)
diff --git a/Engine/Source/Runtime/Experimental/Chaos/Private/Chaos/ChaosDebugDraw.cpp b/Engine/Source/Runtime/Experimental/Chaos/Private/Chaos/ChaosDebugDraw.cpp
--- a/Engine/Source/Runtime/Experimental/Chaos/Private/Chaos/ChaosDebugDraw.cpp
+++ b/Engine/Source/Runtime/Experimental/Chaos/Private/Chaos/ChaosDebugDraw.cpp
@@ -818,7 +818,7 @@ namespace Chaos
 			{
 				if (const TPBDRigidParticleHandle<FReal, 3>* Rigid = InParticle->CastToRigidParticle())
 				{
-					const FVec3 EquivalentBoxSize = Utilities::BoxSizeFromInertia(Rigid->I().GetDiagonal(), Rigid->M());
+					const FVec3 EquivalentBoxSize = Utilities::BoxSizeFromInertia(Rigid->I(), Rigid->M());
 					FDebugDrawQueue::GetInstance().DrawDebugBox(PCOM, 0.5f * Settings.InertiaScale * EquivalentBoxSize, QCOM, FColor::Magenta, false, 0.0f, 0, Settings.LineThickness);
 				}
 			}
diff --git a/Engine/Source/Runtime/Experimental/Chaos/Private/Chaos/Evolution/PBDMinEvolution.cpp b/Engine/Source/Runtime/Experimental/Chaos/Private/Chaos/Evolution/PBDMinEvolution.cpp
--- a/Engine/Source/Runtime/Experimental/Chaos/Private/Chaos/Evolution/PBDMinEvolution.cpp
+++ b/Engine/Source/Runtime/Experimental/Chaos/Private/Chaos/Evolution/PBDMinEvolution.cpp
@@ -112,7 +112,7 @@ namespace Chaos
 		FVec3* CenterOfMass;
 		FRotation3* RotationOfMass;
 		FReal* InvM;
-		FMatrix33* InvI;
+		TVec3<FRealSingle>* InvI;
 		FVec3* Acceleration;
 		FVec3* AngularAcceleration;
 		FVec3* LinearImpulseVelocity;
diff --git a/Engine/Source/Runtime/Experimental/Chaos/Private/Chaos/PBDRigidClustering.cpp b/Engine/Source/Runtime/Experimental/Chaos/Private/Chaos/PBDRigidClustering.cpp
--- a/Engine/Source/Runtime/Experimental/Chaos/Private/Chaos/PBDRigidClustering.cpp
+++ b/Engine/Source/Runtime/Experimental/Chaos/Private/Chaos/PBDRigidClustering.cpp
@@ -137,7 +137,9 @@ namespace Chaos
 
 		ensureMsgf(!ProxyGeometry || ForceMassOrientation, TEXT("If ProxyGeometry is passed, we must override the mass orientation as they are tied"));
 
-		UpdateClusterMassProperties(NewParticle, ChildrenSet, ForceMassOrientation);
+		// TODO: This needs to be rotated to diagonal, used to update I()/InvI() from diagonal, and update transform with rotation.
+		FMatrix33 ClusterInertia(0);
+		UpdateClusterMassProperties(NewParticle, ChildrenSet, ClusterInertia, ForceMassOrientation);
 		UpdateKinematicProperties(NewParticle, MChildren, MEvolution);
 		UpdateGeometry(NewParticle, ChildrenSet, ProxyGeometry, Parameters);
 		GenerateConnectionGraph(NewParticle, Parameters);
@@ -240,7 +242,10 @@ namespace Chaos
 		NoCleanParams.bCopyCollisionParticles = !!UnionsHaveCollisionParticles;
 
 		TSet<FPBDRigidParticleHandle*> ChildrenSet(ChildrenArray);
-		UpdateClusterMassProperties(NewParticle, ChildrenSet, nullptr);
+		
+		// TODO: This needs to be rotated to diagonal, used to update I()/InvI() from diagonal, and update transform with rotation.
+		FMatrix33 ClusterInertia(0);
+		UpdateClusterMassProperties(NewParticle, ChildrenSet, ClusterInertia, nullptr);
 		UpdateKinematicProperties(NewParticle, MChildren, MEvolution);
 
 		UpdateGeometry(NewParticle, ChildrenSet, nullptr, NoCleanParams);
diff --git a/Engine/Source/Runtime/Experimental/Chaos/Private/Chaos/PBDRigidClusteringAlgo.cpp b/Engine/Source/Runtime/Experimental/Chaos/Private/Chaos/PBDRigidClusteringAlgo.cpp
--- a/Engine/Source/Runtime/Experimental/Chaos/Private/Chaos/PBDRigidClusteringAlgo.cpp
+++ b/Engine/Source/Runtime/Experimental/Chaos/Private/Chaos/PBDRigidClusteringAlgo.cpp
@@ -38,6 +38,7 @@ namespace Chaos
 	void UpdateClusterMassProperties(
 		FPBDRigidClusteredParticleHandle* Parent,
 		TSet<FPBDRigidParticleHandle*>& Children,
+		FMatrix33& ParentInertia,
 		const FRigidTransform3* ForceMassOrientation)
 	{
 		SCOPE_CYCLE_COUNTER(STAT_UpdateClusterMassProperties);
@@ -48,7 +49,8 @@ namespace Chaos
 		Parent->SetV(FVec3(0));
 		Parent->SetW(FVec3(0));
 		Parent->SetM(0);
-		Parent->SetI(FMatrix33(0));
+		Parent->SetI(TVec3<FRealSingle>(0)); // Diagonal should be updated from ParentInertia by caller
+		ParentInertia = FMatrix33(0);
 
 		bool bHasChild = false;
 		bool bHasProxyChild = false;
@@ -63,8 +65,9 @@ namespace Chaos
 			ChildRotation = Child->R();
 
 			const FReal ChildMass = Child->M();
+			const FMatrix33 IMatrix(Child->I());
 			const FMatrix33 ChildWorldSpaceI =
-				(ChildRotation * FMatrix::Identity) * Child->I() * (ChildRotation * FMatrix::Identity).GetTransposed();
+				(ChildRotation * FMatrix::Identity) * IMatrix * (ChildRotation * FMatrix::Identity).GetTransposed();
 			if (ChildWorldSpaceI.ContainsNaN())
 			{
 				continue;
@@ -72,7 +75,8 @@ namespace Chaos
 			bHasProxyChild = true;
 			bHasChild = true;
 			bHasProxyChild = true;
-			Parent->I() += ChildWorldSpaceI;
+
+			ParentInertia += ChildWorldSpaceI;
 			Parent->M() += ChildMass;
 			Parent->X() += ChildPosition * ChildMass;
 			Parent->V() += OriginalChild->V() * ChildMass; // Use orig child for vel because we don't sim the proxy
@@ -87,14 +91,15 @@ namespace Chaos
 				const FRotation3& ChildRotation = Child->R();
 				const FReal ChildMass = Child->M();
 
+				const FMatrix33 IMatrix(Child->I());
 				const FMatrix33 ChildWorldSpaceI =
-					(ChildRotation * FMatrix::Identity) * Child->I() * (ChildRotation * FMatrix::Identity).GetTransposed();
+					(ChildRotation * FMatrix::Identity) * IMatrix * (ChildRotation * FMatrix::Identity).GetTransposed();
 				if (ChildWorldSpaceI.ContainsNaN())
 				{
 					continue;
 				}
 				bHasChild = true;
-				Parent->I() += ChildWorldSpaceI;
+				ParentInertia += ChildWorldSpaceI;
 				Parent->M() += ChildMass;
 				Parent->X() += ChildPosition * ChildMass;
 				Parent->V() += OriginalChild->V() * ChildMass; // Use orig child for vel because we don't sim the proxy
@@ -103,10 +108,10 @@ namespace Chaos
 		}
 		for (int32 i = 0; i < 3; i++)
 		{
-			const FMatrix33& InertiaTensor = Parent->I();
-			if (InertiaTensor.GetColumn(i)[i] < SMALL_NUMBER)
+			const FVec3& InertiaTensor = Parent->I();
+			if (InertiaTensor[i] < SMALL_NUMBER)
 			{
-				Parent->SetI(FMatrix33(1.f, 1.f, 1.f));
+				Parent->SetI(TVec3<FRealSingle>(1.f, 1.f, 1.f));
 				break;
 			}
 		}
@@ -123,8 +128,9 @@ namespace Chaos
 			Parent->PreW() = Parent->W();
 			Parent->R() = FRotation3(FMatrix::Identity);
 			Parent->Q() = Parent->R();
-			Parent->I() = FMatrix::Identity;
-			Parent->InvI() = FMatrix::Identity;
+			Parent->I() = TVec3<FRealSingle>(1);
+			Parent->InvI() = TVec3<FRealSingle>(1);
+			ParentInertia = FMatrix33(1,1,1);
 			return;
 		}
 
@@ -159,24 +165,23 @@ namespace Chaos
 				const FReal& p1 = ParentToChild[1];
 				const FReal& p2 = ParentToChild[2];
 				const FReal& m = ChildMass;
-				Parent->I() +=
-					FMatrix33(
+				// TODO: We are computing inertia twice. Need to decide if we prefer this to loop above.
+				ParentInertia += FMatrix33(
 						m * (p1 * p1 + p2 * p2), -m * p1 * p0, -m * p2 * p0,
 						m * (p2 * p2 + p0 * p0), -m * p2 * p1, m * (p1 * p1 + p0 * p0));
 			}
 		}
-		FMatrix33& InertiaTensor = Parent->I();
-		if (Parent->I().ContainsNaN())
+		if (ParentInertia.ContainsNaN())
 		{
-			InertiaTensor = FMatrix33((FReal)1., (FReal)1., (FReal)1.);
+			ParentInertia = FMatrix33(1,1,1);
 		}
 		else
 		{
 			for (int32 i = 0; i < 3; i++)
 			{
-				if (InertiaTensor.GetColumn(i)[i] < SMALL_NUMBER)
+				if (ParentInertia.GetColumn(i)[i] < SMALL_NUMBER)
 				{
-					InertiaTensor = FMatrix33((FReal)1., (FReal)1., (FReal)1.);
+					ParentInertia = FMatrix33(1,1,1);
 					break;
 				}
 			}
@@ -190,12 +195,16 @@ namespace Chaos
 
 		if (Parent->R().ContainsNaN())
 		{
-			InertiaTensor = PMatrix<FReal, 3, 3>(1.f, 1.f, 1.f);
+			ParentInertia = FMatrix33(1,1,1);
 			Parent->R() = TRotation<FReal, 3>(FMatrix::Identity);
 		}
 
 		Parent->Q() = Parent->R();
-		Parent->InvI() = Parent->I().Inverse();
+		
+		// TODO: This should be computed from rotated tensor by caller.
+		// Initializing to keep current behaviour until fixed.
+		Parent->I() = ParentInertia.GetDiagonal(); 
+		Parent->InvI() = TVec3<FRealSingle>(1.f / Parent->I()[0], 1.f / Parent->I()[1], 1.f / Parent->I()[2]);
 	}
 
 
diff --git a/Engine/Source/Runtime/Experimental/Chaos/Private/PhysicsProxy/GeometryCollectionPhysicsProxy.cpp b/Engine/Source/Runtime/Experimental/Chaos/Private/PhysicsProxy/GeometryCollectionPhysicsProxy.cpp
--- a/Engine/Source/Runtime/Experimental/Chaos/Private/PhysicsProxy/GeometryCollectionPhysicsProxy.cpp
+++ b/Engine/Source/Runtime/Experimental/Chaos/Private/PhysicsProxy/GeometryCollectionPhysicsProxy.cpp
@@ -147,7 +147,7 @@ void PopulateSimulatedParticle(
 	const FCollisionFilterData SimFilterIn,
 	const FCollisionFilterData QueryFilterIn,
 	Chaos::FReal MassIn,
-	FVector3f InertiaTensorVec,
+	Chaos::TVec3<Chaos::FRealSingle> InertiaTensorVec,
 	const FTransform& WorldTransform, 
 	const uint8 DynamicState, 
 	const int16 CollisionGroup,
@@ -195,7 +195,7 @@ void PopulateSimulatedParticle(
 		}
 
 		Handle->SetM(MassIn);
-		Handle->SetI(Chaos::PMatrix<Chaos::FReal, 3, 3>(InertiaTensorVec[0], InertiaTensorVec[1], InertiaTensorVec[2]));
+		Handle->SetI(InertiaTensorVec);
 		Handle->SetObjectStateLowLevel(Chaos::EObjectStateType::Dynamic); // this step sets InvM, InvInertia, P, Q
 	}
 
@@ -1028,7 +1028,7 @@ FGeometryCollectionPhysicsProxy::BuildClusters(
 		SimFilter,
 		QueryFilter,
 		Parent->M() > 0.0 ? Parent->M() : Mass[CollectionClusterIndex] * MassScale, 
-		FVector3f(Parent->I().GetDiagonal() != Chaos::FVec3(0.0) ? Parent->I().GetDiagonal() : Chaos::FVec3(InertiaTensor[CollectionClusterIndex] * MassScale)),
+		Parent->I() != Chaos::TVec3<Chaos::FRealSingle>(0.0) ? Parent->I() : Chaos::TVec3<Chaos::FRealSingle>(InertiaTensor[CollectionClusterIndex] * MassScale),
 		ParticleTM, 
 		(uint8)DynamicState[CollectionClusterIndex], 
 		0,
@@ -2365,7 +2365,10 @@ void FGeometryCollectionPhysicsProxy::InitializeSharedCollisionStructures(
 				//CollectionSimulatableParticles[TransformGroupIndex] = true;
 				IsClusterSimulated[TransformGroupIndex] = true;
 
-				UpdateClusterMassProperties(Handles[ClusterTransformIdx].Get(), ChildrenIndices);	//compute mass properties
+
+				// TODO: This needs to be rotated to diagonal, used to update I()/InvI() from diagonal, and update transform with rotation.
+				FMatrix33 ClusterInertia(0);
+				UpdateClusterMassProperties(Handles[ClusterTransformIdx].Get(), ChildrenIndices, ClusterInertia);	//compute mass properties
 				const FTransform ClusterMassToCollection = 
 					FTransform(CollectionSpaceParticles->R(ClusterTransformIdx), 
 							   CollectionSpaceParticles->X(ClusterTransformIdx));
@@ -2422,9 +2425,9 @@ void FGeometryCollectionPhysicsProxy::InitializeSharedCollisionStructures(
 				} // tmp scope
 
 				TUniquePtr<FTriangleMesh> UnionMesh(new FTriangleMesh(MoveTemp(UnionMeshIndices)));
-				const FMatrix& InertiaMatrix = CollectionSpaceParticles->I(ClusterTransformIdx);
-				const FVector InertiaDiagonal(InertiaMatrix.M[0][0], InertiaMatrix.M[1][1], InertiaMatrix.M[2][2]);
-				CollectionInertiaTensor[ClusterTransformIdx] = (FVector3f)InertiaDiagonal;	// LWC_TODO: Precision loss
+				// TODO: Seems this should rotate full matrix and not discard off diagonals.
+				const FVec3& InertiaDiagonal = CollectionSpaceParticles->I(ClusterTransformIdx);
+				CollectionInertiaTensor[ClusterTransformIdx] = InertiaDiagonal;	// LWC_TODO: Precision loss
 				CollectionMass[ClusterTransformIdx] = (FRealSingle)CollectionSpaceParticles->M(ClusterTransformIdx);
 
 				const int32 SizeSpecificIdx = GeometryCollection::SizeSpecific::FindIndexForVolume(SharedParams.SizeSpecificData, InstanceBoundingBox);
diff --git a/Engine/Source/Runtime/Experimental/Chaos/Private/PhysicsProxy/SkeletalMeshPhysicsProxy.cpp b/Engine/Source/Runtime/Experimental/Chaos/Private/PhysicsProxy/SkeletalMeshPhysicsProxy.cpp
--- a/Engine/Source/Runtime/Experimental/Chaos/Private/PhysicsProxy/SkeletalMeshPhysicsProxy.cpp
+++ b/Engine/Source/Runtime/Experimental/Chaos/Private/PhysicsProxy/SkeletalMeshPhysicsProxy.cpp
@@ -235,8 +235,8 @@ void FSkeletalMeshPhysicsProxy::CreateRigidBodyCallback(FParticlesType& Particle
 
 			Particles.M(RigidBodyId) = FMath::Clamp(TotalMass, (Chaos::FReal)Parameters.MinMass, (Chaos::FReal)Parameters.MaxMass);
 			Particles.InvM(RigidBodyId) = Particles.M(RigidBodyId) < KINDA_SMALL_NUMBER ?  BIG_NUMBER/2 : 1.f / Particles.M(RigidBodyId);
-			Particles.I(RigidBodyId) = MassProperties.InertiaTensor;
-			Particles.InvI(RigidBodyId) = MassProperties.InertiaTensor.Inverse();
+			Particles.I(RigidBodyId) = Chaos::TVec3<Chaos::FRealSingle>((Chaos::FRealSingle)MassProperties.InertiaTensor.M[0][0], (Chaos::FRealSingle)MassProperties.InertiaTensor.M[1][1], (Chaos::FRealSingle)MassProperties.InertiaTensor.M[2][2]);
+			Particles.InvI(RigidBodyId) = Chaos::TVec3<Chaos::FRealSingle>(1.f / (Chaos::FRealSingle)MassProperties.InertiaTensor.M[0][0], 1.f / (Chaos::FRealSingle)MassProperties.InertiaTensor.M[1][1], 1.f / (Chaos::FRealSingle)MassProperties.InertiaTensor.M[2][2]);
 
 			Particles.CollisionGroup(RigidBodyId) = Parameters.CollisionGroup;
 			
diff --git a/Engine/Source/Runtime/Experimental/Chaos/Private/PhysicsProxy/StaticMeshPhysicsObject.cpp b/Engine/Source/Runtime/Experimental/Chaos/Private/PhysicsProxy/StaticMeshPhysicsObject.cpp
--- a/Engine/Source/Runtime/Experimental/Chaos/Private/PhysicsProxy/StaticMeshPhysicsObject.cpp
+++ b/Engine/Source/Runtime/Experimental/Chaos/Private/PhysicsProxy/StaticMeshPhysicsObject.cpp
@@ -278,9 +278,9 @@ void FStaticMeshPhysicsProxy::CreateRigidBodyCallback(FParticlesType& Particles)
 		Particles.Q(RigidBodyId) = Particles.R(RigidBodyId);
 
 		const FVector SideSquared(Bounds.GetSize()[0] * Bounds.GetSize()[0], Bounds.GetSize()[1] * Bounds.GetSize()[1], Bounds.GetSize()[2] * Bounds.GetSize()[2]);
-		const FVector Inertia(Parameters.Mass * (SideSquared.Y + SideSquared.Z) / 12.f, Parameters.Mass * (SideSquared.X + SideSquared.Z) / 12.f, Parameters.Mass * (SideSquared.X + SideSquared.Y) / 12.f);
-		Particles.I(RigidBodyId) = Chaos::PMatrix<Chaos::FReal, 3, 3>(Inertia.X, 0.f, 0.f, 0.f, Inertia.Y, 0.f, 0.f, 0.f, Inertia.Z);
-		Particles.InvI(RigidBodyId) = Chaos::PMatrix<Chaos::FReal, 3, 3>(1.f / Inertia.X, 0.f, 0.f, 0.f, 1.f / Inertia.Y, 0.f, 0.f, 0.f, 1.f / Inertia.Z);
+		const FVector3f Inertia((Chaos::FRealSingle)(Parameters.Mass * (SideSquared.Y + SideSquared.Z) / 12.f), (Chaos::FRealSingle)(Parameters.Mass * (SideSquared.X + SideSquared.Z) / 12.f), (Chaos::FRealSingle)(Parameters.Mass * (SideSquared.X + SideSquared.Y) / 12.f));
+		Particles.I(RigidBodyId) = Inertia;
+		Particles.InvI(RigidBodyId) = Chaos::TVec3<Chaos::FRealSingle>(1.f / Inertia.X, 1.f / Inertia.Y, 1.f / Inertia.Z);
 
 		if(Parameters.ObjectType == EObjectStateTypeEnum::Chaos_Object_Sleeping)
 		{
@@ -290,7 +290,7 @@ void FStaticMeshPhysicsProxy::CreateRigidBodyCallback(FParticlesType& Particles)
 		else if(Parameters.ObjectType != EObjectStateTypeEnum::Chaos_Object_Dynamic)
 		{
 			Particles.InvM(RigidBodyId) = 0.f;
-			Particles.InvI(RigidBodyId) = Chaos::PMatrix<Chaos::FReal, 3, 3>(0);
+			Particles.InvI(RigidBodyId) = Chaos::TVec3<Chaos::FRealSingle>(0);
 			Particles.SetObjectState(RigidBodyId, Chaos::EObjectStateType::Kinematic);
 		}
 		else
diff --git a/Engine/Source/Runtime/Experimental/Chaos/Public/Chaos/Evolution/SolverBodyContainer.h b/Engine/Source/Runtime/Experimental/Chaos/Public/Chaos/Evolution/SolverBodyContainer.h
--- a/Engine/Source/Runtime/Experimental/Chaos/Public/Chaos/Evolution/SolverBodyContainer.h
+++ b/Engine/Source/Runtime/Experimental/Chaos/Public/Chaos/Evolution/SolverBodyContainer.h
@@ -50,7 +50,7 @@ namespace Chaos
 					SolverBody.SetR(PrevCoMTransform.GetRotation());
 
 					SolverBody.SetInvM(Particle->InvM());
-					SolverBody.SetInvILocal(Particle->InvI().GetDiagonal());
+					SolverBody.SetInvILocal(Particle->InvI());
 				}
 				else
 				{
diff --git a/Engine/Source/Runtime/Experimental/Chaos/Public/Chaos/PBDRigidClusteringAlgo.h b/Engine/Source/Runtime/Experimental/Chaos/Public/Chaos/PBDRigidClusteringAlgo.h
--- a/Engine/Source/Runtime/Experimental/Chaos/Public/Chaos/PBDRigidClusteringAlgo.h
+++ b/Engine/Source/Runtime/Experimental/Chaos/Public/Chaos/PBDRigidClusteringAlgo.h
@@ -8,6 +8,7 @@ namespace Chaos
 	void UpdateClusterMassProperties(
 		FPBDRigidClusteredParticleHandle* Parent,
 		TSet<FPBDRigidParticleHandle*>& Children,
+		FMatrix33& ParentInertia,
 		const FRigidTransform3* ForceMassOrientation = nullptr);
 
 	void CHAOS_API UpdateKinematicProperties(
diff --git a/Engine/Source/Runtime/Experimental/Chaos/Public/Chaos/PBDRigidParticles.h b/Engine/Source/Runtime/Experimental/Chaos/Public/Chaos/PBDRigidParticles.h
--- a/Engine/Source/Runtime/Experimental/Chaos/Public/Chaos/PBDRigidParticles.h
+++ b/Engine/Source/Runtime/Experimental/Chaos/Public/Chaos/PBDRigidParticles.h
@@ -126,27 +126,27 @@ class TPBDRigidParticles : public TRigidParticles<T, d>
 		{
 			// When the state is first initialized, treat it like a static.
 			this->InvM(Index) = 0.0f;
-			this->InvI(Index) = FMatrix33(0);
+			this->InvI(Index) = TVec3<FRealSingle>(0);
 		}
 
 		if ((CurrentState == EObjectStateType::Dynamic || CurrentState == EObjectStateType::Sleeping) && (InObjectState == EObjectStateType::Kinematic || InObjectState == EObjectStateType::Static))
 		{
 			// Transitioning from dynamic to static or kinematic, set inverse mass and inertia tensor to zero.
 			this->InvM(Index) = 0.0f;
-			this->InvI(Index) = FMatrix33(0);
+			this->InvI(Index) = TVec3<FRealSingle>(0);
 		}
 		else if ((CurrentState == EObjectStateType::Kinematic || CurrentState == EObjectStateType::Static || CurrentState == EObjectStateType::Uninitialized) && (InObjectState == EObjectStateType::Dynamic || InObjectState == EObjectStateType::Sleeping))
 		{
 			// Transitioning from kinematic or static to dynamic, compute the inverses.
 			checkSlow(this->M(Index) != 0.0);
-			checkSlow(this->I(Index).M[0][0] != 0.0);
-			checkSlow(this->I(Index).M[1][1] != 0.0);
-			checkSlow(this->I(Index).M[2][2] != 0.0);
+			checkSlow(this->I(Index)[0] != 0.0);
+			checkSlow(this->I(Index)[1] != 0.0);
+			checkSlow(this->I(Index)[2] != 0.0);
 			this->InvM(Index) = 1.f / this->M(Index);
-			this->InvI(Index) = Chaos::FMatrix33(
-				1.f / this->I(Index).M[0][0], 0.f, 0.f,
-				0.f, 1.f / this->I(Index).M[1][1], 0.f,
-				0.f, 0.f, 1.f / this->I(Index).M[2][2]);
+			this->InvI(Index) = TVec3<FRealSingle>(
+				1.f / this->I(Index)[0], 
+				1.f / this->I(Index)[1],
+				1.f / this->I(Index)[2]);
 
 			this->P(Index) = this->X(Index);
 			this->Q(Index) = this->R(Index);
diff --git a/Engine/Source/Runtime/Experimental/Chaos/Public/Chaos/ParticleDirtyFlags.h b/Engine/Source/Runtime/Experimental/Chaos/Public/Chaos/ParticleDirtyFlags.h
--- a/Engine/Source/Runtime/Experimental/Chaos/Public/Chaos/ParticleDirtyFlags.h
+++ b/Engine/Source/Runtime/Experimental/Chaos/Public/Chaos/ParticleDirtyFlags.h
@@ -265,11 +265,11 @@ public:
 	const FRotation3& RotationOfMass() const { return MRotationOfMass; }
 	void SetRotationOfMass(const FRotation3& InRotationOfMass){ MRotationOfMass = InRotationOfMass; }
 
-	const FMatrix33& I() const { return MI; }
-	void SetI(const FMatrix33& InI){ MI = InI; }
+	const TVec3<FRealSingle>& I() const { return MI; }
+	void SetI(const TVec3<FRealSingle>& InI){ MI = InI; }
 
-	const FMatrix33& InvI() const { return MInvI; }
-	void SetInvI(const FMatrix33& InInvI){ MInvI = InInvI; }
+	const TVec3<FRealSingle>& InvI() const { return MInvI; }
+	void SetInvI(const TVec3<FRealSingle>& InInvI){ MInvI = InInvI; }
 
 	FReal M() const { return MM; }
 	void SetM(FReal InM){ MM = InM; }
@@ -280,8 +280,8 @@ public:
 private:
 	FVec3 MCenterOfMass;
 	FRotation3 MRotationOfMass;
-	FMatrix33 MI;
-	FMatrix33 MInvI;
+	TVec3<FRealSingle> MI;
+	TVec3<FRealSingle> MInvI;
 	FReal MM;
 	FReal MInvM;
 
diff --git a/Engine/Source/Runtime/Experimental/Chaos/Public/Chaos/ParticleHandle.h b/Engine/Source/Runtime/Experimental/Chaos/Public/Chaos/ParticleHandle.h
--- a/Engine/Source/Runtime/Experimental/Chaos/Public/Chaos/ParticleHandle.h
+++ b/Engine/Source/Runtime/Experimental/Chaos/Public/Chaos/ParticleHandle.h
@@ -97,8 +97,8 @@ void PBDRigidParticleDefaultConstruct(FConcrete& Concrete, const FPBDRigidPartic
 	Concrete.SetInvM(1);
 	Concrete.SetCenterOfMass(TVector<T,d>(0));
 	Concrete.SetRotationOfMass(TRotation<T, d>::FromIdentity());
-	Concrete.SetI(PMatrix<T, d, d>(1, 1, 1));
-	Concrete.SetInvI(PMatrix<T, d, d>(1, 1, 1));
+	Concrete.SetI(TVec3<FRealSingle>(1, 1, 1));
+	Concrete.SetInvI(TVec3<FRealSingle>(1, 1, 1));
 	Concrete.SetLinearEtherDrag(0.f);
 	Concrete.SetAngularEtherDrag(0.f);
 	Concrete.SetGravityEnabled(Params.bGravityEnabled);
@@ -967,13 +967,13 @@ public:
 		SetWSmooth(W());
 	}
 
-	const PMatrix<T, d, d>& I() const { return PBDRigidParticles->I(ParticleIdx); }
-	PMatrix<T, d, d>& I() { return PBDRigidParticles->I(ParticleIdx); }
-	void SetI(const PMatrix<T, d, d>& InI) { PBDRigidParticles->I(ParticleIdx) = InI; }
+	const TVec3<FRealSingle>& I() const { return PBDRigidParticles->I(ParticleIdx); }
+	TVec3<FRealSingle>& I() { return PBDRigidParticles->I(ParticleIdx); }
+	void SetI(const TVec3<FRealSingle>& InI) { PBDRigidParticles->I(ParticleIdx) = InI; }
 
-	const PMatrix<T, d, d>& InvI() const { return PBDRigidParticles->InvI(ParticleIdx); }
-	PMatrix<T, d, d>& InvI() { return PBDRigidParticles->InvI(ParticleIdx); }
-	void SetInvI(const PMatrix<T, d, d>& InInvI) { PBDRigidParticles->InvI(ParticleIdx) = InInvI; }
+	const TVec3<FRealSingle>& InvI() const { return PBDRigidParticles->InvI(ParticleIdx); }
+	TVec3<FRealSingle>& InvI() { return PBDRigidParticles->InvI(ParticleIdx); }
+	void SetInvI(const TVec3<FRealSingle>& InInvI) { PBDRigidParticles->InvI(ParticleIdx) = InInvI; }
 
 	T M() const { return PBDRigidParticles->M(ParticleIdx); }
 	T& M() { return PBDRigidParticles->M(ParticleIdx); }
@@ -1552,24 +1552,24 @@ public:
 		MHandle->UpdateWorldSpaceState(WorldTransform, BoundsExpansion);
 	}
 
-	const FMatrix33& I() const 
+	const TVec3<FRealSingle> I() const
 	{ 
 		if (MHandle->CastToRigidParticle() && MHandle->ObjectState() == EObjectStateType::Dynamic)
 		{
 			return MHandle->CastToRigidParticle()->I();
 		}
 
-		return ZeroMatrix;
+		return TVec3<FRealSingle>(0);
 	}
 
-	const FMatrix33& InvI() const 
+	const TVec3<FRealSingle> InvI() const
 	{ 
 		if (MHandle->CastToRigidParticle() && MHandle->ObjectState() == EObjectStateType::Dynamic)
 		{
 			return MHandle->CastToRigidParticle()->InvI();
 		}
 
-		return ZeroMatrix;
+		return TVec3<FRealSingle>(0);
 	}
 
 	FReal M() const
@@ -2521,14 +2521,14 @@ public:
 		// Physics thread only. API required for FGeometryParticleStateBase::SyncToParticle
 	}
 
-	const PMatrix<T, d, d>& I() const { return MMassProps.Read().I(); }
-	void SetI(const PMatrix<T, d, d>& InI)
+	const TVec3<FRealSingle>& I() const { return MMassProps.Read().I(); }
+	void SetI(const TVec3<FRealSingle>& InI)
 	{
 		MMassProps.Modify(true,MDirtyFlags,Proxy,[&InI](auto& Data){ Data.SetI(InI);});
 	}
 
-	const PMatrix<T, d, d>& InvI() const { return MMassProps.Read().InvI(); }
-	void SetInvI(const PMatrix<T, d, d>& InInvI)
+	const TVec3<FRealSingle>& InvI() const { return MMassProps.Read().InvI(); }
+	void SetInvI(const TVec3<FRealSingle>& InInvI)
 	{
 		MMassProps.Modify(true,MDirtyFlags,Proxy,[&InInvI](auto& Data){ Data.SetInvI(InInvI);});
 	}
diff --git a/Engine/Source/Runtime/Experimental/Chaos/Public/Chaos/RigidParticles.h b/Engine/Source/Runtime/Experimental/Chaos/Public/Chaos/RigidParticles.h
--- a/Engine/Source/Runtime/Experimental/Chaos/Public/Chaos/RigidParticles.h
+++ b/Engine/Source/Runtime/Experimental/Chaos/Public/Chaos/RigidParticles.h
@@ -175,11 +175,11 @@ public:
 	FORCEINLINE const TVector<T, d>& AngularImpulseVelocity(const int32 Index) const { return MAngularImpulseVelocity[Index]; }
 	FORCEINLINE TVector<T, d>& AngularImpulseVelocity(const int32 Index) { return MAngularImpulseVelocity[Index]; }
 
-	FORCEINLINE const PMatrix<T, d, d>& I(const int32 Index) const { return MI[Index]; }
-	FORCEINLINE PMatrix<T, d, d>& I(const int32 Index) { return MI[Index]; }
+	FORCEINLINE const TVec3<FRealSingle>& I(const int32 Index) const { return MI[Index]; }
+	FORCEINLINE TVec3<FRealSingle>& I(const int32 Index) { return MI[Index]; }
 
-	FORCEINLINE const PMatrix<T, d, d>& InvI(const int32 Index) const { return MInvI[Index]; }
-	FORCEINLINE PMatrix<T, d, d>& InvI(const int32 Index) { return MInvI[Index]; }
+	FORCEINLINE const TVec3<FRealSingle>& InvI(const int32 Index) const { return MInvI[Index]; }
+	FORCEINLINE TVec3<FRealSingle>& InvI(const int32 Index) { return MInvI[Index]; }
 
 	FORCEINLINE const T M(const int32 Index) const { return MM[Index]; }
 	FORCEINLINE T& M(const int32 Index) { return MM[Index]; }
@@ -304,7 +304,27 @@ public:
 			Ar << MRotationOfMass;
 		}
 
-		Ar << MAcceleration << MAngularAcceleration << MLinearImpulseVelocity << MAngularImpulseVelocity << MI << MInvI << MM << MInvM;
+		Ar << MAcceleration << MAngularAcceleration << MLinearImpulseVelocity << MAngularImpulseVelocity;
+
+		Ar.UsingCustomVersion(FUE5ReleaseStreamObjectVersion::GUID);
+		if (Ar.IsLoading() && Ar.CustomVer(FUE5ReleaseStreamObjectVersion::GUID) < FUE5ReleaseStreamObjectVersion::ChaosInertiaConvertedToVec3)
+		{
+			TArray<PMatrix<T, d, d>> IArray;
+			TArray<PMatrix<T, d, d>> InvIArray;
+			Ar << IArray << InvIArray;
+
+			for (int32 Idx = 0; Idx < IArray.Num(); ++Idx)
+			{
+				MI.Add(IArray[Idx].GetDiagonal());
+				MInvI.Add(InvIArray[Idx].GetDiagonal());
+			}
+		}
+		else
+		{
+			Ar << MI << MInvI;
+		}
+		
+		Ar << MM << MInvM;
 
 		Ar.UsingCustomVersion(FExternalPhysicsCustomObjectVersion::GUID);
 		if (Ar.CustomVer(FExternalPhysicsCustomObjectVersion::GUID) >= FExternalPhysicsCustomObjectVersion::AddDampingToRigids)
@@ -320,8 +340,8 @@ public:
 	FORCEINLINE TArray<TVector<T, d>>& AllAngularAcceleration() { return MAngularAcceleration; }
 	FORCEINLINE TArray<TVector<T, d>>& AllLinearImpulseVelocity() { return MLinearImpulseVelocity; }
 	FORCEINLINE TArray<TVector<T, d>>& AllAngularImpulseVelocity() { return MAngularImpulseVelocity; }
-	FORCEINLINE TArray<PMatrix<T, d, d>>& AllI() { return MI; }
-	FORCEINLINE TArray<PMatrix<T, d, d>>& AllInvI() { return MInvI; }
+	FORCEINLINE TArray<TVec3<FRealSingle>>& AllI() { return MI; }
+	FORCEINLINE TArray<TVec3<FRealSingle>>& AllInvI() { return MInvI; }
 	FORCEINLINE TArray<FReal>& AllM() { return MM; }
 	FORCEINLINE TArray<FReal>& AllInvM() { return MInvM; }
 	FORCEINLINE TArray<TVector<T, d>>& AllCenterOfMass() { return MCenterOfMass; }
@@ -342,8 +362,8 @@ private:
 	TArrayCollectionArray<TVector<T, d>> MAngularAcceleration;
 	TArrayCollectionArray<TVector<T, d>> MLinearImpulseVelocity;
 	TArrayCollectionArray<TVector<T, d>> MAngularImpulseVelocity;
-	TArrayCollectionArray<PMatrix<T, d, d>> MI;
-	TArrayCollectionArray<PMatrix<T, d, d>> MInvI;
+	TArrayCollectionArray<TVec3<FRealSingle>> MI;
+	TArrayCollectionArray<TVec3<FRealSingle>> MInvI;
 	TArrayCollectionArray<T> MM;
 	TArrayCollectionArray<T> MInvM;
 	TArrayCollectionArray<TVector<T,d>> MCenterOfMass;
diff --git a/Engine/Source/Runtime/Experimental/Chaos/Public/PhysicsProxy/SingleParticlePhysicsProxy.h b/Engine/Source/Runtime/Experimental/Chaos/Public/PhysicsProxy/SingleParticlePhysicsProxy.h
--- a/Engine/Source/Runtime/Experimental/Chaos/Public/PhysicsProxy/SingleParticlePhysicsProxy.h
+++ b/Engine/Source/Runtime/Experimental/Chaos/Public/PhysicsProxy/SingleParticlePhysicsProxy.h
@@ -650,7 +650,7 @@ public:
 		});
 	}
 
-	const FMatrix33 I() const
+	const Chaos::TVec3<FRealSingle> I() const
 	{
 		return Read([](auto* Particle)
 		{
@@ -659,11 +659,11 @@ public:
 				return Rigid->I();
 			}
 
-			return FMatrix33(0, 0, 0);
+			return Chaos::TVec3<FRealSingle>(0, 0, 0);
 		});
 	}
 
-	void SetI(const FMatrix33& InI)
+	void SetI(const Chaos::TVec3<FRealSingle>& InI)
 	{
 		Write([&InI](auto* Particle)
 		{
@@ -674,7 +674,7 @@ public:
 		});
 	}
 
-	const FMatrix33 InvI() const
+	const Chaos::TVec3<FRealSingle> InvI() const
 	{
 		return Read([](auto* Particle)
 		{
@@ -683,11 +683,11 @@ public:
 				return Rigid->InvI();
 			}
 
-			return FMatrix33(0, 0, 0);
+			return Chaos::TVec3<FRealSingle>(0, 0, 0);
 		});
 	}
 
-	void SetInvI(const FMatrix33& InInvI)
+	void SetInvI(const Chaos::TVec3<FRealSingle>& InInvI)
 	{
 		Write([&InInvI](auto* Particle)
 		{
diff --git a/Engine/Source/Runtime/Experimental/ChaosCore/Public/Chaos/Core.h b/Engine/Source/Runtime/Experimental/ChaosCore/Public/Chaos/Core.h
--- a/Engine/Source/Runtime/Experimental/ChaosCore/Public/Chaos/Core.h
+++ b/Engine/Source/Runtime/Experimental/ChaosCore/Public/Chaos/Core.h
@@ -42,6 +42,6 @@ namespace Chaos
 	// NOTE: if you get a merge conflict on the GUID, you must replace it with a new GUID - do not accept the source or target
 	// or you will likely get DDC version conflicts resulting in crashes during load.
 	// Core version string for Chaos data. Any DDC builder dependent on Chaos for serialization should depend on this version
-	static const FString ChaosVersionString = TEXT("A022A2AA-C0D0-4977-BCE3-3A6CF7D3B5C6");
+	static const FString ChaosVersionString = TEXT("34C4ADF6-BD70-4A04-A811-E9A24C643F32");
 
 }
diff --git a/Engine/Source/Runtime/Experimental/ChaosCore/Public/Chaos/Matrix.h b/Engine/Source/Runtime/Experimental/ChaosCore/Public/Chaos/Matrix.h
--- a/Engine/Source/Runtime/Experimental/ChaosCore/Public/Chaos/Matrix.h
+++ b/Engine/Source/Runtime/Experimental/ChaosCore/Public/Chaos/Matrix.h
@@ -271,6 +271,26 @@ namespace Chaos
 			M[1][3] = 0;
 			M[2][3] = 0;
 		}
+		explicit PMatrix(const TVector<FRealDouble, 3>& Vector)
+			: UE::Math::TMatrix<FReal>()
+		{
+			M[0][0] = Vector[0];
+			M[1][0] = 0;
+			M[2][0] = 0;
+			M[0][1] = 0;
+			M[1][1] = Vector[1];
+			M[2][1] = 0;
+			M[0][2] = 0;
+			M[1][2] = 0;
+			M[2][2] = Vector[2];
+			M[3][0] = 0;
+			M[3][1] = 0;
+			M[3][2] = 0;
+			M[3][3] = 1;
+			M[0][3] = 0;
+			M[1][3] = 0;
+			M[2][3] = 0;
+		}
 		PMatrix(const FRealDouble x00, const FRealDouble x10, const FRealDouble x20, const FRealDouble x11, const FRealDouble x21, const FRealDouble x22)
 		    : UE::Math::TMatrix<FRealDouble>()
 		{
@@ -678,6 +698,26 @@ namespace Chaos
 			M[1][3] = 0;
 			M[2][3] = 0;
 		}
+		explicit PMatrix(const TVector<FRealSingle, 3>& Vector)
+			: UE::Math::TMatrix<FReal>()
+		{
+			M[0][0] = Vector[0];
+			M[1][0] = 0;
+			M[2][0] = 0;
+			M[0][1] = 0;
+			M[1][1] = Vector[1];
+			M[2][1] = 0;
+			M[0][2] = 0;
+			M[1][2] = 0;
+			M[2][2] = Vector[2];
+			M[3][0] = 0;
+			M[3][1] = 0;
+			M[3][2] = 0;
+			M[3][3] = 1;
+			M[0][3] = 0;
+			M[1][3] = 0;
+			M[2][3] = 0;
+		}
 		PMatrix(const FRealSingle x00, const FRealSingle x10, const FRealSingle x20, const FRealSingle x11, const FRealSingle x21, const FRealSingle x22)
 			: UE::Math::TMatrix<FRealSingle>()
 		{
diff --git a/Engine/Source/Runtime/PhysicsCore/Private/ChaosEngineInterface.cpp b/Engine/Source/Runtime/PhysicsCore/Private/ChaosEngineInterface.cpp
--- a/Engine/Source/Runtime/PhysicsCore/Private/ChaosEngineInterface.cpp
+++ b/Engine/Source/Runtime/PhysicsCore/Private/ChaosEngineInterface.cpp
@@ -686,8 +686,7 @@ FTransform FChaosEngineInterface::GetComTransformLocal_AssumesLocked(const FPhys
 
 FVector FChaosEngineInterface::GetLocalInertiaTensor_AssumesLocked(const FPhysicsActorHandle& InActorReference)
 {
-	const Chaos::FMatrix33 Tensor = InActorReference->GetGameThreadAPI().I();
-	return FVector(Tensor.M[0][0],Tensor.M[1][1],Tensor.M[2][2]);
+	return InActorReference->GetGameThreadAPI().I();
 }
 
 FBox FChaosEngineInterface::GetBounds_AssumesLocked(const FPhysicsActorHandle& InActorReference)
@@ -878,8 +877,8 @@ void FChaosEngineInterface::SetMassSpaceInertiaTensor_AssumesLocked(FPhysicsActo
 	if(CHAOS_ENSURE(!FMath::IsNearlyZero(InTensor.X)) && CHAOS_ENSURE(!FMath::IsNearlyZero(InTensor.Y)) && CHAOS_ENSURE(!FMath::IsNearlyZero(InTensor.Z)))
 	{
 		Chaos::FRigidBodyHandle_External& Body_External = InActorReference->GetGameThreadAPI();
-		Body_External.SetI(Chaos::FMatrix33(InTensor.X,InTensor.Y,InTensor.Z));
-		Body_External.SetInvI(Chaos::FMatrix33(1./InTensor.X,1./InTensor.Y,1./InTensor.Z));
+		Body_External.SetI(Chaos::TVec3<Chaos::FRealSingle>(InTensor.X,InTensor.Y,InTensor.Z));
+		Body_External.SetInvI(Chaos::TVec3<Chaos::FRealSingle>(1./InTensor.X,1./InTensor.Y,1./InTensor.Z));
 	}
 }
 
