Fix low level test compile errors

diff --git a/Engine/Source/Programs/HeadlessChaos/Private/GeometryCollection/GeometryCollectionTestClean.cpp b/Engine/Source/Programs/HeadlessChaos/Private/GeometryCollection/GeometryCollectionTestClean.cpp
--- a/Engine/Source/Programs/HeadlessChaos/Private/GeometryCollection/GeometryCollectionTestClean.cpp
+++ b/Engine/Source/Programs/HeadlessChaos/Private/GeometryCollection/GeometryCollectionTestClean.cpp
@@ -120,8 +120,8 @@ namespace GeometryCollectionTest
 			int TinyFaces = 0;
 			for (const FIntVector & Face : InColl->Indices)
 			{
-				FVector p10 = InColl->Vertex[Face.Y] - InColl->Vertex[Face.X];
-				FVector p20 = InColl->Vertex[Face.Z] - InColl->Vertex[Face.X];
+				FVector p10 = FVector(InColl->Vertex[Face.Y] - InColl->Vertex[Face.X]);
+				FVector p20 = FVector(InColl->Vertex[Face.Z] - InColl->Vertex[Face.X]);
 				FVector Cross = FVector::CrossProduct(p20, p10);
 				if (Cross.SizeSquared() < InTinyNumber)
 				{
diff --git a/Engine/Source/Programs/HeadlessChaos/Private/GeometryCollection/GeometryCollectionTestFramework.cpp b/Engine/Source/Programs/HeadlessChaos/Private/GeometryCollection/GeometryCollectionTestFramework.cpp
--- a/Engine/Source/Programs/HeadlessChaos/Private/GeometryCollection/GeometryCollectionTestFramework.cpp
+++ b/Engine/Source/Programs/HeadlessChaos/Private/GeometryCollection/GeometryCollectionTestFramework.cpp
@@ -109,7 +109,7 @@ namespace GeometryCollectionTest
 	{
 		FGeometryCollection* Collection = FGeometryCollection::NewGeometryCollection(SphereGeometry::RawVertexArray, SphereGeometry::RawIndicesArray);
 		TManagedArray<FVector3f>& Vertices = Collection->Vertex;
-		for (int i = 0; i < Vertices.Num(); i++) Vertices[i] = GeomTransform.TransformPosition(Vertices[i]);
+		for (int i = 0; i < Vertices.Num(); i++) Vertices[i] = FVector3f(GeomTransform.TransformPosition(FVector(Vertices[i])));
 		return TSharedPtr<FGeometryCollection>(Collection);
 	}
 
diff --git a/Engine/Source/Programs/HeadlessChaos/Private/GeometryCollection/GeometryCollectionTestInitilization.cpp b/Engine/Source/Programs/HeadlessChaos/Private/GeometryCollection/GeometryCollectionTestInitilization.cpp
--- a/Engine/Source/Programs/HeadlessChaos/Private/GeometryCollection/GeometryCollectionTestInitilization.cpp
+++ b/Engine/Source/Programs/HeadlessChaos/Private/GeometryCollection/GeometryCollectionTestInitilization.cpp
@@ -125,7 +125,7 @@ GTEST_TEST(AllTraits,GeometryCollection_Initilization_TransformedGeometryCollect
 	FVector CenterOfMass(0);  
 	for (int VertexIndex = 0; VertexIndex < NumVertices; VertexIndex++)
 	{
-		CenterOfMass += RestTransforms[BoneMap[VertexIndex]].TransformPosition(Vertices[VertexIndex]);
+		CenterOfMass += RestTransforms[BoneMap[VertexIndex]].TransformPosition(FVector(Vertices[VertexIndex]));
 	}
 	CenterOfMass /= NumVertices;
 	EXPECT_NEAR(CenterOfMass.X - GlobalTranslation[0],0.0f,KINDA_SMALL_NUMBER);
diff --git a/Engine/Source/Programs/HeadlessChaos/Private/GeometryCollection/GeometryCollectionTestMassProperties.cpp b/Engine/Source/Programs/HeadlessChaos/Private/GeometryCollection/GeometryCollectionTestMassProperties.cpp
--- a/Engine/Source/Programs/HeadlessChaos/Private/GeometryCollection/GeometryCollectionTestMassProperties.cpp
+++ b/Engine/Source/Programs/HeadlessChaos/Private/GeometryCollection/GeometryCollectionTestMassProperties.cpp
@@ -134,7 +134,7 @@ namespace GeometryCollectionTest
 		MassSpaceParticles.AddParticles(Vertex.Num());
 		for (int32 Idx = 0; Idx < Vertex.Num(); ++Idx)
 		{
-			MassSpaceParticles.X(Idx) = Transform[BoneMap[Idx]].TransformPosition(Vertex[Idx]);
+			MassSpaceParticles.X(Idx) = Transform[BoneMap[Idx]].TransformPosition(FVector(Vertex[Idx]));
 		}
 
 		CalculateVolumeAndCenterOfMass(MassSpaceParticles, TriMesh->GetElements(), MassProperties.Volume, MassProperties.CenterOfMass);
@@ -207,7 +207,7 @@ namespace GeometryCollectionTest
 		MassSpaceParticles.AddParticles(Vertex.Num());
 		for (int32 Idx = 0; Idx < Vertex.Num(); ++Idx)
 		{
-			MassSpaceParticles.X(Idx) = Transform[BoneMap[Idx]].TransformPosition(Vertex[Idx]);
+			MassSpaceParticles.X(Idx) = Transform[BoneMap[Idx]].TransformPosition(FVector(Vertex[Idx]));
 		}
 
 		CalculateVolumeAndCenterOfMass(MassSpaceParticles, TriMesh->GetElements(), MassProperties.Volume, MassProperties.CenterOfMass);
@@ -289,7 +289,7 @@ namespace GeometryCollectionTest
 		MassSpaceParticles.AddParticles(Vertex.Num());
 		for (int32 Idx = 0; Idx < Vertex.Num(); ++Idx)
 		{
-			MassSpaceParticles.X(Idx) = Transform[BoneMap[Idx]].TransformPosition(Vertex[Idx]);
+			MassSpaceParticles.X(Idx) = Transform[BoneMap[Idx]].TransformPosition(FVector(Vertex[Idx]));
 		}
 
 		CalculateVolumeAndCenterOfMass(MassSpaceParticles, TriMesh->GetElements(), MassProperties.Volume, MassProperties.CenterOfMass);
@@ -370,8 +370,8 @@ namespace GeometryCollectionTest
 		MassSpaceParticles.AddParticles(Vertex.Num());
 		for (int32 Idx = 0; Idx < Vertex.Num(); ++Idx)
 		{
-			FVector VertexPoint = Vertex[Idx];
-			MassSpaceParticles.X(Idx) = Transform[BoneMap[Idx]].TransformPosition(Vertex[Idx]);
+			FVector VertexPoint = FVector(Vertex[Idx]);
+			MassSpaceParticles.X(Idx) = Transform[BoneMap[Idx]].TransformPosition(FVector(Vertex[Idx]));
 			FVector MassSpacePoint = FVector(MassSpaceParticles.X(Idx)[0], MassSpaceParticles.X(Idx)[1], MassSpaceParticles.X(Idx)[2]);
 			SomeVec.Add(MassSpacePoint);
 		}
diff --git a/Engine/Source/Programs/HeadlessChaos/Private/HeadlessChaosTestEPA.cpp b/Engine/Source/Programs/HeadlessChaos/Private/HeadlessChaosTestEPA.cpp
--- a/Engine/Source/Programs/HeadlessChaos/Private/HeadlessChaosTestEPA.cpp
+++ b/Engine/Source/Programs/HeadlessChaos/Private/HeadlessChaosTestEPA.cpp
@@ -1249,9 +1249,10 @@ namespace ChaosTest
 			FVec3 Normal(0.0f);
 			FVec3 Position(0.0f);
 			int32 FaceIndex = -1;
-			bool bResult = ScaledTriangleMesh.LowLevelSweepGeom(ConvexBox, Transform0, Dir, Length, OutTime, Position, Normal, FaceIndex, 0.0f, true);
+			FVec3 FaceNormal(0.0);
+			bool bResult = ScaledTriangleMesh.LowLevelSweepGeom(ConvexBox, Transform0, Dir, Length, OutTime, Position, Normal, FaceIndex, FaceNormal, 0.0f, true);
 
-			bResult = ScaledTriangleMesh.LowLevelSweepGeom(ConvexBox, Transform1, Dir, Length, OutTime, Position, Normal, FaceIndex, 0.0f, true);
+			bResult = ScaledTriangleMesh.LowLevelSweepGeom(ConvexBox, Transform1, Dir, Length, OutTime, Position, Normal, FaceIndex, FaceNormal, 0.0f, true);
 			
 			// Observe that normals are in opposite direction, while rotations are very similar.
 
diff --git a/Engine/Source/Programs/HeadlessChaos/Private/HeadlessChaosTestImplicits.cpp b/Engine/Source/Programs/HeadlessChaos/Private/HeadlessChaosTestImplicits.cpp
--- a/Engine/Source/Programs/HeadlessChaos/Private/HeadlessChaosTestImplicits.cpp
+++ b/Engine/Source/Programs/HeadlessChaos/Private/HeadlessChaosTestImplicits.cpp
@@ -1566,11 +1566,11 @@ namespace ChaosTest {
 
 			for (int32 Idx = 0; Idx < CulledParticles.Num(); ++Idx)
 			{
-				EXPECT_NE(Particles.X(8), CulledParticles[Idx]);	//interior particle gone
+				EXPECT_NE(Particles.X(8), (Chaos::TVector<FRealDouble, 3>)CulledParticles[Idx]);	//interior particle gone
 				bool bFound = false;
 				for (uint32 InnerIdx = 0; InnerIdx < Particles.Size(); ++InnerIdx)	//remaining particles are from the original set
 				{
-					if (Particles.X(InnerIdx) == CulledParticles[Idx])
+					if (Particles.X(InnerIdx) == (Chaos::TVector<FRealDouble,3>)CulledParticles[Idx])
 					{
 						bFound = true;
 						break;
diff --git a/Engine/Source/Programs/HeadlessChaos/Private/HeadlessChaosTestSweep.cpp b/Engine/Source/Programs/HeadlessChaos/Private/HeadlessChaosTestSweep.cpp
--- a/Engine/Source/Programs/HeadlessChaos/Private/HeadlessChaosTestSweep.cpp
+++ b/Engine/Source/Programs/HeadlessChaos/Private/HeadlessChaosTestSweep.cpp
@@ -873,7 +873,8 @@ namespace ChaosTest
 		FVec3 Normal(0.0);
 		FVec3 Position(0.0);
 		int32 FaceIndex = -1;
-		bool bResult = ScaledTriangleMesh.LowLevelSweepGeom(Capsule, CapsuleToTrimesh, Dir, Length, OutTime, Position, Normal, FaceIndex, 0.0f, true);
+		FVec3 FaceNormal(0.0);
+		bool bResult = ScaledTriangleMesh.LowLevelSweepGeom(Capsule, CapsuleToTrimesh, Dir, Length, OutTime, Position, Normal, FaceIndex, FaceNormal, 0.0f, true);
 		FVec3 WorldPosition = TrimeshTransform.TransformPositionNoScale(Position);
 
 		EXPECT_EQ(bResult, true);
diff --git a/Engine/Source/Runtime/Experimental/Chaos/Public/Chaos/SAT.h b/Engine/Source/Runtime/Experimental/Chaos/Public/Chaos/SAT.h
--- a/Engine/Source/Runtime/Experimental/Chaos/Public/Chaos/SAT.h
+++ b/Engine/Source/Runtime/Experimental/Chaos/Public/Chaos/SAT.h
@@ -167,8 +167,8 @@ namespace Chaos
 			// Edge Vertices in other object space
 			const int32 EdgeVertexIndex2A = Convex2.GetEdgeVertex(EdgeIndex2, 0);
 			const int32 EdgeVertexIndex2B = Convex2.GetEdgeVertex(EdgeIndex2, 1);
-			const FVec3 EdgeVertex2AIn1 = Convex2ToConvex1Transform.TransformPositionNoScale(Convex2.GetVertex(EdgeVertexIndex2A));
-			const FVec3 EdgeVertex2BIn1 = Convex2ToConvex1Transform.TransformPositionNoScale(Convex2.GetVertex(EdgeVertexIndex2B));
+			const FVec3 EdgeVertex2AIn1 = Convex2ToConvex1Transform.TransformPositionNoScale(FVector(Convex2.GetVertex(EdgeVertexIndex2A)));
+			const FVec3 EdgeVertex2BIn1 = Convex2ToConvex1Transform.TransformPositionNoScale(FVector(Convex2.GetVertex(EdgeVertexIndex2B)));
 
 			// Planes that use the edge
 			const int32 EdgePlaneIndex2A = Convex2.GetEdgePlane(EdgeIndex2, 0);
diff --git a/Engine/Source/Runtime/Experimental/Chaos/Public/GeometryCollection/GeometryCollectionUtility.h b/Engine/Source/Runtime/Experimental/Chaos/Public/GeometryCollection/GeometryCollectionUtility.h
--- a/Engine/Source/Runtime/Experimental/Chaos/Public/GeometryCollection/GeometryCollectionUtility.h
+++ b/Engine/Source/Runtime/Experimental/Chaos/Public/GeometryCollection/GeometryCollectionUtility.h
@@ -78,8 +78,8 @@ namespace GeometryCollection
 
 			for (int32 Axis = 0; Axis < 3; ++Axis)
 			{
-				const FVector& Normal = Normals[Tri[Axis]];
-				const FVector Edge = (Vertices[Tri[(Axis + 1) % 3]] - Vertices[Tri[Axis]]);
+				const FVector3f& Normal = Normals[Tri[Axis]];
+				const FVector3f Edge = (Vertices[Tri[(Axis + 1) % 3]] - Vertices[Tri[Axis]]);
 				TangentU[Tri[Axis]] = (Edge ^ Normal).GetSafeNormal();
 				TangentV[Tri[Axis]] = (Normal ^ TangentU[Tri[Axis]]).GetSafeNormal();
 			}
