Fix headlessChaos

diff --git a/Engine/Source/Programs/HeadlessChaos/Private/HeadlessChaosTestEPA.cpp b/Engine/Source/Programs/HeadlessChaos/Private/HeadlessChaosTestEPA.cpp
--- a/Engine/Source/Programs/HeadlessChaos/Private/HeadlessChaosTestEPA.cpp
+++ b/Engine/Source/Programs/HeadlessChaos/Private/HeadlessChaosTestEPA.cpp
@@ -683,7 +683,7 @@ namespace ChaosTest
 
 			// Test used to pass the planes to FConvex, but this is not supported any more. Planes are derived from points.
 			TUniquePtr<FConvex> Convex = MakeUnique<FConvex>(SurfaceParticles, 0.0f);
-			TImplicitObjectScaled<FConvex> ScaledConvex(MakeSerializable(Convex), FVec3(1.0f), 0.0f);
+			TImplicitObjectScaled<FConvex> ScaledConvex(MakeSerializable(Convex), nullptr, FVec3(1.0f), 0.0f);
 
 			TSphere<FReal, 3> Sphere(FVec3(0.0f), 34.2120171);
 
@@ -783,7 +783,7 @@ namespace ChaosTest
 				// Planes will derived from the points now, and also faces are merged (not triangles any more)
 				FVec3 ConvexScale ={25,25,1};
 				TUniquePtr<FConvex> Convex = MakeUnique<FConvex>(SurfaceParticles, 0.0f);
-				TImplicitObjectScaled<FConvex> ScaledConvex(MakeSerializable(Convex),ConvexScale,0.0f);
+				TImplicitObjectScaled<FConvex> ScaledConvex(MakeSerializable(Convex), nullptr, ConvexScale,0.0f);
 
 				TBox<FReal,3> Box({-50.0000000,-60.0000000,-30.0000000},{50.0000000,60.0000000,30.0000000});
 
@@ -1054,7 +1054,7 @@ namespace ChaosTest
 			// Planes will derived from the points now, and also faces are merged (not triangles any more)
 			FVec3 GroundConvexScale = { 25,25,1 };
 			TUniquePtr<FConvex> GroundConvex = MakeUnique<FConvex>(GroundSurfaceParticles, 0.0f);
-			TImplicitObjectScaled<FConvex> ScaledGroundConvex(MakeSerializable(GroundConvex), GroundConvexScale, 0.0f);
+			TImplicitObjectScaled<FConvex> ScaledGroundConvex(MakeSerializable(GroundConvex), nullptr, GroundConvexScale, 0.0f);
 
 
 			// Test used to pass planes and verts to FConvex but this is not suported an more. 
@@ -1232,7 +1232,7 @@ namespace ChaosTest
 			Materials.Emplace(0);
 			Materials.Emplace(0);
 			TUniquePtr<FTriangleMeshImplicitObject> TriangleMesh = MakeUnique<FTriangleMeshImplicitObject>(MoveTemp(TrimeshParticles), MoveTemp(Indices), MoveTemp(Materials));
-			TImplicitObjectScaled<FTriangleMeshImplicitObject> ScaledTriangleMesh = TImplicitObjectScaled<FTriangleMeshImplicitObject>(MakeSerializable(TriangleMesh), FVec3(11.5, 11.5, 11.5));
+			TImplicitObjectScaled<FTriangleMeshImplicitObject> ScaledTriangleMesh = TImplicitObjectScaled<FTriangleMeshImplicitObject>(MakeSerializable(TriangleMesh), nullptr, FVec3(11.5, 11.5, 11.5));
 
 			FQuat Rotation0(0.00488796039, 0.00569311855, -0.000786740216, 0.999971569);
 			FQuat Rotation1(0.0117356628, -0.0108017093, -0.000888462295, 0.999872327);
@@ -1287,7 +1287,7 @@ namespace ChaosTest
 
 			// Wrapping in 1,1,1 scale is unnecessary, but this is technically what is happening when sweeping against scaled trimesh.
 			TUniquePtr<FCapsule> Capsule = MakeUnique<FCapsule>(FVec3(0, 0, -33), FVec3(0, 0, 33), 42);
-			TImplicitObjectScaled<FCapsule> ScaledCapsule = TImplicitObjectScaled<FCapsule>(MakeSerializable(Capsule), FVec3(1));
+			TImplicitObjectScaled<FCapsule> ScaledCapsule = TImplicitObjectScaled<FCapsule>(MakeSerializable(Capsule), nullptr, FVec3(1));
 
 
 			const FVec3 Dir(-0.102473199, 0.130887285, -0.986087084);
diff --git a/Engine/Source/Programs/HeadlessChaos/Private/HeadlessChaosTestGJK.cpp b/Engine/Source/Programs/HeadlessChaos/Private/HeadlessChaosTestGJK.cpp
--- a/Engine/Source/Programs/HeadlessChaos/Private/HeadlessChaosTestGJK.cpp
+++ b/Engine/Source/Programs/HeadlessChaos/Private/HeadlessChaosTestGJK.cpp
@@ -624,9 +624,9 @@ namespace ChaosTest
 	{
 		TSphere<FReal, 3> A(FVec3(10, 0, 0), 5);
 		TUniquePtr<TSphere<FReal, 3>> Sphere = MakeUnique<TSphere<FReal, 3>>(FVec3(4, 0, 0), 2);
-		TImplicitObjectScaled<TSphere<FReal, 3>> Unscaled(MakeSerializable(Sphere), FVec3(1));
-		TImplicitObjectScaled<TSphere<FReal, 3>> UniformScaled(MakeSerializable(Sphere), FVec3(2));
-		TImplicitObjectScaled<TSphere<FReal, 3>> NonUniformScaled(MakeSerializable(Sphere), FVec3(2,1,1));
+		TImplicitObjectScaled<TSphere<FReal, 3>> Unscaled(MakeSerializable(Sphere), nullptr, FVec3(1));
+		TImplicitObjectScaled<TSphere<FReal, 3>> UniformScaled(MakeSerializable(Sphere), nullptr, FVec3(2));
+		TImplicitObjectScaled<TSphere<FReal, 3>> NonUniformScaled(MakeSerializable(Sphere), nullptr, FVec3(2,1,1));
 
 		FVec3 InitialDirs[] = { FVec3(1,0,0), FVec3(-1,0,0), FVec3(0,1,0), FVec3(0,-1,0), FVec3(0,0,1), FVec3(0,0,-1) };
 
@@ -1007,9 +1007,9 @@ namespace ChaosTest
 	{
 		TSphere<FReal, 3> A(FVec3(10, 0, 0), 5);
 		TUniquePtr<TSphere<FReal, 3>> Sphere = MakeUnique<TSphere<FReal, 3>>(FVec3(0, 0, 0), 2);
-		TImplicitObjectScaled<TSphere<FReal, 3>> Unscaled(MakeSerializable(Sphere), FVec3(1));
-		TImplicitObjectScaled<TSphere<FReal, 3>> UniformScaled(MakeSerializable(Sphere), FVec3(2));
-		TImplicitObjectScaled<TSphere<FReal, 3>> NonUniformScaled(MakeSerializable(Sphere), FVec3(2, 1, 1));
+		TImplicitObjectScaled<TSphere<FReal, 3>> Unscaled(MakeSerializable(Sphere), nullptr, FVec3(1));
+		TImplicitObjectScaled<TSphere<FReal, 3>> UniformScaled(MakeSerializable(Sphere), nullptr, FVec3(2));
+		TImplicitObjectScaled<TSphere<FReal, 3>> NonUniformScaled(MakeSerializable(Sphere), nullptr, FVec3(2, 1, 1));
 
 		FVec3 InitialDirs[] = { FVec3(1,0,0), FVec3(-1,0,0), FVec3(0,1,0), FVec3(0,-1,0), FVec3(0,0,1), FVec3(0,0,-1) };
 
@@ -1402,7 +1402,7 @@ namespace ChaosTest
 
 			TUniquePtr<FConvex> UniqueConvex = MakeUnique<FConvex>(ConvexParticles, 0.0f);
 			TSerializablePtr<FConvex> AConv(UniqueConvex);
-			const TImplicitObjectScaled<FConvex> A(AConv,FVec3(1.0,1.0,1.0));
+			const TImplicitObjectScaled<FConvex> A(AConv, nullptr, FVec3(1.0,1.0,1.0));
 
 			const FVec3 Pt0(0.0,0.0,-33.0);
 			FVec3 Pt1 = Pt0;
@@ -1503,7 +1503,7 @@ namespace ChaosTest
 			
 			TUniquePtr<FConvex> UniqueConvex = MakeUnique<FConvex>(ConvexParticles, 0.0f);
 			TSerializablePtr<FConvex> AConv(UniqueConvex);
-			const TImplicitObjectScaled<FConvex> AConvScaled(AConv,FVec3(1.0,1.0,1.0));
+			const TImplicitObjectScaled<FConvex> AConvScaled(AConv, nullptr, FVec3(1.0,1.0,1.0));
 
 			FTriangle A(ConvexParticles[0],ConvexParticles[1],ConvexParticles[2]);
 			FTriangleRegister AReg(
@@ -1740,8 +1740,8 @@ namespace ChaosTest
 		const FReal Margin = 0.75f;
 
 		TUniquePtr<FImplicitConvex3> CoreConvexShapePtr = MakeUnique<FImplicitConvex3>(CoreShapeVerts, 0.0f);
-		const TImplicitObjectScaled<FImplicitConvex3> ShapeA(MakeSerializable(CoreConvexShapePtr), Scale, Margin);
-		const TImplicitObjectScaled<FImplicitConvex3> ShapeB(MakeSerializable(CoreConvexShapePtr), Scale, Margin);
+		const TImplicitObjectScaled<FImplicitConvex3> ShapeA(MakeSerializable(CoreConvexShapePtr), nullptr, Scale, Margin);
+		const TImplicitObjectScaled<FImplicitConvex3> ShapeB(MakeSerializable(CoreConvexShapePtr), nullptr, Scale, Margin);
 		const FRigidTransform3 TransformA(FVec3(0.000000000f, 0.000000000f, 182.378937f), FRotation3::FromElements(0.000000000f, 0.000000000f, 0.707106650f, 0.707106888f));	// Top
 		const FRigidTransform3 TransformB(FVec3(0.000000000f, 0.000000000f, 107.378944f), FRotation3::FromElements(0.000000000f, 0.000000000f, 0.000000000f, 1.00000000f));		// Bottom
 
@@ -1804,7 +1804,7 @@ namespace ChaosTest
 		};
 		TArray<FConvex::FVec3Type> ConvexVertices(MoveTemp(ConvexVerts));
 		TUniquePtr<FImplicitConvex3> CoreConvex = MakeUnique<FImplicitConvex3>(ConvexVertices, 0.0f);
-		const TImplicitObjectScaled<FImplicitConvex3> ScaledConvex(MakeSerializable(CoreConvex), FVec3(-1,1,1), 38.4000015);
+		const TImplicitObjectScaled<FImplicitConvex3> ScaledConvex(MakeSerializable(CoreConvex), nullptr, FVec3(-1,1,1), 38.4000015);
 		const TSphere<FReal, 3> Sphere(FVec3(0,0,0), 32);
 		const FRigidTransform3 StartTM(FVec3( -172.000000, -48.0000000, 52.0000000 ), FRotation3::FromIdentity());
 
diff --git a/Engine/Source/Programs/HeadlessChaos/Private/HeadlessChaosTestImplicits.cpp b/Engine/Source/Programs/HeadlessChaos/Private/HeadlessChaosTestImplicits.cpp
--- a/Engine/Source/Programs/HeadlessChaos/Private/HeadlessChaosTestImplicits.cpp
+++ b/Engine/Source/Programs/HeadlessChaos/Private/HeadlessChaosTestImplicits.cpp
@@ -781,15 +781,15 @@ namespace ChaosTest {
 	void ImplicitScaled()
 	{
 		TUniquePtr<TBox<FReal, 3>> UnitCube = MakeUnique<TBox<FReal, 3>>(FVec3(-1), FVec3(1));
-		TImplicitObjectScaled<TBox<FReal,3>> UnitUnscaled(MakeSerializable(UnitCube), FVec3(1));
+		TImplicitObjectScaled<TBox<FReal,3>> UnitUnscaled(MakeSerializable(UnitCube), nullptr, FVec3(1));
 		UnitImplicitObjectNormalsInternal(UnitUnscaled, FString("ImplicitTransformed()"));
 		UnitImplicitObjectNormalsExternal(UnitUnscaled, FString("ImplicitTransformed()"));
 		UnitImplicitObjectIntersections(UnitUnscaled, FString("ImplicitTransformed()"));
 
 		TUniquePtr<TSphere<FReal, 3>> Sphere = MakeUnique<TSphere<FReal, 3>>(FVec3(3, 0, 0), 5);
-		TImplicitObjectScaled<TSphere<FReal, 3>> Unscaled(MakeSerializable(Sphere), FVec3(1));
-		TImplicitObjectScaled<TSphere<FReal, 3>> UniformScale(MakeSerializable(Sphere), FVec3(2));
-		TImplicitObjectScaled<TSphere<FReal, 3>> NonUniformScale(MakeSerializable(Sphere), FVec3(2, 1, 1));
+		TImplicitObjectScaled<TSphere<FReal, 3>> Unscaled(MakeSerializable(Sphere), nullptr, FVec3(1));
+		TImplicitObjectScaled<TSphere<FReal, 3>> UniformScale(MakeSerializable(Sphere), nullptr, FVec3(2));
+		TImplicitObjectScaled<TSphere<FReal, 3>> NonUniformScale(MakeSerializable(Sphere), nullptr, FVec3(2, 1, 1));
 
 		{//phi
 			const FVec3 NearEdge(7.5, 0, 0);
@@ -1807,12 +1807,12 @@ namespace ChaosTest {
 		// ImplicitObjectScaled is ignored.
 		FReal Thickness = 0.1;
 		TUniquePtr<TSphere<FReal, 3>> Sphere = MakeUnique<TSphere<FReal,3>>(FVec3(3, 0, 0), 5);
-		TImplicitObjectScaled<TSphere<FReal, 3>> Unscaled(MakeSerializable(Sphere), FVec3(1));
-		TImplicitObjectScaled<TSphere<FReal, 3>> UnscaledThickened(MakeSerializable(Sphere), FVec3(1), Thickness);
-		TImplicitObjectScaled<TSphere<FReal, 3>> UniformScale(MakeSerializable(Sphere),FVec3(2));
-		TImplicitObjectScaled<TSphere<FReal, 3>> UniformScaleThickened(MakeSerializable(Sphere), FVec3(2), Thickness);
-		TImplicitObjectScaled<TSphere<FReal, 3>> NonUniformScale(MakeSerializable(Sphere), FVec3(2, 1, 1));
-		TImplicitObjectScaled<TSphere<FReal, 3>> NonUniformScaleThickened(MakeSerializable(Sphere), FVec3(2, 1, 1), Thickness);
+		TImplicitObjectScaled<TSphere<FReal, 3>> Unscaled(MakeSerializable(Sphere), nullptr, FVec3(1));
+		TImplicitObjectScaled<TSphere<FReal, 3>> UnscaledThickened(MakeSerializable(Sphere), nullptr, FVec3(1), Thickness);
+		TImplicitObjectScaled<TSphere<FReal, 3>> UniformScale(MakeSerializable(Sphere), nullptr, FVec3(2));
+		TImplicitObjectScaled<TSphere<FReal, 3>> UniformScaleThickened(MakeSerializable(Sphere), nullptr, FVec3(2), Thickness);
+		TImplicitObjectScaled<TSphere<FReal, 3>> NonUniformScale(MakeSerializable(Sphere), nullptr, FVec3(2, 1, 1));
+		TImplicitObjectScaled<TSphere<FReal, 3>> NonUniformScaleThickened(MakeSerializable(Sphere), nullptr, FVec3(2, 1, 1), Thickness);
 
 		//phi
 		{
diff --git a/Engine/Source/Programs/HeadlessChaos/Private/HeadlessChaosTestMostOpposing.cpp b/Engine/Source/Programs/HeadlessChaos/Private/HeadlessChaosTestMostOpposing.cpp
--- a/Engine/Source/Programs/HeadlessChaos/Private/HeadlessChaosTestMostOpposing.cpp
+++ b/Engine/Source/Programs/HeadlessChaos/Private/HeadlessChaosTestMostOpposing.cpp
@@ -144,7 +144,7 @@ namespace ChaosTest
 
 		//identity scale
 		{
-			TImplicitObjectScaledGeneric<FReal, 3> Scaled(MakeSerializable(Convex), FVec3(1, 1, 1));
+			TImplicitObjectScaledGeneric<FReal, 3> Scaled(MakeSerializable(Convex), nullptr, FVec3(1, 1, 1));
 
 			//simple into the triangle
 			bool bHit = Scaled.Raycast(FVec3(0.5, 0, 2), FVec3(0, 0, -1), 3, 0, Time, Position, Normal, FaceIndex);
@@ -180,7 +180,7 @@ namespace ChaosTest
 
 		//non-uniform scale
 		{
-			TImplicitObjectScaledGeneric<FReal, 3> Scaled(MakeSerializable(Convex), FVec3(2, 1, 1));
+			TImplicitObjectScaledGeneric<FReal, 3> Scaled(MakeSerializable(Convex), nullptr, FVec3(2, 1, 1));
 
 			//simple into the triangle
 			bool bHit = Scaled.Raycast(FVec3(0.5, 0, 2), FVec3(0, 0, -1), 3, 0, Time, Position, Normal, FaceIndex);
diff --git a/Engine/Source/Programs/HeadlessChaos/Private/HeadlessChaosTestRaycast.cpp b/Engine/Source/Programs/HeadlessChaos/Private/HeadlessChaosTestRaycast.cpp
--- a/Engine/Source/Programs/HeadlessChaos/Private/HeadlessChaosTestRaycast.cpp
+++ b/Engine/Source/Programs/HeadlessChaos/Private/HeadlessChaosTestRaycast.cpp
@@ -555,12 +555,12 @@ namespace ChaosTest
 		const FReal Thickness = 0.1;
 
 		TUniquePtr<TSphere<FReal, 3>> Sphere = MakeUnique<TSphere<FReal,3>>(FVec3(1), 2);
-		TImplicitObjectScaled<TSphere<FReal, 3>> Unscaled(MakeSerializable(Sphere), FVec3(1));
-		TImplicitObjectScaled<TSphere<FReal, 3>> UnscaledThickened(MakeSerializable(Sphere), FVec3(1), Thickness);
-		TImplicitObjectScaled<TSphere<FReal, 3>> UniformScaled(MakeSerializable(Sphere), FVec3(2));
-		TImplicitObjectScaled<TSphere<FReal, 3>> UniformScaledThickened(MakeSerializable(Sphere), FVec3(2), Thickness);
-		TImplicitObjectScaled<TSphere<FReal, 3>> NonUniformScaled(MakeSerializable(Sphere), FVec3(2,1,1));
-		TImplicitObjectScaled<TSphere<FReal, 3>> NonUniformScaledThickened(MakeSerializable(Sphere), FVec3(2, 1, 1), Thickness);
+		TImplicitObjectScaled<TSphere<FReal, 3>> Unscaled(MakeSerializable(Sphere), nullptr, FVec3(1));
+		TImplicitObjectScaled<TSphere<FReal, 3>> UnscaledThickened(MakeSerializable(Sphere), nullptr, FVec3(1), Thickness);
+		TImplicitObjectScaled<TSphere<FReal, 3>> UniformScaled(MakeSerializable(Sphere), nullptr, FVec3(2));
+		TImplicitObjectScaled<TSphere<FReal, 3>> UniformScaledThickened(MakeSerializable(Sphere), nullptr, FVec3(2), Thickness);
+		TImplicitObjectScaled<TSphere<FReal, 3>> NonUniformScaled(MakeSerializable(Sphere), nullptr, FVec3(2,1,1));
+		TImplicitObjectScaled<TSphere<FReal, 3>> NonUniformScaledThickened(MakeSerializable(Sphere), nullptr, FVec3(2, 1, 1), Thickness);
 
 		//simple
 		bool bHit = Unscaled.Raycast(FVec3(1, 1, 8), FVec3(0, 0, -1), 8, 0, Time, Position, Normal, FaceIndex);
diff --git a/Engine/Source/Programs/HeadlessChaos/Private/HeadlessChaosTestRewind.cpp b/Engine/Source/Programs/HeadlessChaos/Private/HeadlessChaosTestRewind.cpp
--- a/Engine/Source/Programs/HeadlessChaos/Private/HeadlessChaosTestRewind.cpp
+++ b/Engine/Source/Programs/HeadlessChaos/Private/HeadlessChaosTestRewind.cpp
@@ -4258,8 +4258,8 @@ namespace ChaosTest {
 				Dynamic.SetGravityEnabled(true);
 				Dynamic.SetParticleID(FParticleID{ MappingHelper(Idx), INDEX_NONE });
 				Dynamic.SetX(FVec3(0, 5 * MappingHelper(Idx), 100 * MappingHelper(Idx) + 50));	//slightly offset
-				Dynamic.SetI(FMatrix33(1000, 1000, 1000));
-				Dynamic.SetInvI(FMatrix33(1/1000.0, 1/1000.0, 1/1000.0));
+				Dynamic.SetI(FVec3(1000, 1000, 1000));
+				Dynamic.SetInvI(FVec3(1/1000.0, 1/1000.0, 1/1000.0));
 
 				Storage.Add(DynamicProxy);
 			}
diff --git a/Engine/Source/Programs/HeadlessChaos/Private/HeadlessChaosTestSweep.cpp b/Engine/Source/Programs/HeadlessChaos/Private/HeadlessChaosTestSweep.cpp
--- a/Engine/Source/Programs/HeadlessChaos/Private/HeadlessChaosTestSweep.cpp
+++ b/Engine/Source/Programs/HeadlessChaos/Private/HeadlessChaosTestSweep.cpp
@@ -853,7 +853,7 @@ namespace ChaosTest
 		}
 
 		TUniquePtr<FTriangleMeshImplicitObject> TriangleMesh = MakeUnique<FTriangleMeshImplicitObject>(MoveTemp(TrimeshParticles), MoveTemp(Indices), MoveTemp(Materials));
-		TImplicitObjectScaled<FTriangleMeshImplicitObject> ScaledTriangleMesh = TImplicitObjectScaled<FTriangleMeshImplicitObject>(MakeSerializable(TriangleMesh), FVec3(50,50,50));
+		TImplicitObjectScaled<FTriangleMeshImplicitObject> ScaledTriangleMesh = TImplicitObjectScaled<FTriangleMeshImplicitObject>(MakeSerializable(TriangleMesh), nullptr, FVec3(50,50,50));
 
 		const FVec3 X1 = { 0,0,-19.45 };
 		const FVec3 X2 = X1 + FVec3(0, 0, 38.9);
