Fix AABBTree UseGrid SweepIntersectCells

diff --git a/Engine/Source/Programs/HeadlessChaos/Private/HeadlessChaosApp.cpp b/Engine/Source/Programs/HeadlessChaos/Private/HeadlessChaosApp.cpp
--- a/Engine/Source/Programs/HeadlessChaos/Private/HeadlessChaosApp.cpp
+++ b/Engine/Source/Programs/HeadlessChaos/Private/HeadlessChaosApp.cpp
@@ -240,6 +240,7 @@ TEST(BP, BroadphaseTests) {
 	ChaosTest::AABBTreeTestDynamic();
 	ChaosTest::AABBTreeDirtyGridTest();
 	ChaosTest::AABBTreeTimesliceTest();
+	ChaosTest::DoForSweepIntersectCellsImpTest();
 	ChaosTest::BoundingVolumeNoBoundsTest();
 	ChaosTest::BroadphaseCollectionTest();
 	SUCCEED();
diff --git a/Engine/Source/Programs/HeadlessChaos/Private/HeadlessChaosTestBroadphase.cpp b/Engine/Source/Programs/HeadlessChaos/Private/HeadlessChaosTestBroadphase.cpp
--- a/Engine/Source/Programs/HeadlessChaos/Private/HeadlessChaosTestBroadphase.cpp
+++ b/Engine/Source/Programs/HeadlessChaos/Private/HeadlessChaosTestBroadphase.cpp
@@ -716,6 +716,89 @@ namespace ChaosTest
 		FAABBTreeDirtyGridCVars::DirtyElementMaxCellCapacity = DirtyElementMaxCellCapacity;
 	}
 	
+	void DoForSweepIntersectCellsImpTest()
+	{
+		{
+			int32 NumFuncCalled = 0;
+			int32  XArray[2];
+			int32  YArray[2];
+
+			DoForSweepIntersectCellsImp(1.4484817992026819, 1.4432470701435705, 1251.1886035677471, -1183.6311465697545, -866.67708504199993, -747.83750413730752, 1000.0, 0.001,
+				[&](auto X, auto Y) {
+					XArray[NumFuncCalled] = X;
+					YArray[NumFuncCalled] = Y;
+					++NumFuncCalled;
+
+				});
+
+			EXPECT_EQ(NumFuncCalled, 2);
+			EXPECT_EQ(XArray[0], 1000);
+			EXPECT_EQ(YArray[0], -2000);
+
+			EXPECT_EQ(XArray[1], 0);
+			EXPECT_EQ(YArray[1], -2000);
+		}
+
+		{
+			int32 NumFuncCalled = 0;
+			int32  XArray[2];
+			int32  YArray[2];
+
+			DoForSweepIntersectCellsImp(1.4484817992026819, 1.4432470701435705, 1251.1886035677471, -1183.6311465697545, 866.67708504199993, -747.83750413730752, 1000.0, 0.001,
+				[&](auto X, auto Y) {
+					XArray[NumFuncCalled] = X;
+					YArray[NumFuncCalled] = Y;
+					++NumFuncCalled;
+
+				});
+
+			EXPECT_EQ(NumFuncCalled, 2);
+			EXPECT_EQ(XArray[0], 1000);
+			EXPECT_EQ(YArray[0], -2000);
+
+			EXPECT_EQ(XArray[1], 2000);
+			EXPECT_EQ(YArray[1], -2000);
+		}
+
+		{
+			int32 NumFuncCalled = 0;
+			int32  XArray[3];
+			int32  YArray[3];
+
+			DoForSweepIntersectCellsImp(1.2928878353696973, 1.2928878353697257, -1013.1421764369597, 210.55865232178132, 712.84350045280678, -265.39563631071809, 1000.0, 0.001,
+				[&](auto X, auto Y) {
+					XArray[NumFuncCalled] = X;
+					YArray[NumFuncCalled] = Y;
+					++NumFuncCalled;
+
+				});
+
+			EXPECT_EQ(NumFuncCalled, 3);
+			EXPECT_EQ(XArray[0], -2000);
+			EXPECT_EQ(YArray[0], 0);
+
+			EXPECT_EQ(XArray[1], -1000);
+			EXPECT_EQ(YArray[1], 0);
+
+			EXPECT_EQ(XArray[2], -1000);
+			EXPECT_EQ(YArray[2], -1000);
+
+		}
+		{
+			int32 NumFuncCalled = 0;
+			DoForSweepIntersectCellsImp(4000, 4000, 0, 0, 7000 - 0.01, 3000 - 0.01, 1000.0, 0.001,
+				[&](auto X, auto Y) {
+					++NumFuncCalled;
+
+				});
+
+			// This was verified manually on a paper grid
+			EXPECT_EQ(NumFuncCalled, 153);
+
+		}
+
+	}
+
 
 	void AABBTreeTimesliceTest()
 	{
diff --git a/Engine/Source/Programs/HeadlessChaos/Public/HeadlessChaosTestBroadphase.h b/Engine/Source/Programs/HeadlessChaos/Public/HeadlessChaosTestBroadphase.h
--- a/Engine/Source/Programs/HeadlessChaos/Public/HeadlessChaosTestBroadphase.h
+++ b/Engine/Source/Programs/HeadlessChaos/Public/HeadlessChaosTestBroadphase.h
@@ -19,6 +19,8 @@ namespace ChaosTest {
 
 	void AABBTreeTimesliceTest();
 
+	void DoForSweepIntersectCellsImpTest();
+
 	void BroadphaseCollectionTest();
 
 	void BoundingVolumeNoBoundsTest();
diff --git a/Engine/Source/Runtime/Experimental/Chaos/Public/Chaos/AABBTreeDirtyGridUtils.h b/Engine/Source/Runtime/Experimental/Chaos/Public/Chaos/AABBTreeDirtyGridUtils.h
--- a/Engine/Source/Runtime/Experimental/Chaos/Public/Chaos/AABBTreeDirtyGridUtils.h
+++ b/Engine/Source/Runtime/Experimental/Chaos/Public/Chaos/AABBTreeDirtyGridUtils.h
@@ -382,13 +382,12 @@ namespace Chaos
 					FReal CrossingHorizontalCellBorderT = std::numeric_limits<FReal>::max();
 					if (!DxTooSmall)
 					{
-
-						CrossingVerticleCellBorderT = DtDx * ((FReal)(CurrentCellIndexX0 - ReferenceCellIndexX + (DeltaCelIndexX > 0 ? 1 : 0)) * DirtyElementGridCellSize - X1);
+						CrossingVerticleCellBorderT = DtDx * ((FReal)(CurrentCellIndexX1 - ReferenceCellIndexX + (DeltaCelIndexX > 0 ? 1 : 0)) * DirtyElementGridCellSize - X1);
 					}
 
 					if (!DyTooSmall)
 					{
-						CrossingHorizontalCellBorderT = DtDy * ((FReal)(CurrentCellIndexY0 - ReferenceCellIndexY + (DeltaCelIndexY > 0 ? 1 : 0)) * DirtyElementGridCellSize - Y1);
+						CrossingHorizontalCellBorderT = DtDy * ((FReal)(CurrentCellIndexY1 - ReferenceCellIndexY + (DeltaCelIndexY > 0 ? 1 : 0)) * DirtyElementGridCellSize - Y1);
 					}
 
 					if (CrossingHorizontalCellBorderT < CrossingVerticleCellBorderT)
