Chaos - fix headless chaos

diff --git a/Engine/Source/Programs/HeadlessChaos/Private/HeadlessChaosTestConstraints.cpp b/Engine/Source/Programs/HeadlessChaos/Private/HeadlessChaosTestConstraints.cpp
--- a/Engine/Source/Programs/HeadlessChaos/Private/HeadlessChaosTestConstraints.cpp
+++ b/Engine/Source/Programs/HeadlessChaos/Private/HeadlessChaosTestConstraints.cpp
@@ -312,20 +312,23 @@ namespace ChaosTest {
 		const FRotation3& Rot = Evolution.GetParticleHandles().Handle(0)->R();
 		const float DeltaTime = 1.0f / 30.0f;
 
+		const float PositionTolerance = KINDA_SMALL_NUMBER;
+		const float RotationTolerance = SMALL_NUMBER;
+
 		for (int Iteration = 0; Iteration < 100; Iteration++)
 		{
 			Evolution.AdvanceOneTimeStep(DeltaTime);
 			Evolution.EndFrame(DeltaTime);
 
 			//UE_LOG(LogChaos, Warning, TEXT("Pos %s"), *Pos.ToString());
-			EXPECT_GT(Pos.Z, 12.f); // should never go past hard-stop
+			EXPECT_GT(Pos.Z, 12.f - PositionTolerance); // should never go past hard-stop
 		}
 
-		EXPECT_GT(Pos.Z, 12.f);	// suspension min limit
-		EXPECT_LT(Pos.Z, 15.f);	// suspension max limit
-		EXPECT_LT(Rot.X, SMALL_NUMBER);
-		EXPECT_LT(Rot.Y, SMALL_NUMBER);
-		EXPECT_LT(Rot.Z, SMALL_NUMBER);
+		EXPECT_GT(Pos.Z, 12.f - PositionTolerance);	// suspension min limit
+		EXPECT_LT(Pos.Z, 15.f + PositionTolerance);	// suspension max limit
+		EXPECT_LT(Rot.X, RotationTolerance);
+		EXPECT_LT(Rot.Y, RotationTolerance);
+		EXPECT_LT(Rot.Z, RotationTolerance);
 
 	}
 
diff --git a/Engine/Source/Runtime/Experimental/Chaos/Public/Chaos/Levelset.h b/Engine/Source/Runtime/Experimental/Chaos/Public/Chaos/Levelset.h
--- a/Engine/Source/Runtime/Experimental/Chaos/Public/Chaos/Levelset.h
+++ b/Engine/Source/Runtime/Experimental/Chaos/Public/Chaos/Levelset.h
@@ -36,6 +36,8 @@ class CHAOS_API FLevelSet final : public FImplicitObject
 	FLevelSet(FLevelSet&& Other);
 	virtual ~FLevelSet();
 
+	virtual TUniquePtr<FImplicitObject> Copy() const override { return DeepCopy(); }
+	virtual TUniquePtr<FImplicitObject> CopyWithScale(const FVec3& Scale) const { return DeepCopyWithScale(Scale); }
 	virtual TUniquePtr<FImplicitObject> DeepCopy() const;
 	virtual TUniquePtr<FImplicitObject> DeepCopyWithScale(const FVec3& Scale) const;
 
