Fix Vectorized Raycast unit test

diff --git a/Engine/Source/Programs/HeadlessChaos/Private/HeadlessChaosTestRaycast.cpp b/Engine/Source/Programs/HeadlessChaos/Private/HeadlessChaosTestRaycast.cpp
--- a/Engine/Source/Programs/HeadlessChaos/Private/HeadlessChaosTestRaycast.cpp
+++ b/Engine/Source/Programs/HeadlessChaos/Private/HeadlessChaosTestRaycast.cpp
@@ -546,9 +546,6 @@ namespace ChaosTest
 
 	void VectorizedAABBRaycast()
 	{
-		VectorRegister4Float Time;
-		VectorRegister4Float Position;
-
 		// AABB with 12cm thickness
 		FAABBVectorized AABB(MakeVectorRegisterFloat(565429.188-12.0, -17180.4355-12.0, -95264.4219-12.0, 0.f), MakeVectorRegisterFloat(568988.312+12.0, -13372.2793+12.0, -93649.4609+12.0, 0.f));
 
@@ -558,8 +555,8 @@ namespace ChaosTest
 		const VectorRegister4Float RayInvDir = MakeVectorRegisterFloat(16.0575714, -1.00243938, -31.8996048, 0.f);
 		const VectorRegister4Float RayLength = MakeVectorRegisterFloat(Length, Length, Length, Length);
 		const VectorRegister4Float RayInvLength = VectorDivide(MakeVectorRegisterFloatConstant(1.f, 1.f, 1.f, 1.f), RayLength);
-		constexpr bool bParallel[3] = { false, false, false };
-		bool bHit = AABB.RaycastFast(RayStart, RayDir, RayInvDir, bParallel, RayLength, RayInvLength, Time, Position);
+		const VectorRegister4Float Parallel = VectorZero();
+		bool bHit = AABB.RaycastFast(RayStart, RayInvDir, Parallel, RayLength);
 		EXPECT_FALSE(bHit);
 
 	}
