Fixed a non uniform scaling bug for meshes
UE-148004: Collision ImpactNormal is incorrect with scaled meshes

diff --git a/Engine/Source/Programs/HeadlessChaos/Private/HeadlessChaosTestMostOpposing.cpp b/Engine/Source/Programs/HeadlessChaos/Private/HeadlessChaosTestMostOpposing.cpp
--- a/Engine/Source/Programs/HeadlessChaos/Private/HeadlessChaosTestMostOpposing.cpp
+++ b/Engine/Source/Programs/HeadlessChaos/Private/HeadlessChaosTestMostOpposing.cpp
@@ -19,53 +19,128 @@ namespace ChaosTest
 
 	void TrimeshMostOpposing()
 	{
-		FReal Time;
-		FVec3 Position;
-		FVec3 Normal;
-		TArray<uint16> DummyMaterials;
-		int32 FaceIndex;
-
-		FTriangleMeshImplicitObject::ParticlesType Particles;
-		Particles.AddParticles(6);
-		Particles.X(0) = FVec3(1, 1, 1);
-		Particles.X(1) = FVec3(5, 1, 1);
-		Particles.X(2) = FVec3(1, 5, 1);
+		{
+			FReal Time;
+			FVec3 Position;
+			FVec3 Normal;
+			TArray<uint16> DummyMaterials;
+			int32 FaceIndex;
+
+			FTriangleMeshImplicitObject::ParticlesType Particles;
+			Particles.AddParticles(6);
+			Particles.X(0) = FVec3(1, 1, 1);
+			Particles.X(1) = FVec3(5, 1, 1);
+			Particles.X(2) = FVec3(1, 5, 1);
+
+			Particles.X(3) = FVec3(1, 1, 1);
+			Particles.X(4) = FVec3(1, 5, 1);
+			Particles.X(5) = FVec3(1, 1, -5);
+
+			TArray<TVec3<int32>> Indices;
+			Indices.Emplace(0, 1, 2);
+			Indices.Emplace(3, 4, 5);
+			FTriangleMeshImplicitObject Tri(MoveTemp(Particles), MoveTemp(Indices), MoveTemp(DummyMaterials));
 
-		Particles.X(3) = FVec3(1, 1, 1);
-		Particles.X(4) = FVec3(1, 5, 1);
-		Particles.X(5) = FVec3(1, 1, -5);
+			//simple into the triangle
+			bool bHit = Tri.Raycast(FVec3(3, 2, 2), FVec3(0, 0, -1), 2, 0, Time, Position, Normal, FaceIndex);
+			EXPECT_TRUE(bHit);
+			EXPECT_EQ(Tri.FindMostOpposingFace(Position, FVec3(0, 0, -1), FaceIndex, 0.01), 0);
+			EXPECT_EQ(Tri.GetFaceNormal(0).X, Normal.X);
+			EXPECT_EQ(Tri.GetFaceNormal(0).Y, Normal.Y);
+			EXPECT_EQ(Tri.GetFaceNormal(0).Z, Normal.Z);
 
-		TArray<TVec3<int32>> Indices;
-		Indices.Emplace(0, 1, 2);
-		Indices.Emplace(3, 4, 5);
-		FTriangleMeshImplicitObject Tri(MoveTemp(Particles), MoveTemp(Indices), MoveTemp(DummyMaterials));
+			//simple into second triangle
+			bHit = Tri.Raycast(FVec3(0, 2, 0), FVec3(1, 0, 0), 2, 0, Time, Position, Normal, FaceIndex);
+			EXPECT_TRUE(bHit);
+			EXPECT_EQ(FaceIndex, 1);
+			EXPECT_EQ(Tri.FindMostOpposingFace(Position, FVec3(1, 0, 0), FaceIndex, 0.01), 1);
+			const FVec3 FaceNormal = Tri.GetFaceNormal(1);
+			EXPECT_EQ(FaceNormal.X, Normal.X);
+			EXPECT_EQ(FaceNormal.Y, Normal.Y);
+			EXPECT_EQ(FaceNormal.Z, Normal.Z);
 
-		//simple into the triangle
-		bool bHit = Tri.Raycast(FVec3(3, 2, 2), FVec3(0, 0, -1), 2, 0, Time, Position, Normal, FaceIndex);
-		EXPECT_TRUE(bHit);
-		EXPECT_EQ(Tri.FindMostOpposingFace(Position, FVec3(0, 0, -1), FaceIndex, 0.01), 0);
-		EXPECT_EQ(Tri.GetFaceNormal(0).X, Normal.X);
-		EXPECT_EQ(Tri.GetFaceNormal(0).Y, Normal.Y);
-		EXPECT_EQ(Tri.GetFaceNormal(0).Z, Normal.Z);
+			//very close to edge, for now just return face hit regardless of direction because that's the implementation we currently rely on.
+			//todo: inconsistent with hulls, should make them the same, but may have significant impact on existing content
 
-		//simple into second triangle
-		bHit = Tri.Raycast(FVec3(0, 2, 0), FVec3(1, 0, 0), 2, 0, Time, Position, Normal, FaceIndex);
-		EXPECT_TRUE(bHit);
-		EXPECT_EQ(FaceIndex, 1);
-		EXPECT_EQ(Tri.FindMostOpposingFace(Position, FVec3(1, 0, 0), FaceIndex, 0.01), 1);
-		const FVec3 FaceNormal = Tri.GetFaceNormal(1);
-		EXPECT_EQ(FaceNormal.X, Normal.X);
-		EXPECT_EQ(FaceNormal.Y, Normal.Y);
-		EXPECT_EQ(FaceNormal.Z, Normal.Z);
+			bHit = Tri.Raycast(FVec3(0, 2, 0.9), FVec3(1, 0, 0), 2, 0, Time, Position, Normal, FaceIndex);
+			EXPECT_TRUE(bHit);
+			EXPECT_EQ(FaceIndex, 1);
+			EXPECT_EQ(Tri.FindMostOpposingFace(Position, FVec3(1, 0, 0), FaceIndex, 0.01), 1);
+			EXPECT_EQ(Tri.FindMostOpposingFace(Position, FVec3(0, 0, -1), FaceIndex, 0.01), 1);	//ignores direction completely as per current implementation
+		}
 
-		//very close to edge, for now just return face hit regardless of direction because that's the implementation we currently rely on.
-		//todo: inconsistent with hulls, should make them the same, but may have significant impact on existing content
+		// Non-uniform scale (Actual bug regression test)
+		{
+			// 4 triangles
+			FTriangleMeshImplicitObject::ParticlesType Particles;
+			Particles.AddParticles(12);
+			// in z-y plane
+			Particles.X(0) = FVec3(0, 0, 0);
+			Particles.X(1) = FVec3(0, 1, 0);
+			Particles.X(2) = FVec3(0, 0, 1);
+
+			// In x-z plane
+			Particles.X(3) = FVec3(0, 0, 0);
+			Particles.X(4) = FVec3(1, 0, 0);
+			Particles.X(5) = FVec3(0, 0, 1);
+
+			// In x-y plane
+			Particles.X(6) = FVec3(0, 0, 0);
+			Particles.X(7) = FVec3(1, 0, 0);
+			Particles.X(8) = FVec3(0, 1, 0);
+
+			// One 45 degree slanted triangle
+			Particles.X(9) = FVec3(1, 0, 0);
+			Particles.X(10) = FVec3(1, 1, 0);
+			Particles.X(11) = FVec3(0, 0, 1);
+
+			TArray<TVec3<int32>> Indices;
+			Indices.Emplace(0, 1, 2);
+			Indices.Emplace(3, 4, 5);
+			Indices.Emplace(6, 7, 8);
+			Indices.Emplace(9, 10, 11);
+			TArray<uint16> DummyMaterials;
+			TUniquePtr<FTriangleMeshImplicitObject> Tri = MakeUnique<FTriangleMeshImplicitObject>(MoveTemp(Particles), MoveTemp(Indices), MoveTemp(DummyMaterials));
+
+			// Using typical non uniform scale values
+			//const FVec3 Scale(10, 1, 0.1);
+			const FVec3 Scale(10, 1, 0.1);
+			TImplicitObjectScaledGeneric<FReal, 3> ScaledTri(MakeSerializable(Tri), nullptr, Scale);
+
+			FVec3 OpposeSlantedFace = (FVec3(-1, 0, -1)/Scale).GetSafeNormal(); // Note we are transforming a normal here.
+			EXPECT_NEAR(FVec3::DotProduct(OpposeSlantedFace, FVec3(Scale.Z, 0, Scale.X).GetSafeNormal()), -1.0f, KINDA_SMALL_NUMBER); // Check if OpposeSlantedFace opposes the slanted plane
+
+			// We should hit the slanted face, since we are directly opposing it
+			EXPECT_EQ(ScaledTri.FindMostOpposingFace(FVec3(0, 0, 0), OpposeSlantedFace, INDEX_NONE, 100.0), 3);
+
+			// Now modify the test vector just a bit 
+			FVec3 OpposeSlantedFaceNotExactlyOrthogonal = FVec3(OpposeSlantedFace.X*3.0f, OpposeSlantedFace.Y, OpposeSlantedFace.Z).GetSafeNormal();
+
+			// We should not get the triangle in the y-z plane (should still be the slanted plane)
+			// Check that the slanted triangle (3) is more apposing than the y-z triangle (0)
+			EXPECT_GT(FVec3::DotProduct(OpposeSlantedFaceNotExactlyOrthogonal, FVec3(-1, 0, 0)), FVec3::DotProduct(OpposeSlantedFaceNotExactlyOrthogonal, -OpposeSlantedFace));
+
+			EXPECT_NE(ScaledTri.FindMostOpposingFace(FVec3(0, 0, 0), OpposeSlantedFaceNotExactlyOrthogonal, INDEX_NONE, 100.0), 0);
+
+			// Now check mirroring cases:
+			FVec3 MirrorScale(1, -1, 1);
+			TImplicitObjectScaledGeneric<FReal, 3> MirTri(MakeSerializable(Tri), nullptr, MirrorScale);
+			EXPECT_EQ(MirTri.FindMostOpposingFace(FVec3(0, 0, 0), FVec3(0,0,-1), INDEX_NONE, 100.0), 2);
+
+			// Note: Switching to FindMostOpposingFaceScaled version instead of using wrapper class
+			MirrorScale = FVec3(-1, 1, 1);
+			EXPECT_EQ(Tri->FindMostOpposingFaceScaled(FVec3(0, 0, 0), FVec3(0, 0, -1), INDEX_NONE, 100.0, MirrorScale), 2);
+
+			MirrorScale = FVec3(1, 1, -1);
+			EXPECT_EQ(Tri->FindMostOpposingFaceScaled(FVec3(0, 0, 0), FVec3(0, 0, 1), INDEX_NONE, 100.0, MirrorScale), 2);
+
+			MirrorScale = FVec3(-1, -1, -1);
+			EXPECT_EQ(Tri->FindMostOpposingFaceScaled(FVec3(0, 0, 0), FVec3(0, 0, 1), INDEX_NONE, 100.0, MirrorScale), 2);
+
+			MirrorScale = FVec3(-1, -1, 0);
+			EXPECT_EQ(Tri->FindMostOpposingFaceScaled(FVec3(0, 0, 0), FVec3(0, 0, -1), INDEX_NONE, 100.0, MirrorScale), 2);		
 
-		bHit = Tri.Raycast(FVec3(0, 2, 0.9), FVec3(1, 0, 0), 2, 0, Time, Position, Normal, FaceIndex);
-		EXPECT_TRUE(bHit);
-		EXPECT_EQ(FaceIndex, 1);
-		EXPECT_EQ(Tri.FindMostOpposingFace(Position, FVec3(1, 0, 0), FaceIndex, 0.01), 1);
-		EXPECT_EQ(Tri.FindMostOpposingFace(Position, FVec3(0, 0, -1), FaceIndex, 0.01), 1);	//ignores direction completely as per current implementation
+		}
 	}
 
 
@@ -213,6 +288,13 @@ namespace ChaosTest
 			EXPECT_EQ(Scaled.FindMostOpposingFace(Position, FVec3(1, 0, 0), FaceIndex, 0.01), 3);
 			EXPECT_EQ(Scaled.FindMostOpposingFace(Position, FVec3(-1, 0, 0), FaceIndex, 0.01), 3);	//too far to care about other face
 		}
+
+		// Reflection
+		{
+			// Reflecting in x-z plane should return the same face as not reflecting (in this particular case)
+			FVec3 ReflectionScale = FVec3(1,-1,1);
+			EXPECT_EQ(Convex->FindMostOpposingFaceScaled(FVec3(0, 0, 0), FVec3(-2, 0, -1).GetSafeNormal(), INDEX_NONE, 1000.0f, ReflectionScale),2);
+		}
 	}
 
 
diff --git a/Engine/Source/Runtime/Experimental/Chaos/Private/Chaos/TriangleMeshImplicitObject.cpp b/Engine/Source/Runtime/Experimental/Chaos/Private/Chaos/TriangleMeshImplicitObject.cpp
--- a/Engine/Source/Runtime/Experimental/Chaos/Private/Chaos/TriangleMeshImplicitObject.cpp
+++ b/Engine/Source/Runtime/Experimental/Chaos/Private/Chaos/TriangleMeshImplicitObject.cpp
@@ -1216,7 +1216,7 @@ bool FTriangleMeshImplicitObject::SweepGeom(const TImplicitObjectScaled<FConvex>
 }
 
 template <typename IdxType>
-int32 FTriangleMeshImplicitObject::FindMostOpposingFace(const TArray<TVec3<IdxType>>& Elements, const FVec3& Position, const FVec3& UnitDir, int32 HintFaceIndex, FReal SearchDist) const
+int32 FTriangleMeshImplicitObject::FindMostOpposingFace(const TArray<TVec3<IdxType>>& Elements, const FVec3& Position, const FVec3& UnitDir, int32 HintFaceIndex, FReal SearchDist, const FVec3& Scale) const
 {
 	//todo: this is horribly slow, need adjacency information
 	const FReal SearchDist2 = SearchDist * SearchDist;
@@ -1231,9 +1231,9 @@ int32 FTriangleMeshImplicitObject::FindMostOpposingFace(const TArray<TVec3<IdxTy
 
 	for (int32 TriIdx : PotentialIntersections)
 	{
-		const FVec3& A = MParticles.X(Elements[TriIdx][0]);
-		const FVec3& B = MParticles.X(Elements[TriIdx][1]);
-		const FVec3& C = MParticles.X(Elements[TriIdx][2]);
+		const FVec3& A = MParticles.X(Elements[TriIdx][0]) * Scale;
+		const FVec3& B = MParticles.X(Elements[TriIdx][1]) * Scale;
+		const FVec3& C = MParticles.X(Elements[TriIdx][2]) * Scale;
 
 		const FVec3 AB = B - A;
 		const FVec3 AC = C - A;
@@ -1245,6 +1245,13 @@ int32 FTriangleMeshImplicitObject::FindMostOpposingFace(const TArray<TVec3<IdxTy
 			continue;
 		}
 
+		// Check if the scale is reflective
+		bool bNormalSignFlip = Scale.X * Scale.Y * Scale.Z < 0;
+		if (bNormalSignFlip)
+		{
+			Normal = -Normal;
+		}
+
 		const TPlane<FReal, 3> TriPlane{A, Normal};
 		const FVec3 ClosestPointOnTri = FindClosestPointOnTriangle(TriPlane, A, B, C, Position);
 		const FReal Distance2 = (ClosestPointOnTri - Position).SizeSquared();
@@ -1263,14 +1270,27 @@ int32 FTriangleMeshImplicitObject::FindMostOpposingFace(const TArray<TVec3<IdxTy
 }
 
 int32 FTriangleMeshImplicitObject::FindMostOpposingFace(const FVec3& Position, const FVec3& UnitDir, int32 HintFaceIndex, FReal SearchDist) const
+{
+	const FVec3 Scale(1.0f, 1.0f, 1.0f);
+	if (MElements.RequiresLargeIndices())
+	{
+		return FindMostOpposingFace(MElements.GetLargeIndexBuffer(), Position, UnitDir, HintFaceIndex, SearchDist, Scale);
+	}
+	else
+	{
+		return FindMostOpposingFace(MElements.GetSmallIndexBuffer(), Position, UnitDir, HintFaceIndex, SearchDist, Scale);
+	}
+}
+
+int32 FTriangleMeshImplicitObject::FindMostOpposingFaceScaled(const FVec3& Position, const FVec3& UnitDir, int32 HintFaceIndex, FReal SearchDist, const FVec3& Scale) const
 {
 	if (MElements.RequiresLargeIndices())
 	{
-		return FindMostOpposingFace(MElements.GetLargeIndexBuffer(), Position, UnitDir, HintFaceIndex, SearchDist);
+		return FindMostOpposingFace(MElements.GetLargeIndexBuffer(), Position, UnitDir, HintFaceIndex, SearchDist, Scale);
 	}
 	else
 	{
-		return FindMostOpposingFace(MElements.GetSmallIndexBuffer(), Position, UnitDir, HintFaceIndex, SearchDist);
+		return FindMostOpposingFace(MElements.GetSmallIndexBuffer(), Position, UnitDir, HintFaceIndex, SearchDist, Scale);
 	}
 }
 
diff --git a/Engine/Source/Runtime/Experimental/Chaos/Public/Chaos/ImplicitObject.h b/Engine/Source/Runtime/Experimental/Chaos/Public/Chaos/ImplicitObject.h
--- a/Engine/Source/Runtime/Experimental/Chaos/Public/Chaos/ImplicitObject.h
+++ b/Engine/Source/Runtime/Experimental/Chaos/Public/Chaos/ImplicitObject.h
@@ -286,7 +286,7 @@ public:
 	virtual int32 FindMostOpposingFaceScaled(const FVec3& Position, const FVec3& UnitDir, int32 HintFaceIndex, FReal SearchDist, const FVec3& Scale) const
 	{
 		// For now, this is the implementation that used to be in FImplicitObjectScaled to use as a default.
-		// @todo(chaos): implement FindMostOpposingFaceScaled for all types that can be wrapped in ImplicitObjectScaled
+		// @todo(chaos): implement FindMostOpposingFaceScaled for all types that can be wrapped in ImplicitObjectScaled (Since this default implementation won't work correctly for many cases)
 		// NOTE: this isn't strictly correct. The distance check does not account for non-uniforms scales, but also the algorithm is finding the
 		// face which has the most opposing normal. The scaled-space normal for a face i is Ns_i = (N_i / S) / |(N_i / S)|
 		// We want to find i with the minimum F_i = (D . Ns_i), where D is the unit direction we are interested in.
@@ -295,8 +295,11 @@ public:
 		// be ignored in the minimum check.
 		const FReal UnscaledSearchDist = SearchDist / Scale.Min();	//this is not quite right since it's no longer a sphere, but the whole thing is fuzzy anyway
 		const FVec3 UnscaledPosition = Position / Scale;
-		const FVec3 UnscaledNormal = GetInnerUnscaledNormal(UnitDir, Scale);
-		return FindMostOpposingFace(UnscaledPosition, UnscaledNormal, HintFaceIndex, UnscaledSearchDist);
+		//const FVec3 UnscaledDir = ScaleNormalizedHelper(UnitDir, FVec3(1.0f / Scale.X, 1.0f / Scale.Y, 1.0f / Scale.Z)); // If we want to transform like a vector
+		const FVec3 UnscaledDir = GetInnerUnscaledNormal(UnitDir, Scale); // If we want to transform like a Normal
+		// This will not work correctly for most cases (e.g. if we are comparing angles (dot product) after a non conformal transformation (like non uniform scaling))
+		// So really provide an implementation for FImplicitObjectScaled and avoid this default
+		return FindMostOpposingFace(UnscaledPosition, UnscaledDir, HintFaceIndex, UnscaledSearchDist);
 	}
 
 
@@ -374,8 +377,8 @@ public:
 
 protected:
 
-	// Safely scale a normal - used with both scale and inverse scale
-	FVec3 ScaledNormalHelper(const FVec3& Normal, const FVec3& Scale) const
+	// Safely scale a normalized Vec3 - used with both scale and inverse scale
+	static FVec3 ScaleNormalizedHelper(const FVec3& Normal, const FVec3& Scale)
 	{
 		const FVec3 ScaledNormal = Scale * Normal;
 		const FReal ScaledNormalLen = ScaledNormal.Size();
@@ -386,15 +389,15 @@ protected:
 
 	// Convert a normal in the inner unscaled object space into a normal in the outer scaled object space.
 	// Note: INverse scale is passed in, not the scale
-	FVec3 GetOuterScaledNormal(const FVec3& InnerNormal, const FVec3& Scale) const
+	static FVec3 GetOuterScaledNormal(const FVec3& InnerNormal, const FVec3& Scale)
 	{
-		return ScaledNormalHelper(InnerNormal, FVec3(1.0f / Scale.X, 1.0f / Scale.Y, 1.0f / Scale.Z));
+		return ScaleNormalizedHelper(InnerNormal, FVec3(1.0f / Scale.X, 1.0f / Scale.Y, 1.0f / Scale.Z));
 	}
 
 	// Convert a normal in the outer scaled object space into a normal in the inner unscaled object space
-	FVec3 GetInnerUnscaledNormal(const FVec3& OuterNormal, const FVec3& Scale) const
+	static FVec3 GetInnerUnscaledNormal(const FVec3& OuterNormal, const FVec3& Scale)
 	{
-		return ScaledNormalHelper(OuterNormal, Scale);
+		return ScaleNormalizedHelper(OuterNormal, Scale);
 	}
 
 	// Not all derived types support a margin, and for some it represents some other
diff --git a/Engine/Source/Runtime/Experimental/Chaos/Public/Chaos/TriangleMeshImplicitObject.h b/Engine/Source/Runtime/Experimental/Chaos/Public/Chaos/TriangleMeshImplicitObject.h
--- a/Engine/Source/Runtime/Experimental/Chaos/Public/Chaos/TriangleMeshImplicitObject.h
+++ b/Engine/Source/Runtime/Experimental/Chaos/Public/Chaos/TriangleMeshImplicitObject.h
@@ -533,6 +533,7 @@ namespace Chaos
 
 
 		virtual int32 FindMostOpposingFace(const FVec3& Position, const FVec3& UnitDir, int32 HintFaceIndex, FReal SearchDistance) const override;
+		virtual int32 FindMostOpposingFaceScaled(const FVec3& Position, const FVec3& UnitDir, int32 HintFaceIndex, FReal SearchDist, const FVec3& Scale) const override;
 		virtual FVec3 FindGeometryOpposingNormal(const FVec3& DenormDir, int32 FaceIndex, const FVec3& OriginalNormal) const override;
 
 		virtual const FAABB3 BoundingBox() const
@@ -812,7 +813,7 @@ namespace Chaos
 		bool OverlapImp(const TArray<TVec3<IdxType>>& Elements, const FVec3& Point, const FReal Thickness) const;
 
 		template<typename IdxType>
-		int32 FindMostOpposingFace(const TArray<TVec3<IdxType>>& Elements, const FVec3& Position, const FVec3& UnitDir, int32 HintFaceIndex, FReal SearchDist) const;
+		int32 FindMostOpposingFace(const TArray<TVec3<IdxType>>& Elements, const FVec3& Position, const FVec3& UnitDir, int32 HintFaceIndex, FReal SearchDist, const FVec3& Scale) const;
 
 		template <typename IdxType>
 		void RebuildBVImp(const TArray<TVec3<IdxType>>& Elements, BVHType& BVH);
diff --git a/Engine/Source/Runtime/PhysicsCore/Private/PhysicsInterfaceUtilsCore.cpp b/Engine/Source/Runtime/PhysicsCore/Private/PhysicsInterfaceUtilsCore.cpp
--- a/Engine/Source/Runtime/PhysicsCore/Private/PhysicsInterfaceUtilsCore.cpp
+++ b/Engine/Source/Runtime/PhysicsCore/Private/PhysicsInterfaceUtilsCore.cpp
@@ -300,8 +300,8 @@ uint32 FindFaceIndexImp(const THitLocation& PHit, const FVector& UnitDir)
 #else
 	const FTransform WorldTM(PHit.Actor->R(), PHit.Actor->X());
 	const FVector LocalPosition = WorldTM.InverseTransformPositionNoScale(PHit.WorldPosition);
-	const FVector LocalNormal = WorldTM.InverseTransformVectorNoScale(UnitDir);
-	return PHit.Shape->GetGeometry()->FindMostOpposingFace(LocalPosition, LocalNormal, PHit.FaceIndex, 1);	//todo:this number matches the one above, but is it right?
+	const FVector LocalUnitDir = WorldTM.InverseTransformVectorNoScale(UnitDir);
+	return PHit.Shape->GetGeometry()->FindMostOpposingFace(LocalPosition, LocalUnitDir, PHit.FaceIndex, 1);	//todo:this number matches the one above, but is it right?
 #endif
 }
 
