Chaos: enable the new convex generation algorithm by default
- fixed unit tests
- add a way to override the default build method if necessary

diff --git a/Engine/Source/Programs/HeadlessChaos/Private/HeadlessChaosTestGJK.cpp b/Engine/Source/Programs/HeadlessChaos/Private/HeadlessChaosTestGJK.cpp
--- a/Engine/Source/Programs/HeadlessChaos/Private/HeadlessChaosTestGJK.cpp
+++ b/Engine/Source/Programs/HeadlessChaos/Private/HeadlessChaosTestGJK.cpp
@@ -1770,7 +1770,7 @@ namespace ChaosTest
 		const FVec3 Scale = FVec3(50.0f);
 		const FReal Margin = 0.75f;
 
-		TUniquePtr<FImplicitConvex3> CoreConvexShapePtr = MakeUnique<FImplicitConvex3>(CoreShapeVerts, 0.0f);
+		TUniquePtr<FImplicitConvex3> CoreConvexShapePtr = MakeUnique<FImplicitConvex3>(CoreShapeVerts, 0.0f, FConvexBuilder::EBuildMethod::Original);
 		const TImplicitObjectScaled<FImplicitConvex3> ShapeA(MakeSerializable(CoreConvexShapePtr), nullptr, Scale, Margin);
 		const TImplicitObjectScaled<FImplicitConvex3> ShapeB(MakeSerializable(CoreConvexShapePtr), nullptr, Scale, Margin);
 		const FRigidTransform3 TransformA(FVec3(0.000000000f, 0.000000000f, 182.378937f), FRotation3::FromElements(0.000000000f, 0.000000000f, 0.707106650f, 0.707106888f));	// Top
diff --git a/Engine/Source/Programs/HeadlessChaos/Private/HeadlessChaosTestMostOpposing.cpp b/Engine/Source/Programs/HeadlessChaos/Private/HeadlessChaosTestMostOpposing.cpp
--- a/Engine/Source/Programs/HeadlessChaos/Private/HeadlessChaosTestMostOpposing.cpp
+++ b/Engine/Source/Programs/HeadlessChaos/Private/HeadlessChaosTestMostOpposing.cpp
@@ -171,7 +171,7 @@ namespace ChaosTest
 		EXPECT_FLOAT_EQ(Position.Y, 2);
 		EXPECT_FLOAT_EQ(Position.Z, 1);
 		EXPECT_EQ(FaceIndex, INDEX_NONE);	//convex should not compute its own face index as this is too expensive
-		EXPECT_EQ(Convex.FindMostOpposingFace(Position, FVec3(0, 0, -1), FaceIndex, 0.01 + SMALL_NUMBER), 1);	//front face, just so happens that convex hull generates the planes in this order
+		EXPECT_EQ(Convex.FindMostOpposingFace(Position, FVec3(0, 0, -1), FaceIndex, 0.01 + SMALL_NUMBER), 0);	//front face, just so happens that convex hull generates the planes in this order
 
 		//simple into second triangle
 		bHit = Convex.Raycast(FVec3(0, 2, 0), FVec3(1, 0, 0), 2, 0, Time, Position, Normal, FaceIndex);
@@ -187,7 +187,7 @@ namespace ChaosTest
 		EXPECT_TRUE(bHit);
 		EXPECT_EQ(FaceIndex, INDEX_NONE);
 		EXPECT_EQ(Convex.FindMostOpposingFace(Position, FVec3(1, 0, 0), FaceIndex, 0.01 + SMALL_NUMBER), 3);
-		EXPECT_EQ(Convex.FindMostOpposingFace(Position, FVec3(0, 0, -1), FaceIndex, 0.01 + SMALL_NUMBER), 1);
+		EXPECT_EQ(Convex.FindMostOpposingFace(Position, FVec3(0, 0, -1), FaceIndex, 0.01 + SMALL_NUMBER), 0);
 
 		//again but far enough away from edge
 		bHit = Convex.Raycast(FVec3(0, 2, 0.9), FVec3(1, 0, 0), 2, 0, Time, Position, Normal, FaceIndex);
@@ -228,7 +228,7 @@ namespace ChaosTest
 			EXPECT_EQ(Position.Y, 0);
 			EXPECT_EQ(Position.Z, 2 - Time);
 			EXPECT_EQ(FaceIndex, INDEX_NONE);	//convex should not compute its own face index as this is too expensive
-			EXPECT_EQ(Scaled.FindMostOpposingFace(Position, FVec3(0, 0, -1), FaceIndex, 0.01), 2);	//x+ face, just so happens that convex hull generates the planes in this order
+			EXPECT_EQ(Scaled.FindMostOpposingFace(Position, FVec3(0, 0, -1), FaceIndex, 0.01), 0);	//x+ face, just so happens that convex hull generates the planes in this order
 
 			//simple into second triangle
 			bHit = Scaled.Raycast(FVec3(-2, 0, 0.5), FVec3(1, 0, 0), 3, 0, Time, Position, Normal, FaceIndex);
@@ -237,20 +237,20 @@ namespace ChaosTest
 			EXPECT_EQ(Position.Y, 0);
 			EXPECT_EQ(Position.Z, 0.5);
 			EXPECT_EQ(FaceIndex, INDEX_NONE);	//convex should not compute its own face index as this is too expensive
-			EXPECT_EQ(Scaled.FindMostOpposingFace(Position, FVec3(1, 0, 0), FaceIndex, 0.01), 3);	//x- face, just so happens that convex hull generates the planes in this order
+			EXPECT_EQ(Scaled.FindMostOpposingFace(Position, FVec3(1, 0, 0), FaceIndex, 0.01), 1);	//x- face, just so happens that convex hull generates the planes in this order
 
 			bHit = Scaled.Raycast(FVec3(-0.001, 0, 2), FVec3(0, 0, -1), 3, 0, Time, Position, Normal, FaceIndex);
 			EXPECT_TRUE(bHit);
 			EXPECT_EQ(FaceIndex, INDEX_NONE);
-			EXPECT_EQ(Scaled.FindMostOpposingFace(Position, FVec3(1, 0, 0), FaceIndex, 0.01), 3);
-			EXPECT_EQ(Scaled.FindMostOpposingFace(Position, FVec3(-1, 0, 0), FaceIndex, 0.01), 2);
+			EXPECT_EQ(Scaled.FindMostOpposingFace(Position, FVec3(1, 0, 0), FaceIndex, 0.01), 1);
+			EXPECT_EQ(Scaled.FindMostOpposingFace(Position, FVec3(-1, 0, 0), FaceIndex, 0.01), 0);
 
 			//again but far enough away from edge
 			bHit = Scaled.Raycast(FVec3(-0.1, 0, 2), FVec3(0, 0, -1), 3, 0, Time, Position, Normal, FaceIndex);
 			EXPECT_TRUE(bHit);
 			EXPECT_EQ(FaceIndex, INDEX_NONE);
-			EXPECT_EQ(Scaled.FindMostOpposingFace(Position, FVec3(1, 0, 0), FaceIndex, 0.01), 3);
-			EXPECT_EQ(Scaled.FindMostOpposingFace(Position, FVec3(-1, 0, 0), FaceIndex, 0.01), 3);	//too far to care about other face
+			EXPECT_EQ(Scaled.FindMostOpposingFace(Position, FVec3(1, 0, 0), FaceIndex, 0.01), 1);
+			EXPECT_EQ(Scaled.FindMostOpposingFace(Position, FVec3(-1, 0, 0), FaceIndex, 0.01), 1);	//too far to care about other face
 		}
 
 		//non-uniform scale
@@ -264,7 +264,7 @@ namespace ChaosTest
 			EXPECT_EQ(Position.Y, 0);
 			EXPECT_EQ(Position.Z, 2 - Time);
 			EXPECT_EQ(FaceIndex, INDEX_NONE);	//convex should not compute its own face index as this is too expensive
-			EXPECT_EQ(Scaled.FindMostOpposingFace(Position, FVec3(0, 0, -1), FaceIndex, 0.01), 2);	//x+ face, just so happens that convex hull generates the planes in this order
+			EXPECT_EQ(Scaled.FindMostOpposingFace(Position, FVec3(0, 0, -1), FaceIndex, 0.01), 0);	//x+ face, just so happens that convex hull generates the planes in this order
 
 			//simple into second triangle
 			bHit = Scaled.Raycast(FVec3(-2, 0, 0.5), FVec3(1, 0, 0), 3, 0, Time, Position, Normal, FaceIndex);
@@ -273,20 +273,20 @@ namespace ChaosTest
 			EXPECT_EQ(Position.Y, 0);
 			EXPECT_EQ(Position.Z, 0.5);
 			EXPECT_EQ(FaceIndex, INDEX_NONE);	//convex should not compute its own face index as this is too expensive
-			EXPECT_EQ(Scaled.FindMostOpposingFace(Position, FVec3(1, 0, 0), FaceIndex, 0.01), 3);	//x- face, just so happens that convex hull generates the planes in this order
+			EXPECT_EQ(Scaled.FindMostOpposingFace(Position, FVec3(1, 0, 0), FaceIndex, 0.01), 1);	//x- face, just so happens that convex hull generates the planes in this order
 
 			bHit = Scaled.Raycast(FVec3(-0.001, 0, 2), FVec3(0, 0, -1), 3, 0, Time, Position, Normal, FaceIndex);
 			EXPECT_TRUE(bHit);
 			EXPECT_EQ(FaceIndex, INDEX_NONE);
-			EXPECT_EQ(Scaled.FindMostOpposingFace(Position, FVec3(1, 0, 0), FaceIndex, 0.01), 3);
-			EXPECT_EQ(Scaled.FindMostOpposingFace(Position, FVec3(-1, 0, 0), FaceIndex, 0.01), 2);
+			EXPECT_EQ(Scaled.FindMostOpposingFace(Position, FVec3(1, 0, 0), FaceIndex, 0.01), 1);
+			EXPECT_EQ(Scaled.FindMostOpposingFace(Position, FVec3(-1, 0, 0), FaceIndex, 0.01), 0);
 
 			//again but far enough away from edge
 			bHit = Scaled.Raycast(FVec3(-0.1, 0, 2), FVec3(0, 0, -1), 3, 0, Time, Position, Normal, FaceIndex);
 			EXPECT_TRUE(bHit);
 			EXPECT_EQ(FaceIndex, INDEX_NONE);
-			EXPECT_EQ(Scaled.FindMostOpposingFace(Position, FVec3(1, 0, 0), FaceIndex, 0.01), 3);
-			EXPECT_EQ(Scaled.FindMostOpposingFace(Position, FVec3(-1, 0, 0), FaceIndex, 0.01), 3);	//too far to care about other face
+			EXPECT_EQ(Scaled.FindMostOpposingFace(Position, FVec3(1, 0, 0), FaceIndex, 0.01), 1);
+			EXPECT_EQ(Scaled.FindMostOpposingFace(Position, FVec3(-1, 0, 0), FaceIndex, 0.01), 1);	//too far to care about other face
 		}
 
 		// Reflection
diff --git a/Engine/Source/Runtime/Experimental/Chaos/Private/Chaos/CollisionConvexMesh.cpp b/Engine/Source/Runtime/Experimental/Chaos/Private/Chaos/CollisionConvexMesh.cpp
--- a/Engine/Source/Runtime/Experimental/Chaos/Private/Chaos/CollisionConvexMesh.cpp
+++ b/Engine/Source/Runtime/Experimental/Chaos/Private/Chaos/CollisionConvexMesh.cpp
@@ -15,7 +15,7 @@ namespace Chaos
 
 	int32 FConvexBuilder::ComputeHorizonEpsilonFromMeshExtends = 1;
 
-	bool FConvexBuilder::bUseGeometryTConvexHull3 = false;
+	bool FConvexBuilder::bUseGeometryTConvexHull3 = true;
 
 	FAutoConsoleVariableRef CVarConvexGeometryCheckEnable(TEXT("p.Chaos.ConvexGeometryCheckEnable"), FConvexBuilder::PerformGeometryCheck, TEXT("Perform convex geometry complexity check for Chaos physics."));
 
@@ -23,11 +23,19 @@ namespace Chaos
 
 	FAutoConsoleVariableRef CVarConvexParticlesWarningThreshold(TEXT("p.Chaos.ConvexParticlesWarningThreshold"), FConvexBuilder::VerticesThreshold, TEXT("Threshold beyond which we warn about collision geometry complexity."));
 	
-	FAutoConsoleVariableRef CvarUseGeometryTConvexHull3(TEXT("p.Chaos.Convex.UseTConvexHull3Builder"), FConvexBuilder::bUseGeometryTConvexHull3, TEXT("Use the newer Geometry Tools code path for generating convex hulls.[def:false]"));
+	FAutoConsoleVariableRef CvarUseGeometryTConvexHull3(TEXT("p.Chaos.Convex.UseTConvexHull3Builder"), FConvexBuilder::bUseGeometryTConvexHull3, TEXT("Use the newer Geometry Tools code path for generating convex hulls when default build method is set.[def:true]"));
 
 
+	bool FConvexBuilder::UseConvexHull3(FConvexBuilder::EBuildMethod BuildMethod)
+	{
+		if (BuildMethod == EBuildMethod::Default && FConvexBuilder::bUseGeometryTConvexHull3)
+		{
+			return true;
+		}
+		return BuildMethod == EBuildMethod::ConvexHull3;
+	}
 
-	void FConvexBuilder::Build(const TArray<FVec3Type>& InVertices, TArray <FPlaneType>& OutPlanes, TArray<TArray<int32>>& OutFaceIndices, TArray<FVec3Type>& OutVertices, FAABB3Type& OutLocalBounds)
+	void FConvexBuilder::Build(const TArray<FVec3Type>& InVertices, TArray <FPlaneType>& OutPlanes, TArray<TArray<int32>>& OutFaceIndices, TArray<FVec3Type>& OutVertices, FAABB3Type& OutLocalBounds, EBuildMethod BuildMethod)
 	{
 		OutPlanes.Reset();
 		OutVertices.Reset();
@@ -91,12 +99,7 @@ namespace Chaos
 
 		if (NumVerticesToUse >= 4)
 		{
-			// @todo(chaos) : Experimental path for convex generation. This is still under
-			//                testing, and will eventually replace the convex generation
-			//                in chaos. Currenty defaults to off, but can be enabled with 
-			//                the cvar : p.Chaos.Convex.UseTConvexHull3Builder
-
-			if (bUseGeometryTConvexHull3) // Use the newer Geometry Tools code path for generating convex hulls.
+			if (UseConvexHull3(BuildMethod)) // Use the newer Geometry Tools code path for generating convex hulls.
 			{
 				UE::Geometry::TConvexHull3<FRealType> HullCompute;
 				if (HullCompute.Solve<FVec3Type>(*VerticesToUse))
diff --git a/Engine/Source/Runtime/Experimental/Chaos/Public/Chaos/CollisionConvexMesh.h b/Engine/Source/Runtime/Experimental/Chaos/Public/Chaos/CollisionConvexMesh.h
--- a/Engine/Source/Runtime/Experimental/Chaos/Public/Chaos/CollisionConvexMesh.h
+++ b/Engine/Source/Runtime/Experimental/Chaos/Public/Chaos/CollisionConvexMesh.h
@@ -40,6 +40,13 @@ namespace Chaos
 
 	public:
 
+		enum EBuildMethod
+		{
+			Default = 0,	// use what is set as default from cvars
+			Original = 1,	// original method, fast but has issues
+			ConvexHull3 = 2,// new method, slower but less issues
+		};
+
 		class Params
 		{
 		public:
@@ -252,8 +259,10 @@ namespace Chaos
 
 	public:
 
-		static CHAOS_API void Build(const TArray<FVec3Type>& InVertices, TArray<FPlaneType>& OutPlanes, TArray<TArray<int32>>& OutFaceIndices, TArray<FVec3Type>& OutVertices, FAABB3Type& OutLocalBounds);
+		static CHAOS_API void Build(const TArray<FVec3Type>& InVertices, TArray<FPlaneType>& OutPlanes, TArray<TArray<int32>>& OutFaceIndices, TArray<FVec3Type>& OutVertices, FAABB3Type& OutLocalBounds, EBuildMethod BuildMethod = EBuildMethod::Default);
 
+		static bool UseConvexHull3(EBuildMethod BuildMethod);
+		
 		static void BuildConvexHull(const TArray<FVec3Type>& InVertices, TArray<TVec3<int32>>& OutIndices, const Params& InParams = Params())
 		{
 			TRACE_CPUPROFILER_EVENT_SCOPE(Chaos::BuildConvexHull);
diff --git a/Engine/Source/Runtime/Experimental/Chaos/Public/Chaos/Convex.h b/Engine/Source/Runtime/Experimental/Chaos/Public/Chaos/Convex.h
--- a/Engine/Source/Runtime/Experimental/Chaos/Public/Chaos/Convex.h
+++ b/Engine/Source/Runtime/Experimental/Chaos/Public/Chaos/Convex.h
@@ -100,7 +100,7 @@ namespace Chaos
 			ComputeUnitMassInertiaTensorAndRotationOfMass(Volume);
 		}
 
-		FConvex(const TArray<FVec3Type>& InVertices, const FReal InMargin)
+		FConvex(const TArray<FVec3Type>& InVertices, const FReal InMargin, FConvexBuilder::EBuildMethod BuildMethod = FConvexBuilder::EBuildMethod::Default)
 		    : FImplicitObject(EImplicitObject::IsConvex | EImplicitObject::HasBoundingBox, ImplicitObjectType::Convex)
 		{
 			const int32 NumVertices = InVertices.Num();
@@ -110,23 +110,17 @@ namespace Chaos
 			}
 
 			TArray<TArray<int32>> FaceIndices;
-			FConvexBuilder::Build(InVertices, Planes, FaceIndices, Vertices, LocalBoundingBox);
+			FConvexBuilder::Build(InVertices, Planes, FaceIndices, Vertices, LocalBoundingBox, BuildMethod);
 			CHAOS_ENSURE(Planes.Num() == FaceIndices.Num());
 
 			// @todo(chaos): this only works with triangles. Fix that an we can run MergeFaces before calling this
 			CalculateVolumeAndCenterOfMass(Vertices, FaceIndices, Volume, CenterOfMass);
 
-			// @todo(chaos): should be based on size, or passed in
-			if (!FConvexBuilder::bUseGeometryTConvexHull3)
-			{
-				// @todo(convex) : TConvexHull3 does not need to merge faces, and 
-				// it appears that this code path can leave the convex in an
-				// undefined state. We should consider removing the merge faces when
-				// we transition to the UE::Geometry convex generation. 
-				const FRealType DistanceTolerance = 1.0f;
-				FConvexBuilder::MergeFaces(Planes, FaceIndices, Vertices, DistanceTolerance);
-				CHAOS_ENSURE(Planes.Num() == FaceIndices.Num());
-			}
+			// it appears that this code path can leave the convex in an
+			// undefined state.
+			const FRealType DistanceTolerance = 1.0f;
+			FConvexBuilder::MergeFaces(Planes, FaceIndices, Vertices, DistanceTolerance);
+			CHAOS_ENSURE(Planes.Num() == FaceIndices.Num());
 
 			CreateStructureData(MoveTemp(FaceIndices));
 
