Fix convex most opposing unity test failure
- Context : this was failing because of the recent change of convex generation algorithm that affects the order of the resulting faces
- Made the test more resilient to face order

diff --git a/Engine/Source/Programs/HeadlessChaos/Private/HeadlessChaosTestMostOpposing.cpp b/Engine/Source/Programs/HeadlessChaos/Private/HeadlessChaosTestMostOpposing.cpp
--- a/Engine/Source/Programs/HeadlessChaos/Private/HeadlessChaosTestMostOpposing.cpp
+++ b/Engine/Source/Programs/HeadlessChaos/Private/HeadlessChaosTestMostOpposing.cpp
@@ -292,8 +292,10 @@ namespace ChaosTest
 		// Reflection
 		{
 			// Reflecting in x-z plane should return the same face as not reflecting (in this particular case)
-			FVec3 ReflectionScale = FVec3(1,-1,1);
-			EXPECT_EQ(Convex->FindMostOpposingFaceScaled(FVec3(0, 0, 0), FVec3(-2, 0, -1).GetSafeNormal(), INDEX_NONE, 1000.0f, ReflectionScale),2);
+			const FVec3 ReflectionScale = FVec3(1, -1, 1);
+			const int32 ReflectionFace = Convex->FindMostOpposingFaceScaled(FVec3(0, 0, 0), FVec3(-2, 0, -1).GetSafeNormal(), INDEX_NONE, 1000.0f, ReflectionScale);
+			const int32 NoReflectionFace = Convex->FindMostOpposingFace(FVec3(0, 0, 0), FVec3(-2, 0, -1).GetSafeNormal(), INDEX_NONE, 1000.0f);
+			EXPECT_EQ(NoReflectionFace, ReflectionFace);
 		}
 	}
 
