Sweep height field optimization

diff --git a/Engine/Source/Programs/HeadlessChaos/Private/HeadlessChaosTestHeightfield.cpp b/Engine/Source/Programs/HeadlessChaos/Private/HeadlessChaosTestHeightfield.cpp
--- a/Engine/Source/Programs/HeadlessChaos/Private/HeadlessChaosTestHeightfield.cpp
+++ b/Engine/Source/Programs/HeadlessChaos/Private/HeadlessChaosTestHeightfield.cpp
@@ -24,9 +24,9 @@ namespace ChaosTest {
 			FReal Count = 0;
 			for(int32 Row = 0; Row < Rows; ++Row)
 			{
-				for(int32 Col = 0; Col < Columns; ++Col)
+				for(int32 Col = 0; Col < Columns; ++Col, ++Count)
 				{
-					Heights[Row * Columns + Col] = CountToWorldScale * Count++;
+					Heights[Row * Columns + Col] = CountToWorldScale * Count;
 				}
 			}
 
@@ -455,9 +455,253 @@ namespace ChaosTest {
 	}
 
 
+	void SweepSphereTest()
+	{
+		const int32 Columns = 10;
+		const int32 Rows = 10;
+		const FReal CountToWorldScale = 1;
+
+
+		TSphere<FReal, 3> Sphere(FVec3(0.0, 0.0, 0.0), 1.0);
+		{
+
+			TArray<FReal> Heights;
+			Heights.AddZeroed(Rows * Columns);
+
+			FReal Count = 0;
+			for (int32 Row = 0; Row < Rows; ++Row)
+			{
+				for (int32 Col = 0; Col < Columns; ++Col, ++Count)
+				{
+					Heights[Row * Columns + Col] = CountToWorldScale * Count;
+				}
+			}
+
+			auto ComputeExpectedNormal = [&](const FVec3& Scale)
+			{
+				//Compute expected normal
+				const FVec3 A(0, 0, 0);
+				const FVec3 B(Scale[0], 0, CountToWorldScale * Scale[2]);
+				const FVec3 C(0, Scale[1], Columns * CountToWorldScale * Scale[2]);
+
+				FVec3 ExpectedNormal = FVec3::CrossProduct((B - A), (C - A));
+				ExpectedNormal.SafeNormalize();
+				return ExpectedNormal;
+			};
+
+			auto AlongZTest = [&](const FVec3& Scale)
+			{
+				TArray<FReal> HeightsCopy = Heights;
+				FHeightField Heightfield(MoveTemp(HeightsCopy), TArray<uint8>(), Rows, Columns, Scale);
+				const auto& Bounds = Heightfield.BoundingBox();	//Current API forces us to do this to cache the bounds
+
+				//test straight down sweep
+				Count = 0;
+				FReal TOI;
+				FVec3 Position, Normal, FaceNormal;
+				int32 FaceIdx;
+
+				const FVec3 ExpectedNormal = ComputeExpectedNormal(Scale);
+
+				int32 ExpectedFaceIdx = 0;
+				for (int32 Row = 0; Row < Rows-1; ++Row)
+				{
+					for (int32 Col = 0; Col < Columns-1; ++Col)
+					{
+						const FVec3 Start(Col * Scale[0] + 0.5, Row * Scale[1]+0.5, 1000 * Scale[2]);
+						FRigidTransform3 StartTM(Start, TRotation<FReal, 3>());
+						FVec3 Dir(0, 0, -1);
+
+						bool Result = Heightfield.SweepGeom(Sphere, StartTM, Dir, 2000 * Scale[2], TOI, Position, Normal, FaceIdx, FaceNormal, 0.0, true);
+						EXPECT_TRUE(Result);
+						const FReal ExpectedTOI = -1.0;
+						EXPECT_NEAR(TOI, ExpectedTOI, 1e-2);
+						FVec3 AddGeomSize(1.0, 0.0, 1.0);
+						FVec3 ExpectedPos(Col ,Row , Heights[Row * Columns + Col]);
+						ExpectedPos += AddGeomSize;
+						ExpectedPos *= Scale;
+						EXPECT_VECTOR_NEAR(Position, ExpectedPos, 1e-2);
+						EXPECT_VECTOR_NEAR(Normal, ExpectedNormal, 1e-2);
+
+					}
+				}
+			};
+
+			AlongZTest(FVec3(1));
+			AlongZTest(FVec3(1, 1, 3));
+			AlongZTest(FVec3(1, 1, .3));
+
+
+			auto AlongXTest = [&](const FVec3& Scale)
+			{
+				TArray<FReal> HeightsCopy = Heights;
+				FHeightField Heightfield(MoveTemp(HeightsCopy), TArray<uint8>(), Rows, Columns, Scale);
+				const auto& Bounds = Heightfield.BoundingBox();	//Current API forces us to do this to cache the bounds
+
+				//test straight down sweep
+				Count = 0;
+				FReal TOI;
+				FVec3 Position, Normal, FaceNormal;
+				int32 FaceIdx;
+
+				const FVec3 ExpectedNormal = ComputeExpectedNormal(Scale);
+
+				int32 ExpectedFaceIdx = 0;
+				for (int32 Row = 0; Row < Rows - 1; ++Row)
+				{
+					for (int32 Col = 0; Col < Columns - 1; ++Col)
+					{
+
+						const FVec3 Start(-Scale[0], Row * Scale[1], Heights[Row * Columns + Col] * Scale[2] + 0.01 * Scale[2]);
+						FRigidTransform3 StartTM(Start, TRotation<FReal, 3>());
+						FVec3 Dir(1, 0, 0);
+
+						bool Result = Heightfield.SweepGeom(Sphere, StartTM, Dir, 2000 * Scale[0], TOI, Position, Normal, FaceIdx, FaceNormal, 0.0, true);
+
+						EXPECT_TRUE(Result);
+						const FReal ExpectedTOI = -1.0;
+						EXPECT_NEAR(TOI, ExpectedTOI, 1e-2);
+
+						FVec3 ExpectedPos(1.0, Row, Scale[2] * (1.0 + Row * Columns) );
+						EXPECT_VECTOR_NEAR(Position, ExpectedPos, 1e-2);
+						EXPECT_VECTOR_NEAR(Normal, ExpectedNormal, 1e-2);
+					}
+				}
+			};
+
+			AlongXTest(FVec3(1));
+			AlongXTest(FVec3(1, 1, 3));
+			AlongXTest(FVec3(1, 1, .3));
+		}
+	}
+
+	void SweepBoxTest()
+	{
+		const int32 Columns = 10;
+		const int32 Rows = 10;
+		const FReal CountToWorldScale = 1;
+
+
+		TBox<FReal, 3> Box(FVec3(-0.05, -0.05, -0.05), FVec3(0.05, 0.05, 0.05));
+		{
+
+			TArray<FReal> Heights;
+			Heights.AddZeroed(Rows * Columns);
+
+			FReal Count = 0;
+			for (int32 Row = 0; Row < Rows; ++Row)
+			{
+				for (int32 Col = 0; Col < Columns; ++Col, ++Count)
+				{
+					Heights[Row * Columns + Col] = CountToWorldScale * Count;
+				}
+			}
+
+			auto ComputeExpectedNormal = [&](const FVec3& Scale)
+			{
+				//Compute expected normal
+				const FVec3 A(0, 0, 0);
+				const FVec3 B(Scale[0], 0, CountToWorldScale * Scale[2]);
+				const FVec3 C(0, Scale[1], Columns * CountToWorldScale * Scale[2]);
+
+				FVec3 ExpectedNormal = FVec3::CrossProduct((B - A), (C - A));
+				ExpectedNormal.SafeNormalize();
+				return ExpectedNormal;
+			};
+
+			auto AlongZTest = [&](const FVec3& Scale)
+			{
+				TArray<FReal> HeightsCopy = Heights;
+				FHeightField Heightfield(MoveTemp(HeightsCopy), TArray<uint8>(), Rows, Columns, Scale);
+				const auto& Bounds = Heightfield.BoundingBox();	//Current API forces us to do this to cache the bounds
+
+				//test straight down sweep
+				Count = 0;
+				FReal TOI;
+				FVec3 Position, Normal, FaceNormal;
+				int32 FaceIdx;
+
+				const FVec3 ExpectedNormal = ComputeExpectedNormal(Scale);
+
+				int32 ExpectedFaceIdx = 0;
+				for (int32 Row = 0; Row < Rows - 1; ++Row)
+				{
+					for (int32 Col = 0; Col < Columns - 1; ++Col)
+					{
+						const FVec3 Start(Col * Scale[0] + 0.5, Row * Scale[1] + 0.5, 1000 * Scale[2]);
+						FRigidTransform3 StartTM(Start, TRotation<FReal, 3>());
+						FVec3 Dir(0, 0, -1);
+
+						bool Result = Heightfield.SweepGeom(Box, StartTM, Dir, 2000 * Scale[2], TOI, Position, Normal, FaceIdx, FaceNormal, 0.0, true);
+						EXPECT_TRUE(Result);
+						FVec3 AddGeomSize(1.0, 0.0, 1.0);
+						FVec3 ExpectedPos(Col, Row, Heights[Row * Columns + Col]);
+						ExpectedPos += AddGeomSize;
+						ExpectedPos *= Scale;
+						EXPECT_VECTOR_NEAR(Position, ExpectedPos, 1e-2);
+						EXPECT_VECTOR_NEAR(Normal, ExpectedNormal, 1e-2);
+
+					}
+				}
+			};
+
+			AlongZTest(FVec3(1));
+			AlongZTest(FVec3(1, 1, 3));
+			AlongZTest(FVec3(1, 1, .3));
+
+
+			auto AlongXTest = [&](const FVec3& Scale)
+			{
+				TArray<FReal> HeightsCopy = Heights;
+				FHeightField Heightfield(MoveTemp(HeightsCopy), TArray<uint8>(), Rows, Columns, Scale);
+				const auto& Bounds = Heightfield.BoundingBox();	//Current API forces us to do this to cache the bounds
+
+				//test straight down sweep
+				Count = 0;
+				FReal TOI;
+				FVec3 Position, Normal, FaceNormal;
+				int32 FaceIdx;
+
+				const FVec3 ExpectedNormal = ComputeExpectedNormal(Scale);
+
+				int32 ExpectedFaceIdx = 0;
+				for (int32 Row = 0; Row < Rows - 1; ++Row)
+				{
+					for (int32 Col = 0; Col < Columns - 1; ++Col)
+					{
+
+						const FVec3 Start(-Scale[0], Row * Scale[1], Heights[Row * Columns + Col] * Scale[2] + 0.01 * Scale[2]);
+						FRigidTransform3 StartTM(Start, TRotation<FReal, 3>());
+						FVec3 Dir(1, 0, 0);
+
+						bool Result = Heightfield.SweepGeom(Box, StartTM, Dir, 2000 * Scale[0], TOI, Position, Normal, FaceIdx, FaceNormal, 0.0, true);
+
+						EXPECT_TRUE(Result);
+
+						FVec3 ExpectedPos(1.0, Row, Scale[2] * (1.0 + Row * Columns));
+						EXPECT_VECTOR_NEAR(Position, ExpectedPos, 1e-2);
+						EXPECT_VECTOR_NEAR(Normal, ExpectedNormal, 1e-2);
+					}
+				}
+			};
+
+			AlongXTest(FVec3(1));
+			AlongXTest(FVec3(1, 1, 3));
+			AlongXTest(FVec3(1, 1, .3));
+		}
+	}
+
+	void SweepTest()
+	{
+		SweepSphereTest();
+		SweepBoxTest();
+	}
+	
+
 	TEST(ChaosTests, Heightfield)
 	{
 		ChaosTest::Raycast();
+		ChaosTest::SweepTest();
 		EditHeights();
 		SUCCEED();
 	}
diff --git a/Engine/Source/Runtime/Experimental/Chaos/Private/Chaos/HeightField.cpp b/Engine/Source/Runtime/Experimental/Chaos/Private/Chaos/HeightField.cpp
--- a/Engine/Source/Runtime/Experimental/Chaos/Private/Chaos/HeightField.cpp
+++ b/Engine/Source/Runtime/Experimental/Chaos/Private/Chaos/HeightField.cpp
@@ -316,6 +316,14 @@ namespace Chaos
 				bParallel[Axis] = FMath::IsNearlyZero(Dir[Axis], (FReal)1.e-8);
 				InvDir[Axis] = bParallel[Axis] ? 0 : 1 / Dir[Axis];
 			}
+
+			StartPointSimd = MakeVectorRegisterFloatFromDouble(MakeVectorRegister(StartPoint.X, StartPoint.Y, StartPoint.Z, 0.0));
+			DirSimd = MakeVectorRegisterFloatFromDouble(MakeVectorRegister(Dir.X, Dir.Y, Dir.Z, 0.0));
+			Parallel = VectorCompareLT(VectorAbs(DirSimd), GlobalVectorConstants::SmallNumber);
+			InvDirSimd = VectorBitwiseNotAnd(Parallel, VectorDivide(VectorOne(), DirSimd));
+			Inflation3DSimd = MakeVectorRegisterFloatFromDouble(MakeVectorRegister(Inflation3D.X, Inflation3D.Y, Inflation3D.Z, 0.0));
+			Inflation3DSimd = VectorAbs(Inflation3DSimd);
+
 		}
 
 		bool VisitSweep(int32 Payload, FReal& CurrentLength)
@@ -325,7 +333,7 @@ namespace Chaos
 
 			const int32 FullIndex = Payload + SubY;
 
-			auto TestTriangle = [&](int32 FaceIndex, const FVec3& A, const FVec3& B, const FVec3& C) -> bool
+			auto TestTriangle = [&](int32 FaceIndex, const VectorRegister4Float& A, const VectorRegister4Float& B, const VectorRegister4Float& C) -> bool
 			{
 				if(OutTime == 0)
 				{
@@ -333,17 +341,20 @@ namespace Chaos
 				}
 
 				//Convert into local space of A to get better precision
-				VectorRegister4Float AReg = MakeVectorRegisterFloatFromDouble(MakeVectorRegister(A.X, A.Y, A.Z, 0.0));
-				VectorRegister4Float BReg = MakeVectorRegisterFloatFromDouble(MakeVectorRegister(B.X, B.Y, B.Z, 0.0));
-				VectorRegister4Float CReg = MakeVectorRegisterFloatFromDouble(MakeVectorRegister(C.X, C.Y, C.Z, 0.0));
+				FTriangleRegister Triangle(VectorZero(), VectorSubtract(B, A), VectorSubtract(C, A));
 
-				FTriangleRegister Triangle(VectorZero(), VectorSubtract(BReg, AReg), VectorSubtract(CReg, AReg));
+				const UE::Math::TQuat<FReal>& RotationDouble = StartTM.GetRotation();
+				VectorRegister4Float Rotation = MakeVectorRegisterFloatFromDouble(MakeVectorRegister(RotationDouble.X, RotationDouble.Y, RotationDouble.Z, RotationDouble.W));
+				// Normalize rotation
+				Rotation = VectorNormalizeSafe(Rotation, GlobalVectorConstants::Float0001);
 
-				FReal Time;
-				FVec3 LocalHitPosition;
-				FVec3 HitNormal;
-				const FRigidTransform3 LocalStartTM(StartTM.GetTranslation() - A,StartTM.GetRotation());
-				if(GJKRaycast2<FReal>(Triangle, OtherGeom, LocalStartTM, Dir, CurrentLength, Time, LocalHitPosition, HitNormal, Thickness, bComputeMTD))
+				const UE::Math::TVector<FReal>& TranslationDouble = StartTM.GetTranslation();
+				VectorRegister4Float Translation = MakeVectorRegisterFloatFromDouble(MakeVectorRegister(TranslationDouble.X, TranslationDouble.Y, TranslationDouble.Z, 0.0));
+				Translation = VectorSubtract(Translation, A);
+
+				FRealSingle Time;
+				VectorRegister4Float OutPositionSimd, OutNormalSimd;
+				if(GJKRaycast2ImplSimd(Triangle, OtherGeom, Rotation, Translation, DirSimd, static_cast<FRealSingle>(CurrentLength), Time, OutPositionSimd, OutNormalSimd, bComputeMTD, GlobalVectorConstants::Float1000))
 				{
 					if(Time < OutTime)
 					{
@@ -357,8 +368,18 @@ namespace Chaos
 
 						if(!bHole)
 						{
-							OutNormal = HitNormal;
-							OutPosition = LocalHitPosition + A;
+							alignas(16) FRealSingle OutFloat[4];
+							VectorStoreAligned(OutNormalSimd, OutFloat);
+							OutNormal.X = OutFloat[0];
+							OutNormal.Y = OutFloat[1];
+							OutNormal.Z = OutFloat[2];
+
+							OutPositionSimd = VectorAdd(OutPositionSimd, A);
+							VectorStoreAligned(OutPositionSimd, OutFloat);
+							OutPosition.X = OutFloat[0];
+							OutPosition.Y = OutFloat[1];
+							OutPosition.Z = OutFloat[2];
+
 							OutTime = Time;
 							OutFaceIndex = FaceIndex;
 
@@ -368,13 +389,17 @@ namespace Chaos
 								// we adopt the triangle normal but this leaves us with an incorrect MTD from the GJK call
 								// above. #TODO possibly re-do GJK with a plane, or some geom vs.plane special case to solve
 								// both triangles as planes 
-								const FVec3 AB = B - A;
-								const FVec3 AC = C - A;
+								const VectorRegister4Float AB = VectorSubtract(B, A);
+								const VectorRegister4Float AC = VectorSubtract(C, A);
 
-								FVec3 TriNormal = FVec3::CrossProduct(AB, AC);
-								TriNormal.SafeNormalize();
+								VectorRegister4Float TriNormal = VectorCross(AB, AC);
+								TriNormal = VectorNormalize(TriNormal);
+
+								VectorStoreAligned(TriNormal, OutFloat);
+								OutNormal.X = OutFloat[0];
+								OutNormal.Y = OutFloat[1];
+								OutNormal.Z = OutFloat[2];
 
-								OutNormal = TriNormal;
 								CurrentLength = 0;
 								return false;
 							}
@@ -387,16 +412,12 @@ namespace Chaos
 				return true;
 			};
 
-			FVec3 Points[4];
-			FAABB3 CellBounds;
-			HfData->GetPointsAndBoundsScaled(FullIndex, Points, CellBounds);
-			CellBounds.ThickenSymmetrically(Inflation3D);
+			VectorRegister4Float Points[4];
+			FAABBVectorized CellBounds;
+			HfData->GetPointsAndBoundsScaledSimd(FullIndex, Points, CellBounds);
+			CellBounds.Thicken(Inflation3DSimd);
 
-			// Check cell bounds
-			//todo: can do it without raycast
-			FReal TOI;
-			FVec3 HitPoint;
-			if (CellBounds.RaycastFast(StartPoint, Dir, InvDir, bParallel, CurrentLength, 1.0f / CurrentLength, TOI, HitPoint))
+			if (CellBounds.RaycastFast(StartPointSimd, InvDirSimd, Parallel, VectorSetFloat1(static_cast<FRealSingle>(CurrentLength))))
 			{
 				bool bContinue = TestTriangle(Payload * 2, Points[0], Points[1], Points[3]);
 				if (bContinue)
@@ -404,6 +425,7 @@ namespace Chaos
 					TestTriangle(Payload * 2 + 1, Points[0], Points[3], Points[2]);
 				}
 			}
+
 			return OutTime > 0;
 		}
 
@@ -412,6 +434,12 @@ namespace Chaos
 		FVec3 OutNormal;
 		int32 OutFaceIndex;
 
+		VectorRegister4Float StartPointSimd;
+		VectorRegister4Float InvDirSimd;
+		VectorRegister4Float DirSimd;
+		VectorRegister4Float Parallel;
+		VectorRegister4Float Inflation3DSimd;
+
 	private:
 
 		const typename FHeightField::FDataType* HfData;
diff --git a/Engine/Source/Runtime/Experimental/Chaos/Public/Chaos/HeightField.h b/Engine/Source/Runtime/Experimental/Chaos/Public/Chaos/HeightField.h
--- a/Engine/Source/Runtime/Experimental/Chaos/Public/Chaos/HeightField.h
+++ b/Engine/Source/Runtime/Experimental/Chaos/Public/Chaos/HeightField.h
@@ -286,6 +286,28 @@ namespace Chaos
 				OutBounds = FAABB3(FVec3((FReal)X, (FReal)Y, MinZ), FVec3((FReal)X + 1, (FReal)Y + 1, MaxZ));
 			}
 
+			FORCEINLINE void GetPointsAndBoundsSimd(int32 Index, VectorRegister4Float OutPts[4], FAABBVectorized& OutBounds) const
+			{
+				const FRealSingle H0 = static_cast<FRealSingle>(MinValue + Heights[Index] * HeightPerUnit);
+				const FRealSingle H1 = static_cast<FRealSingle>(MinValue + Heights[Index + 1] * HeightPerUnit);
+				const FRealSingle H2 = static_cast<FRealSingle>(MinValue + Heights[Index + NumCols] * HeightPerUnit);
+				const FRealSingle H3 = static_cast<FRealSingle>(MinValue + Heights[Index + NumCols + 1] * HeightPerUnit);
+
+				const int32 X = Index % (NumCols);
+				const int32 Y = Index / (NumCols);
+
+				OutPts[0] = MakeVectorRegisterFloat((FRealSingle)X, (FRealSingle)Y, H0, 0.0f );
+				OutPts[1] = MakeVectorRegisterFloat((FRealSingle)X + 1, (FRealSingle)Y, H1, 0.0f);
+				OutPts[2] = MakeVectorRegisterFloat((FRealSingle)X, (FRealSingle)Y + 1, H2, 0.0f);
+				OutPts[3] = MakeVectorRegisterFloat((FRealSingle)X + 1, (FRealSingle)Y + 1, H3, 0.0f);
+
+				const FRealSingle MinZ = FMath::Min<FRealSingle>(H0, FMath::Min<FRealSingle>(H1, FMath::Min<FRealSingle>(H2, H3)));
+				const FRealSingle MaxZ = FMath::Max<FRealSingle>(H0, FMath::Max<FRealSingle>(H1, FMath::Max<FRealSingle>(H2, H3)));
+
+				OutBounds = FAABBVectorized(MakeVectorRegisterFloat((FRealSingle)X, (FRealSingle)Y, MinZ, 0.0f),
+											MakeVectorRegisterFloat((FRealSingle)X + 1, (FRealSingle)Y + 1, MaxZ, 0.0f));
+			}
+			
 			FORCEINLINE void GetBounds(TVec2<int32> CellIdx, TVec2<int32> Area, FAABBVectorized& OutBounds) const
 			{
 				FRealSingle MinHeight = UE_BIG_NUMBER;
@@ -350,6 +372,25 @@ namespace Chaos
 				OutBounds = FAABB3::FromPoints(OutBounds.Min() * Scale, OutBounds.Max() * Scale);
 			}
 
+			FORCEINLINE void GetPointsAndBoundsScaledSimd(int32 Index, VectorRegister4Float OutPts[4], FAABBVectorized& OutBounds) const
+			{
+				GetPointsAndBoundsSimd(Index, OutPts, OutBounds);
+
+				VectorRegister4Float ScaleSimd = MakeVectorRegisterFloatFromDouble(MakeVectorRegisterDouble(Scale.X, Scale.Y, Scale.Z, 0.0));
+
+				OutPts[0] = VectorMultiply(OutPts[0], ScaleSimd);
+				OutPts[1] = VectorMultiply(OutPts[1], ScaleSimd);
+				OutPts[2] = VectorMultiply(OutPts[2], ScaleSimd);
+				OutPts[3] = VectorMultiply(OutPts[3], ScaleSimd);
+
+				VectorRegister4Float P0 = VectorMultiply(OutBounds.GetMin(), ScaleSimd);
+				VectorRegister4Float P1 = VectorMultiply(OutBounds.GetMax(), ScaleSimd);
+				VectorRegister4Float Min = VectorMin(P0, P1);
+				VectorRegister4Float Max = VectorMax(P0, P1);
+
+				OutBounds = FAABBVectorized(Min, Max);
+			}
+
 			FORCEINLINE void GetBoundsScaled(TVec2<int32> CellIdx, TVec2<int32> Area, FAABBVectorized& OutBounds) const
 			{
 				GetBounds(CellIdx, Area, OutBounds);
