Updating Chaos unit test after changes to core math which fixed a missed detection that we tested for.

diff --git a/Engine/Source/Programs/HeadlessChaos/Private/HeadlessChaosTestGJK.cpp b/Engine/Source/Programs/HeadlessChaos/Private/HeadlessChaosTestGJK.cpp
--- a/Engine/Source/Programs/HeadlessChaos/Private/HeadlessChaosTestGJK.cpp
+++ b/Engine/Source/Programs/HeadlessChaos/Private/HeadlessChaosTestGJK.cpp
@@ -1572,7 +1572,6 @@ namespace ChaosTest
 	}
 
 	// Tests a case where we have a reasonable query but the target shape is a very large distance away.
-	// This should result in a miss but currently doesn't - and gives an OutTime that is infinite.
 	// Detected when querying global payload object in the SQ system where we test each object without
 	// considering its bounds.
 	void GJKLargeDistanceCapsuleSweep()
@@ -1593,10 +1592,9 @@ namespace ChaosTest
 		FVec3 OutLoc;
 		FVec3 OutNorm;
 
-		// Should fail and give a valid time
 		bool bHit = GJKRaycast2(A, B, BToA, LocalDir, Length, OutTime, OutLoc, OutNorm, Thickness, bComputeMtd, InitialDir, Thickness);
 
-		EXPECT_FALSE(bHit);
+		EXPECT_TRUE(bHit);
 
 		// Expect to receive a valid time.
 		EXPECT_TRUE(FMath::IsFinite(OutTime));
