Fix Sweep inflated bounds test assuming start position is at bounds center. Fixes incorrect results with objects with geometry not centered at origin.

diff --git a/Engine/Source/Programs/HeadlessChaos/Private/HeadlessChaosTestPhysicsScene.cpp b/Engine/Source/Programs/HeadlessChaos/Private/HeadlessChaosTestPhysicsScene.cpp
--- a/Engine/Source/Programs/HeadlessChaos/Private/HeadlessChaosTestPhysicsScene.cpp
+++ b/Engine/Source/Programs/HeadlessChaos/Private/HeadlessChaosTestPhysicsScene.cpp
@@ -1869,4 +1869,57 @@ namespace ChaosTest {
 
 		EXPECT_TRUE(HitBuffer.HasBlockingHit());
 	}
+
+	GTEST_TEST(EngineInterface, SweepOffsetActor)
+	{
+		FChaosScene Scene(nullptr, /*AsyncDt=*/-1);
+
+		FActorCreationParams Params;
+		Params.Scene = &Scene;
+		Params.bSimulatePhysics = false;
+		Params.bStatic = true;
+		Params.InitialTM = FTransform::Identity;
+		Params.Scene = &Scene;
+
+		FPhysicsActorHandle StaticCube = nullptr;
+
+		FChaosEngineInterface::CreateActor(Params, StaticCube);
+		ASSERT_NE(StaticCube, nullptr);
+
+		// Add geometry, placing a box at the origin
+		constexpr FReal BoxSize = static_cast<FReal>(50.0);
+		const FVec3 HalfBoxExtent{ BoxSize };
+
+		// We require a union here, although the second geometry isn't used we need the particle to
+		// have more than one shape in its shapes array otherwise the query acceleration will treat
+		// it as a special case and skip bounds checking during the overlap
+		TArray<TUniquePtr<FImplicitObject>> Geoms;
+		Geoms.Emplace(MakeUnique<TBox<FReal, 3>>(-HalfBoxExtent, HalfBoxExtent));
+		Geoms.Emplace(MakeUnique<TBox<FReal, 3>>(-HalfBoxExtent, HalfBoxExtent));
+
+		auto& Particle = StaticCube->GetGameThreadAPI();
+		{
+			TUniquePtr<FImplicitObjectUnion> GeomUnion = MakeUnique<FImplicitObjectUnion>(MoveTemp(Geoms));
+			Particle.SetGeometry(MoveTemp(GeomUnion));
+		}
+
+		TArray<FPhysicsActorHandle> Particles{ StaticCube };
+		Scene.AddActorsToScene_AssumesLocked(Particles);
+
+		FChaosSQAccelerator SQ{ *Scene.GetSpacialAcceleration() };
+		FSQHitBuffer<ChaosInterface::FSweepHit> HitBuffer;
+		FBlockAllQueryCallback QueryCallback;
+
+		// Another box of same size that is offset from origin by 200.
+		FVec3 Offset(200.f,0,0);
+		TBox<FReal, 3> QueryBox(-HalfBoxExtent + Offset, HalfBoxExtent + Offset);
+
+		// Sweep positions offset box directly above box at origin, should hit box sweeping downward.
+		const FTransform QueryTM{ FVec3{-200.f, 0, 100.0f} };
+		const FVec3 Dir(0,0,-1);
+		const FReal Length = 200;
+		SQ.Sweep(QueryBox, QueryTM, Dir, Length, HitBuffer, EHitFlags::None, FQueryFilterData(), QueryCallback, FQueryDebugParams());
+
+		EXPECT_TRUE(HitBuffer.HasBlockingHit());
+	}
 }
diff --git a/Engine/Source/Runtime/PhysicsCore/Private/SQAccelerator.cpp b/Engine/Source/Runtime/PhysicsCore/Private/SQAccelerator.cpp
--- a/Engine/Source/Runtime/PhysicsCore/Private/SQAccelerator.cpp
+++ b/Engine/Source/Runtime/PhysicsCore/Private/SQAccelerator.cpp
@@ -287,7 +287,8 @@ private:
 					//todo: use fast raycast
 					Chaos::FReal TmpTime;
 					FVec3 TmpPos;
-					if (!InflatedWorldBounds.RaycastFast( SQ == ESQType::Raycast ? StartPoint : StartTM.GetLocation(), CurData->Dir, CurData->InvDir, CurData->bParallel, CurData->CurrentLength, CurData->InvCurrentLength, TmpTime, TmpPos))
+					const FVec3 InflatedBoundsTraceStart = SQ == ESQType::Raycast ? StartPoint : QueryGeomWorldBounds.Center();
+					if (!InflatedWorldBounds.RaycastFast(InflatedBoundsTraceStart, CurData->Dir, CurData->InvDir, CurData->bParallel, CurData->CurrentLength, CurData->InvCurrentLength, TmpTime, TmpPos))
 					{
 						continue;
 					}
