Chaos - Add Managed Array Remove at and Insert at in cloth tests.

diff --git a/Engine/Source/Programs/HeadlessChaos/Private/HeadlessChaosApp.cpp b/Engine/Source/Programs/HeadlessChaos/Private/HeadlessChaosApp.cpp
--- a/Engine/Source/Programs/HeadlessChaos/Private/HeadlessChaosApp.cpp
+++ b/Engine/Source/Programs/HeadlessChaos/Private/HeadlessChaosApp.cpp
@@ -166,6 +166,11 @@ TEST(BroadphaseTests, Broadphase) {
 //	SUCCEED();
 //}
 
+TEST(ClothTests, ClothCollection) {
+	ChaosTest::ClothCollection();
+	SUCCEED();
+}
+
 TEST(RaycastTests, Raycast) {
 	ChaosTest::SphereRaycast();
 	ChaosTest::PlaneRaycast();
diff --git a/Engine/Source/Programs/HeadlessChaos/Private/HeadlessChaosTestCloth.cpp b/Engine/Source/Programs/HeadlessChaos/Private/HeadlessChaosTestCloth.cpp
--- a/Engine/Source/Programs/HeadlessChaos/Private/HeadlessChaosTestCloth.cpp
+++ b/Engine/Source/Programs/HeadlessChaos/Private/HeadlessChaosTestCloth.cpp
@@ -12,6 +12,7 @@
 #include "Chaos/PBDSpringConstraints.h"
 #include "Chaos/TriangleMesh.h"
 #include "Chaos/Utilities.h"
+#include "GeometryCollection/ManagedArrayCollection.h"
 
 #include "Modules/ModuleManager.h"
 
@@ -324,4 +325,83 @@ namespace ChaosTest {
 		AddAxialConstraint(Evolution, MoveTemp(Triangles), 1.0);
 	}
 
+	void ClothCollection()
+	{
+		FManagedArrayCollection ManagedArrayCollection;
+
+		const FName DataGroup("Data");
+		TManagedArray<int32> Value;
+		ManagedArrayCollection.AddExternalAttribute<int32>("Value", DataGroup, Value);
+
+		const FName ViewGroup("View");
+		FManagedArrayCollection::FConstructionParameters DataDependency(DataGroup);
+		TManagedArray<int32> ValueStart;
+		TManagedArray<int32> ValueEnd;
+		TManagedArray<int32> Id;
+		ManagedArrayCollection.AddExternalAttribute<int32>("ValueStart", ViewGroup, ValueStart, DataDependency);
+		ManagedArrayCollection.AddExternalAttribute<int32>("ValueEnd", ViewGroup, ValueEnd, DataDependency);
+		ManagedArrayCollection.AddExternalAttribute<int32>("Id", ViewGroup, Id);
+
+		auto InsertView = [&ManagedArrayCollection, &Value, &ValueStart, &ValueEnd, &Id, &DataGroup, &ViewGroup](int32 ViewPosition, int32 NumValues)
+			{
+				check(ViewPosition <= ManagedArrayCollection.NumElements(ViewGroup));
+				EXPECT_TRUE(ManagedArrayCollection.InsertElements(1, ViewPosition, ViewGroup) == ViewPosition);
+
+				const int32 ValuePosition = (ViewPosition > 0) ? ValueEnd[ViewPosition - 1] + 1: 0;
+				EXPECT_TRUE(ManagedArrayCollection.InsertElements(NumValues, ValuePosition, DataGroup) == ValuePosition);
+			
+				static int32 ViewId = 0;
+				Id[ViewPosition] = ViewId++;
+				ValueStart[ViewPosition] = ValuePosition;
+				ValueEnd[ViewPosition] = ValuePosition + NumValues - 1;
+
+				for (int32 ValueIndex = ValueStart[ViewPosition]; ValueIndex <= ValueEnd[ViewPosition]; ++ValueIndex)
+				{
+					Value[ValueIndex] = Id[ViewPosition];
+				}
+			};
+
+		auto HasKeptIntegrity = [&Value , &ValueStart, &ValueEnd, &Id]()->bool
+			{
+				EXPECT_TRUE(ValueEnd.Num() == ValueStart.Num() && Id.Num() == ValueStart.Num());
+				for (int32 ViewIndex = 0; ViewIndex < ValueStart.Num(); ++ViewIndex)
+				{
+					for (int32 ValueIndex = ValueStart[ViewIndex]; ValueIndex <= ValueEnd[ViewIndex]; ++ValueIndex)
+					{
+						if (Value[ValueIndex] != Id[ViewIndex])
+						{
+							return false;
+						}
+					}
+				}
+				return true;
+			};
+
+		// Insert 5 views from the start of the array
+		for (int32 Index = 0; Index < 5; ++Index)
+		{
+			InsertView(0, FMath::Rand() / 32);
+		}
+		EXPECT_TRUE(HasKeptIntegrity());
+
+		// Add 1 view in the middle of the array
+		InsertView(ManagedArrayCollection.NumElements(ViewGroup) / 2, FMath::Rand() / 32);
+		EXPECT_TRUE(HasKeptIntegrity());
+
+		// Add 1 view at the end of the array
+		InsertView(ManagedArrayCollection.NumElements(ViewGroup), FMath::Rand() / 32);
+		EXPECT_TRUE(HasKeptIntegrity());
+
+		// Remove 1 view from the start of the array
+		ManagedArrayCollection.RemoveElements(ViewGroup, 1, 0);
+		EXPECT_TRUE(HasKeptIntegrity());
+
+		// Remove 1 view from the end of the array
+		ManagedArrayCollection.RemoveElements(ViewGroup, 1, ManagedArrayCollection.NumElements(ViewGroup) - 1);
+		EXPECT_TRUE(HasKeptIntegrity());
+
+		// Remove 2 views from the middle
+		ManagedArrayCollection.RemoveElements(ViewGroup, 2, ManagedArrayCollection.NumElements(ViewGroup) / 2);
+		EXPECT_TRUE(HasKeptIntegrity());
+	}
 } // namespace ChaosTest
\ No newline at end of file
diff --git a/Engine/Source/Programs/HeadlessChaos/Public/HeadlessChaosTestCloth.h b/Engine/Source/Programs/HeadlessChaos/Public/HeadlessChaosTestCloth.h
--- a/Engine/Source/Programs/HeadlessChaos/Public/HeadlessChaosTestCloth.h
+++ b/Engine/Source/Programs/HeadlessChaos/Public/HeadlessChaosTestCloth.h
@@ -10,4 +10,6 @@ namespace ChaosTest {
 	/**/
 	void EdgeConstraints();
 
+	/* Test Remove & Insert at position as used in the ClothCollection. */
+	void ClothCollection();
 }
\ No newline at end of file
