Fix GJK unit test introduce by some core math stuff that has changed.

diff --git a/Engine/Source/Programs/HeadlessChaos/Private/HeadlessChaosTestGJK.cpp b/Engine/Source/Programs/HeadlessChaos/Private/HeadlessChaosTestGJK.cpp
--- a/Engine/Source/Programs/HeadlessChaos/Private/HeadlessChaosTestGJK.cpp
+++ b/Engine/Source/Programs/HeadlessChaos/Private/HeadlessChaosTestGJK.cpp
@@ -1572,6 +1572,7 @@ namespace ChaosTest
 	}
 
 	// Tests a case where we have a reasonable query but the target shape is a very large distance away.
+	// This should result in a miss but currently doesn't - and gives an OutTime that is infinite.
 	// Detected when querying global payload object in the SQ system where we test each object without
 	// considering its bounds.
 	void GJKLargeDistanceCapsuleSweep()
@@ -1592,9 +1593,10 @@ namespace ChaosTest
 		FVec3 OutLoc;
 		FVec3 OutNorm;
 
+		// Should fail and give a valid time
 		bool bHit = GJKRaycast2(A, B, BToA, LocalDir, Length, OutTime, OutLoc, OutNorm, Thickness, bComputeMtd, InitialDir, Thickness);
 
-		EXPECT_TRUE(bHit);
+		EXPECT_FALSE(bHit);
 
 		// Expect to receive a valid time.
 		EXPECT_TRUE(FMath::IsFinite(OutTime));
diff --git a/Engine/Source/Runtime/Experimental/Chaos/Public/Chaos/GJK.h b/Engine/Source/Runtime/Experimental/Chaos/Public/Chaos/GJK.h
--- a/Engine/Source/Runtime/Experimental/Chaos/Public/Chaos/GJK.h
+++ b/Engine/Source/Runtime/Experimental/Chaos/Public/Chaos/GJK.h
@@ -1793,7 +1793,13 @@ namespace Chaos
 				break;	//if taking too long just stop. This should never happen
 			}
 
-			V = VectorNormalizeAccurate(V);
+			// Normalization
+			const VectorRegister4Float VRepLength = VectorReciprocalLen(V);
+			if (VectorMaskBits(VRepLength) == 0)
+			{
+				return false;
+			}
+			V = VectorMultiply(V, VRepLength);
 
 			SupportA = SupportAFunc(V);
 			SupportB = SupportBFunc(VectorNegate(V));
