HeadlessChaos - fix debug build target

diff --git a/Engine/Source/Programs/HeadlessChaos/HeadlessChaos.Target.cs b/Engine/Source/Programs/HeadlessChaos/HeadlessChaos.Target.cs
--- a/Engine/Source/Programs/HeadlessChaos/HeadlessChaos.Target.cs
+++ b/Engine/Source/Programs/HeadlessChaos/HeadlessChaos.Target.cs
@@ -29,5 +29,8 @@ public class HeadlessChaosTarget : TestTargetRules
 		bCompilePhysX = false;
 		bCompileAPEX = false;
 		bCompileNvCloth = false;
+
+		// our gtest does not have debug crt libs
+		bDebugBuildsActuallyUseDebugCRT = false;
 	}
 }
