HeadlessChaos - fix compiler errors

diff --git a/Engine/Source/Programs/HeadlessChaos/Private/HeadlessChaosTestConstraints.cpp b/Engine/Source/Programs/HeadlessChaos/Private/HeadlessChaosTestConstraints.cpp
--- a/Engine/Source/Programs/HeadlessChaos/Private/HeadlessChaosTestConstraints.cpp
+++ b/Engine/Source/Programs/HeadlessChaos/Private/HeadlessChaosTestConstraints.cpp
@@ -60,8 +60,8 @@ namespace ChaosTest {
 			Evolution.AddConstraintRule(&ConstraintRule);
 
 			// The effect of stiffness parameter (which is set to 0.5 above) is iteration depeendent
-			Evolution.SetNumIterations(1);
-			Evolution.SetNumPushOutIterations(1);
+			Evolution.SetNumPositionIterations(1);
+			Evolution.SetNumVelocityIterations(1);
 
 			Evolution.AdvanceOneTimeStep(0.1);
 			Evolution.EndFrame(0.1);
@@ -96,7 +96,7 @@ namespace ChaosTest {
 		TArray<FPBDRigidParticleHandle*> Dynamics = Evolution.CreateDynamicParticles(2);
 		TArray<FVec3> PositionConstraintPositions = { FVec3(0, 0, 0) };
 
-		Evolution.SetNumIterations(Iterations);
+		Evolution.SetNumPositionIterations(Iterations);
 
 		Dynamics[1]->X() = FVec3(500, 0, 0);
 		FVec3 JointConstraintPosition = FVec3(0, 0, 0);
@@ -262,8 +262,8 @@ namespace ChaosTest {
 		THandleArray<FChaosPhysicsMaterial> PhysicalMaterials;
 		TEvolution Evolution(Particles, PhysicalMaterials);
 		InitEvolutionSettings(Evolution);
-		Evolution.SetNumIterations(1);
-		Evolution.SetNumPushOutIterations(1);
+		Evolution.SetNumPositionIterations(1);
+		Evolution.SetNumVelocityIterations(1);
 
 		// disable gravity
 		Evolution.GetGravityForces().SetAcceleration(FVec3(0, 0, -980.f));
diff --git a/Engine/Source/Programs/HeadlessChaos/Public/HeadlessChaosTestConstraints.h b/Engine/Source/Programs/HeadlessChaos/Public/HeadlessChaosTestConstraints.h
--- a/Engine/Source/Programs/HeadlessChaos/Public/HeadlessChaosTestConstraints.h
+++ b/Engine/Source/Programs/HeadlessChaos/Public/HeadlessChaosTestConstraints.h
@@ -32,7 +32,7 @@ namespace ChaosTest
 			PhysicalMaterial->DisabledLinearThreshold = 0;
 			PhysicalMaterial->DisabledAngularThreshold = 0;
 
-			Evolution.SetNumIterations(NumIterations);
+			Evolution.SetNumPositionIterations(NumIterations);
 			Evolution.GetGravityForces().SetAcceleration(Gravity * FVec3(0, 0, -1));
 		}
 
