Chaos
- wire ContactModifier mass and inertia scales through to the solver
- separate user mass and inertia scales from shock propagation scales in SolverBody
- remove unused mass conditioning in Joint Solver and hook up to SolverBody shock prop scale
- add tests for inertia modification and reset

diff --git a/Engine/Source/Programs/HeadlessChaos/Private/HeadlessChaosTestContactModification.cpp b/Engine/Source/Programs/HeadlessChaos/Private/HeadlessChaosTestContactModification.cpp
--- a/Engine/Source/Programs/HeadlessChaos/Private/HeadlessChaosTestContactModification.cpp
+++ b/Engine/Source/Programs/HeadlessChaos/Private/HeadlessChaosTestContactModification.cpp
@@ -467,6 +467,10 @@ namespace ChaosTest
 					int32 NumContacts = PairModifier.GetNumContacts();
 					for (int32 PointIdx = 0; PointIdx < NumContacts; ++PointIdx)
 					{
+						// Make sure that values are not held between frames
+						// @todo(chaos): this test actually has zero restitution
+						//EXPECT_GT(PairModifier.GetRestitution(), 0);
+
 						// Remove restitution
 						PairModifier.ModifyRestitution(0);
 					}
@@ -627,6 +631,10 @@ namespace ChaosTest
 				if ((UniqueIndices[0] == Idx0 && UniqueIndices[1] == Idx1) ||
 					(UniqueIndices[0] == Idx1 && UniqueIndices[1] == Idx0))
 				{
+					// Make sure that values are not held between frames
+					EXPECT_LT(PairModifier.GetDynamicFriction(), 1);
+					EXPECT_LT(PairModifier.GetStaticFriction(), 1);
+
 					PairModifier.ModifyDynamicFriction(1);
 					PairModifier.ModifyStaticFriction(1);
 				}
@@ -914,4 +922,81 @@ namespace ChaosTest
 		Solver->UnregisterObject(FloorProxy);
 		Module->DestroySolver(Solver);
 	}
+
+
+	// Drop a dynamic cube onto a kinematic cube with an offset so the dynamic cube would start to
+	// rotate, except that we have set the inertia scale to zero so it shouldn't rotate. It should actually 
+	// just stop when it hits, and never tip off.
+	GTEST_TEST(AllTraits, ContactModification_ModifyParticleInertiaZero)
+	{
+		FChaosSolversModule* Module = FChaosSolversModule::GetModule();
+		auto* Solver = Module->CreateSolver(nullptr, /*AsyncDt=*/-1);
+		InitSolverSettings(Solver);
+		Solver->SetThreadingMode_External(EThreadingModeTemp::SingleThread);
+
+		// We use contact mod to rotate cube and maintain angular velocity of 0.
+
+		// simulated cube dropped from just above the floor
+		FSingleParticlePhysicsProxy* ModifiedCubeProxy = FSingleParticlePhysicsProxy::Create(Chaos::FPBDRigidParticle::CreateParticle());
+		auto& ModifiedCubeParticle = ModifiedCubeProxy->GetGameThreadAPI();
+		auto ModifiedCubeGeom = TSharedPtr<FImplicitObject, ESPMode::ThreadSafe>(new TBox<FReal, 3>(FVec3(-100), FVec3(100)));
+		ModifiedCubeParticle.SetGeometry(ModifiedCubeGeom);
+		Solver->RegisterObject(ModifiedCubeProxy);
+		ModifiedCubeParticle.SetGravityEnabled(true);
+		ModifiedCubeParticle.SetX(FVec3(0, 0, 150));
+		SetCubeInertiaTensor(ModifiedCubeParticle, /*Dimension=*/200, /*Mass=*/1);
+		ChaosTest::SetParticleSimDataToCollide({ ModifiedCubeProxy->GetParticle_LowLevel() });
+
+		// static floor at origin, occupying Z = [-100,0]
+		FSingleParticlePhysicsProxy* FloorProxy = FSingleParticlePhysicsProxy::Create(Chaos::FGeometryParticle::CreateParticle());
+		auto& FloorParticle = FloorProxy->GetGameThreadAPI();
+		auto FloorGeom = TSharedPtr<FImplicitObject, ESPMode::ThreadSafe>(new TBox<FReal, 3>(FVec3(-500, -500, -100), FVec3(-90, 500, 0)));
+		FloorParticle.SetGeometry(FloorGeom);
+		Solver->RegisterObject(FloorProxy);
+		FloorParticle.SetX(FVec3(0, 0, 0));
+		ChaosTest::SetParticleSimDataToCollide({ FloorProxy->GetParticle_LowLevel() });
+
+		FContactModificationTestCallback* Callback = Solver->CreateAndRegisterSimCallbackObject_External<FContactModificationTestCallback>(true);
+		Callback->TestLambda = [](Chaos::FCollisionContactModifier& Modifier)
+		{
+			for (FContactPairModifier& PairModifier : Modifier)
+			{
+				TVec2<FGeometryParticleHandle*> Particles = PairModifier.GetParticlePair();
+				int32 DynamicParticleIdx = (Particles[0]->ObjectState() == EObjectStateType::Dynamic) ? 0 : 1;
+
+				// Make sure that the values were not held over from the last call
+				EXPECT_NEAR(PairModifier.GetInvInertiaScale(0), 1.0, UE_SMALL_NUMBER);
+				EXPECT_NEAR(PairModifier.GetInvInertiaScale(1), 1.0, UE_SMALL_NUMBER);
+
+				PairModifier.ModifyRestitution(0);
+				PairModifier.ModifyInvInertiaScale(0, DynamicParticleIdx);
+			}
+		};
+
+
+		const float Dt = .1f;
+		const int32 Steps = 20;
+		for (int Step = 0; Step < Steps; ++Step)
+		{
+			Solver->AdvanceAndDispatch_External(Dt);
+			Solver->UpdateGameThreadStructures();
+		}
+
+		// Do we have rotation applied in contact mod?
+		EXPECT_NEAR(ModifiedCubeParticle.R().X, 0.0, 0.001);
+		EXPECT_NEAR(ModifiedCubeParticle.R().Y, 0.0, 0.001);
+		EXPECT_NEAR(ModifiedCubeParticle.R().Z, 0.0, 0.001);
+		EXPECT_NEAR(ModifiedCubeParticle.R().W, 1.0, 0.001);
+
+		// Body should be sat on the floor even though it is hanging off the edge
+		EXPECT_NEAR(ModifiedCubeParticle.X().X, 0.0, 0.001);
+		EXPECT_NEAR(ModifiedCubeParticle.X().Y, 0.0, 0.001);
+		EXPECT_NEAR(ModifiedCubeParticle.X().Z, 100.0, 0.001);
+
+		Solver->UnregisterAndFreeSimCallbackObject_External(Callback);
+		Solver->UnregisterObject(ModifiedCubeProxy);
+		Solver->UnregisterObject(FloorProxy);
+		Module->DestroySolver(Solver);
+	}
+
 }
diff --git a/Engine/Source/Runtime/Experimental/Chaos/Private/Chaos/Collision/PBDCollisionSolver.cpp b/Engine/Source/Runtime/Experimental/Chaos/Private/Chaos/Collision/PBDCollisionSolver.cpp
--- a/Engine/Source/Runtime/Experimental/Chaos/Private/Chaos/Collision/PBDCollisionSolver.cpp
+++ b/Engine/Source/Runtime/Experimental/Chaos/Private/Chaos/Collision/PBDCollisionSolver.cpp
@@ -93,17 +93,17 @@ namespace Chaos
 			bool bInvMassUpdated = false;
 			if (Body0.Level() < Body1.Level())
 			{
-				if (Body0.InvMScale() != InvMassScale)
+				if (Body0.ShockPropagationScale() != InvMassScale)
 				{
-					Body0.SetInvMScale(InvMassScale);
+					Body0.SetShockPropagationScale(InvMassScale);
 					bInvMassUpdated = true;
 				}
 			}
 			else
 			{
-				if (Body1.InvMScale() != InvMassScale)
+				if (Body1.ShockPropagationScale() != InvMassScale)
 				{
-					Body1.SetInvMScale(InvMassScale);
+					Body1.SetShockPropagationScale(InvMassScale);
 					bInvMassUpdated = true;
 				}
 			}
diff --git a/Engine/Source/Runtime/Experimental/Chaos/Private/Chaos/Collision/PBDCollisionSolver.h b/Engine/Source/Runtime/Experimental/Chaos/Private/Chaos/Collision/PBDCollisionSolver.h
--- a/Engine/Source/Runtime/Experimental/Chaos/Private/Chaos/Collision/PBDCollisionSolver.h
+++ b/Engine/Source/Runtime/Experimental/Chaos/Private/Chaos/Collision/PBDCollisionSolver.h
@@ -205,11 +205,10 @@ namespace Chaos
 			State.Stiffness = InStiffness;
 		}
 
-		void SetSolverBodies(FSolverBody* SolverBody0, FSolverBody* SolverBody1)
+		void SetSolverBodies(FSolverBody& SolverBody0, FSolverBody& SolverBody1)
 		{
-			// @todo(chaos): change FConstraintSolverBody to take a pointer because it supports null
-			State.SolverBodies[0] = (SolverBody0 != nullptr) ? FConstraintSolverBody(*SolverBody0) : FConstraintSolverBody();
-			State.SolverBodies[1] = (SolverBody1 != nullptr) ? FConstraintSolverBody(*SolverBody1) : FConstraintSolverBody();
+			State.SolverBodies[0] = FConstraintSolverBody(SolverBody0);
+			State.SolverBodies[1] = FConstraintSolverBody(SolverBody1);
 		}
 
 		void ResetSolverBodies()
diff --git a/Engine/Source/Runtime/Experimental/Chaos/Private/Chaos/Collision/SolverCollisionContainer.cpp b/Engine/Source/Runtime/Experimental/Chaos/Private/Chaos/Collision/SolverCollisionContainer.cpp
--- a/Engine/Source/Runtime/Experimental/Chaos/Private/Chaos/Collision/SolverCollisionContainer.cpp
+++ b/Engine/Source/Runtime/Experimental/Chaos/Private/Chaos/Collision/SolverCollisionContainer.cpp
@@ -76,7 +76,7 @@ namespace Chaos
 		}
 
 		/**
-		 * @brief Initialize the data required for the solver, and bind to the bodies
+		 * @brief Initialize the data required for the solver and bind to the bodies
 		*/
 		void GatherInput(
 			const FReal Dt,
@@ -126,7 +126,11 @@ namespace Chaos
 
 			Solver.SetStiffness(FSolverReal(Constraint->GetStiffness()));
 
-			Solver.SetSolverBodies(Body0, Body1);
+			Solver.SetSolverBodies(*Body0, *Body1);
+			Solver.SolverBody0().SetInvMScale(Constraint->GetInvMassScale0());
+			Solver.SolverBody0().SetInvIScale(Constraint->GetInvInertiaScale0());
+			Solver.SolverBody1().SetInvMScale(Constraint->GetInvMassScale1());
+			Solver.SolverBody1().SetInvIScale(Constraint->GetInvInertiaScale1());
 
 			if (!bChaos_PBDCollisionSolver_VectorRegister)
 			{
diff --git a/Engine/Source/Runtime/Experimental/Chaos/Private/Chaos/Joint/PBDJointCachedSolverGaussSeidel.cpp b/Engine/Source/Runtime/Experimental/Chaos/Private/Chaos/Joint/PBDJointCachedSolverGaussSeidel.cpp
--- a/Engine/Source/Runtime/Experimental/Chaos/Private/Chaos/Joint/PBDJointCachedSolverGaussSeidel.cpp
+++ b/Engine/Source/Runtime/Experimental/Chaos/Private/Chaos/Joint/PBDJointCachedSolverGaussSeidel.cpp
@@ -92,11 +92,10 @@ void FPBDJointCachedSolver::Update(
 
 void FPBDJointCachedSolver::UpdateMass0()
 {
-	// @todo(chaos): this needs to recache all the mass-dependent state on the axis data etc
-	if ((ConditionedInvMs[0] > 0) && (InvMScales[0] > 0))
+	if (Body0().IsDynamic())
 	{
-		InvMs[0] = InvMScales[0] * ConditionedInvMs[0];
-		InvIs[0] = Utilities::ComputeWorldSpaceInertia(CurrentQs[0], InvMScales[0] * ConditionedInvILs[0]);
+		InvMs[0] = Body0().InvM();
+		InvIs[0] = Utilities::ComputeWorldSpaceInertia(CurrentQs[0], Body0().InvILocal());
 	}
 	else
 	{
@@ -107,11 +106,10 @@ void FPBDJointCachedSolver::UpdateMass0()
 
 void FPBDJointCachedSolver::UpdateMass1()
 {
-	// @todo(chaos): this needs to recache all the mass-dependent state on the axis data etc
-	if ((ConditionedInvMs[1] > 0) && (InvMScales[1] > 0))
+	if (Body1().IsDynamic())
 	{
-		InvMs[1] = InvMScales[1] * ConditionedInvMs[1];
-		InvIs[1] = Utilities::ComputeWorldSpaceInertia(CurrentQs[1], InvMScales[1] * ConditionedInvILs[1]);
+		InvMs[1] = Body1().InvM();
+		InvIs[1] = Utilities::ComputeWorldSpaceInertia(CurrentQs[1], Body1().InvILocal());
 	}
 	else
 	{
@@ -120,18 +118,18 @@ void FPBDJointCachedSolver::UpdateMass1()
 	}
 }
 
-void FPBDJointCachedSolver::SetInvMassScales(const FReal InvMScale0, const FReal InvMScale1, const FReal Dt)
+void FPBDJointCachedSolver::SetShockPropagationScales(const FReal InvMScale0, const FReal InvMScale1, const FReal Dt)
 {
 	bool bNeedsUpdate = false;
-	if (InvMScales[0] != InvMScale0)
+	if (Body0().ShockPropagationScale() != InvMScale0)
 	{
-		InvMScales[0] = InvMScale0;
+		Body0().SetShockPropagationScale(InvMScale0);
 		UpdateMass0();
 		bNeedsUpdate = true;
 	}
-	if (InvMScales[1] != InvMScale1)
+	if (Body1().ShockPropagationScale() != InvMScale1)
 	{
-		InvMScales[1] = InvMScale1;
+		Body1().SetShockPropagationScale(InvMScale1);
 		UpdateMass1();
 		bNeedsUpdate = true;
 	}
@@ -159,11 +157,6 @@ void FPBDJointCachedSolver::SetInvMassScales(const FReal InvMScale0, const FReal
 	}
 }
 
-void FPBDJointCachedSolver::EnableProjection()
-{
-	Body0().SetInvMScale(0);
-}
-
 /** Main init function to cache datas that could be reused in the apply */
 
 void FPBDJointCachedSolver::Init(
@@ -183,11 +176,11 @@ void FPBDJointCachedSolver::Init(
 	// \todo(chaos): joint should support parent/child in either order
 	SolverBodies[0].SetInvMScale(JointSettings.ParentInvMassScale);
 	SolverBodies[1].SetInvMScale(FReal(1));
+	SolverBodies[0].SetInvIScale(JointSettings.ParentInvMassScale);
+	SolverBodies[1].SetInvIScale(FReal(1));
+	SolverBodies[0].SetShockPropagationScale(FReal(1));
+	SolverBodies[1].SetShockPropagationScale(FReal(1));
 
-	InvMScales[0] = FReal(1);
-	InvMScales[1] = FReal(1);
-	FPBDJointUtilities::ConditionInverseMassAndInertia(Body0().InvM(), Body1().InvM(), Body0().InvILocal(), Body1().InvILocal(), SolverSettings.MinParentMassRatio, SolverSettings.MaxInertiaRatio, ConditionedInvMs[0], ConditionedInvMs[1], ConditionedInvILs[0], ConditionedInvILs[1]);
-	
 	NetLinearImpulse = FVec3(0);
 	NetAngularImpulse = FVec3(0);
 
@@ -245,14 +238,11 @@ void FPBDJointCachedSolver::InitProjection(
 
 		ConnectorRs[1].EnforceShortestArcWith(ConnectorRs[0]);
 
-		InvMScales[0] = 0;
-		InvMScales[1] = 1;
-
 		// Fake spherical inertia for angular projection (avoid cost of recomputing inertia at current world space rotation)
 		InvMs[0] = 0;
 		InvIs[0] = FMatrix33(0, 0, 0);
-		InvMs[1] = ConditionedInvMs[1];
-		InvIs[1] = FMatrix33::FromDiagonal(FVec3(ConditionedInvILs[1].GetMax()));
+		InvMs[1] = Body1().InvM();
+		InvIs[1] = FMatrix33::FromDiagonal(FVec3(Body1().InvILocal().GetMin()));
 
 		if(bHasLinearProjection)
 		{
diff --git a/Engine/Source/Runtime/Experimental/Chaos/Private/Chaos/Joint/PBDJointSolverGaussSeidel.cpp b/Engine/Source/Runtime/Experimental/Chaos/Private/Chaos/Joint/PBDJointSolverGaussSeidel.cpp
--- a/Engine/Source/Runtime/Experimental/Chaos/Private/Chaos/Joint/PBDJointSolverGaussSeidel.cpp
+++ b/Engine/Source/Runtime/Experimental/Chaos/Private/Chaos/Joint/PBDJointSolverGaussSeidel.cpp
@@ -232,7 +232,7 @@ namespace Chaos
 		}
 	}
 
-	void FPBDJointSolver::SetInvMassScales(const FReal InvMScale0, const FReal InvMScale1)
+	void FPBDJointSolver::SetShockPropagationScales(const FReal InvMScale0, const FReal InvMScale1)
 	{
 		if (InvMScales[0] != InvMScale0)
 		{
diff --git a/Engine/Source/Runtime/Experimental/Chaos/Private/Chaos/PBDJointConstraints.cpp b/Engine/Source/Runtime/Experimental/Chaos/Private/Chaos/PBDJointConstraints.cpp
--- a/Engine/Source/Runtime/Experimental/Chaos/Private/Chaos/PBDJointConstraints.cpp
+++ b/Engine/Source/Runtime/Experimental/Chaos/Private/Chaos/PBDJointConstraints.cpp
@@ -1092,7 +1092,6 @@ namespace Chaos
 			PreparePhase3Serial(Dt, SolverData);
 		}
 
-		// UpdateConstraintProjection(It, NumIts, SolverData);
 		for (int32 ConstraintIndex : SolverData.GetConstraintIndices(ContainerId))
 		{
 			ApplyPhase3Single(Dt, ConstraintIndex, It, NumIts);
@@ -1101,26 +1100,6 @@ namespace Chaos
 		return true;
 	}
 
-	void FPBDJointConstraints::UpdateConstraintProjection(const int32 It, const int32 NumIts, const FPBDIslandSolverData& SolverData)
-	{
-		const bool bUseProjection = It == NumIts - 1;
-		if (bUseProjection)
-		{
-			for (int32 ConstraintIndex : SolverData.GetConstraintIndices(ContainerId))
-			{
-				const FPBDJointSettings& JointSettings = ConstraintSettings[ConstraintIndex];
-				if (Settings.bUseLinearSolver)
-				{
-					if (JointSettings.bProjectionEnabled)
-					{
-						FPBDJointCachedSolver& Solver = CachedConstraintSolvers[ConstraintIndex];
-						Solver.EnableProjection();
-					}
-				}
-			}
-		}
-	}
-
 	//////////////////////////////////////////////////////////////////////////
 	//
 	// End Graph API Solver.
@@ -1245,7 +1224,7 @@ namespace Chaos
 
 			// Set parent inverse mass scale based on current shock propagation state
 			const FReal ShockPropagationInvMassScale = CalculateShockPropagationInvMassScale(Solver.Body0(), Solver.Body1(), JointSettings, It, NumIts);
-			Solver.SetInvMassScales(ShockPropagationInvMassScale, FReal(1), Dt);
+			Solver.SetShockPropagationScales(ShockPropagationInvMassScale, FReal(1), Dt);
 
 			const FReal IterationStiffness = CalculateIterationStiffness(It, NumIts);
 			for (int32 PairIt = 0; PairIt < NumPairIts; ++PairIt)
@@ -1285,7 +1264,7 @@ namespace Chaos
 				}
 				else
 				{
-					Solver.SetInvMassScales(ShockPropagationInvMassScale, FReal(1));
+					Solver.SetShockPropagationScales(ShockPropagationInvMassScale, FReal(1));
 				}
 
 				Solver.ApplyConstraints(Dt, IterationStiffness, Settings, JointSettings);
@@ -1357,7 +1336,7 @@ namespace Chaos
 			
 			// Set parent inverse mass scale based on current shock propagation state
 			const FReal ShockPropagationInvMassScale = CalculateShockPropagationInvMassScale(Solver.Body0(), Solver.Body1(), JointSettings, It, NumIts);
-			Solver.SetInvMassScales(ShockPropagationInvMassScale, FReal(1), Dt);
+			Solver.SetShockPropagationScales(ShockPropagationInvMassScale, FReal(1), Dt);
 
 			const FReal IterationStiffness = CalculateIterationStiffness(It, NumIts);
 			Solver.ApplyVelocityConstraints(Dt, IterationStiffness, Settings, JointSettings);
@@ -1376,7 +1355,7 @@ namespace Chaos
 
 			// Set parent inverse mass scale based on current shock propagation state
 			const FReal ShockPropagationInvMassScale = CalculateShockPropagationInvMassScale(Solver.Body0(), Solver.Body1(), JointSettings, It, NumIts);
-			Solver.SetInvMassScales(ShockPropagationInvMassScale, FReal(1));
+			Solver.SetShockPropagationScales(ShockPropagationInvMassScale, FReal(1));
 
 			const FReal IterationStiffness = CalculateIterationStiffness(It, NumIts);
 			Solver.ApplyVelocityConstraints(Dt, IterationStiffness, Settings, JointSettings);
diff --git a/Engine/Source/Runtime/Experimental/Chaos/Private/Chaos/PBDSuspensionConstraints.cpp b/Engine/Source/Runtime/Experimental/Chaos/Private/Chaos/PBDSuspensionConstraints.cpp
--- a/Engine/Source/Runtime/Experimental/Chaos/Private/Chaos/PBDSuspensionConstraints.cpp
+++ b/Engine/Source/Runtime/Experimental/Chaos/Private/Chaos/PBDSuspensionConstraints.cpp
@@ -175,7 +175,7 @@ namespace Chaos
 				Body1->SetX(PosBody1);
 				check(Body1->InvM() == 0);
 
-				Solver->SetSolverBodies(Body0, Body1);
+				Solver->SetSolverBodies(*Body0, *Body1);
 				Solver->SetNumManifoldPoints(1);
 				Solver->SetFriction(0, 0, 0);
 
@@ -230,7 +230,7 @@ namespace Chaos
 			}
 
 			ConstraintSolverBodies[ConstraintIndex] = nullptr;
-			CollisionSolvers[ConstraintIndex]->SetSolverBodies(nullptr, nullptr);
+			CollisionSolvers[ConstraintIndex]->ResetSolverBodies();
 		}
 	}
 
diff --git a/Engine/Source/Runtime/Experimental/Chaos/Public/Chaos/Collision/PBDCollisionConstraint.h b/Engine/Source/Runtime/Experimental/Chaos/Public/Chaos/Collision/PBDCollisionConstraint.h
--- a/Engine/Source/Runtime/Experimental/Chaos/Public/Chaos/Collision/PBDCollisionConstraint.h
+++ b/Engine/Source/Runtime/Experimental/Chaos/Public/Chaos/Collision/PBDCollisionConstraint.h
@@ -104,11 +104,11 @@ namespace Chaos
 	{
 	public:
 		FPBDCollisionConstraintMaterial()
-			: MaterialStaticFriction(0)
-			, MaterialDynamicFriction(0)
+			: MaterialDynamicFriction(0)
+			, MaterialStaticFriction(0)
 			, MaterialRestitution(0)
-			, Friction(0)
-			, AngularFriction(0)
+			, DynamicFriction(0)
+			, StaticFriction(0)
 			, Restitution(0)
 			, RestitutionPadding(0)
 			, RestitutionThreshold(0)
@@ -120,13 +120,13 @@ namespace Chaos
 		}
 
 		// Material properties pulled from the materials of the two shapes involved in the contact
-		FReal MaterialStaticFriction;
 		FReal MaterialDynamicFriction;
+		FReal MaterialStaticFriction;
 		FReal MaterialRestitution;
 
 		// Final material properties (post modifier) used by the solver. These get reset every frame to the material values above
-		FReal Friction;			// @todo(chaos): rename DynamicFriction
-		FReal AngularFriction;	// @todo(chaos): rename StaticFriction
+		FReal DynamicFriction;
+		FReal StaticFriction;
 		FReal Restitution;
 
 		FReal RestitutionPadding; // For StandardPBD implementation of resitution, we pad constraints on initial contact to enforce outward velocity
@@ -139,9 +139,13 @@ namespace Chaos
 		// Reset the material properties to those pulled from the shape's materials (i.e., back to the state before any contact modification)
 		void ResetMaterialModifications()
 		{
-			Friction = MaterialDynamicFriction;
-			AngularFriction = MaterialStaticFriction;
+			DynamicFriction = MaterialDynamicFriction;
+			StaticFriction = MaterialStaticFriction;
 			Restitution = MaterialRestitution;
+			InvMassScale0 = FReal(1);
+			InvMassScale1 = FReal(1);
+			InvInertiaScale0 = FReal(1);
+			InvInertiaScale1 = FReal(1);
 		}
 
 		// @todo(chaos): remove this
@@ -392,11 +396,11 @@ namespace Chaos
 		void SetRestitutionPadding(const FReal InRestitutionPadding) { Material.RestitutionPadding = InRestitutionPadding; }
 		FReal GetRestitutionPadding() const { return Material.RestitutionPadding; }
 
-		void SetStaticFriction(const FReal InStaticFriction) { Material.AngularFriction = InStaticFriction; }
-		FReal GetStaticFriction() const { return FMath::Max(Material.AngularFriction, Material.Friction); }
+		void SetStaticFriction(const FReal InStaticFriction) { Material.StaticFriction = InStaticFriction; }
+		FReal GetStaticFriction() const { return FMath::Max(Material.StaticFriction, Material.DynamicFriction); }
 
-		void SetDynamicFriction(const FReal InDynamicFriction) { Material.Friction = InDynamicFriction; }
-		FReal GetDynamicFriction() const { return Material.Friction; }
+		void SetDynamicFriction(const FReal InDynamicFriction) { Material.DynamicFriction = InDynamicFriction; }
+		FReal GetDynamicFriction() const { return Material.DynamicFriction; }
 
 		EContactShapesType GetShapesType() const { return ShapesType; }
 
diff --git a/Engine/Source/Runtime/Experimental/Chaos/Public/Chaos/Evolution/SolverBody.h b/Engine/Source/Runtime/Experimental/Chaos/Public/Chaos/Evolution/SolverBody.h
--- a/Engine/Source/Runtime/Experimental/Chaos/Public/Chaos/Evolution/SolverBody.h
+++ b/Engine/Source/Runtime/Experimental/Chaos/Public/Chaos/Evolution/SolverBody.h
@@ -465,30 +465,43 @@ namespace Chaos
 		inline const FSolverBody& SolverBody() const { check(IsValid()); return *Body; }
 
 		/**
-		 * @brief A scale applied to both inverse mass and inverse inertia
+		 * @brief A scale applied to inverse mass
 		*/
-		inline FSolverReal InvMScale() const { return State.InvMassScale; }
-		inline void SetInvMScale(FReal InInvMassScale) { State.InvMassScale = FSolverReal(InInvMassScale); }
+		inline FSolverReal InvMScale() const { return State.InvMScale; }
+		inline void SetInvMScale(FReal InValue) { State.InvMScale = FSolverReal(InValue); }
 
 		/**
-		 * @brief The scaled inverse mass
+		 * @brief A scale applied to inverse inertia
 		*/
-		FSolverReal InvM() const { return State.InvMassScale * Body->InvM(); }
+		inline FSolverReal InvIScale() const { return State.InvIScale; }
+		inline void SetInvIScale(FReal InValue) { State.InvIScale = FSolverReal(InValue); }
+
+
+		/**
+		 * @brief Shock propagation mass and inertia scaling (0 for infinite mass, 1 for no effect)
+		*/
+		inline FSolverReal ShockPropagationScale() const { return State.ShockPropagationScale; }
+		inline void SetShockPropagationScale(FReal InValue) { State.ShockPropagationScale = FSolverReal(InValue); }
+
+		/**
+		 * @brief The net scaled inverse mass
+		*/
+		FSolverReal InvM() const { return (State.ShockPropagationScale * State.InvMScale) * Body->InvM(); }
 
 		/**
-		 * @brief The scaled inverse inertia
+		 * @brief The net scaled inverse inertia
 		*/
-		FSolverMatrix33 InvI() const { return State.InvMassScale * Body->InvI(); }
+		FSolverMatrix33 InvI() const { return (State.ShockPropagationScale * State.InvIScale) * Body->InvI(); }
 
 		/**
-		 * @brief The scaled local space inverse inertia
+		 * @brief The net scaled local space inverse inertia
 		*/
-		FSolverVec3 InvILocal() const { return State.InvMassScale * Body->InvILocal(); }
+		FSolverVec3 InvILocal() const { return (State.ShockPropagationScale * State.InvIScale) * Body->InvILocal(); }
 
 		/**
-		 * @brief Whether the body is dynamic (i.e., has a finite mass) after InvMassScale is applied
+		 * @brief Whether the body is dynamic (i.e., has a finite mass) after scaling is applied
 		*/
-		inline bool IsDynamic() const { return (Body->InvM() != 0) && (InvMScale() != 0); }
+		inline bool IsDynamic() const { return (InvM() > TNumericLimits<FSolverReal>::Min()); }
 
 		//
 		// From here all methods just forward to the FSolverBody
@@ -524,9 +537,13 @@ namespace Chaos
 		struct FState
 		{
 			FState() 
-				: InvMassScale(FSolverReal(1))
+				: InvMScale(FSolverReal(1))
+				, InvIScale(FSolverReal(1))
+				, ShockPropagationScale(FSolverReal(1))
 			{}
-			FSolverReal InvMassScale;
+			FSolverReal InvMScale;
+			FSolverReal InvIScale;
+			FSolverReal ShockPropagationScale;
 		};
 
 		// The body we decorate
diff --git a/Engine/Source/Runtime/Experimental/Chaos/Public/Chaos/Joint/PBDJointCachedSolverGaussSeidel.h b/Engine/Source/Runtime/Experimental/Chaos/Public/Chaos/Joint/PBDJointCachedSolverGaussSeidel.h
--- a/Engine/Source/Runtime/Experimental/Chaos/Public/Chaos/Joint/PBDJointCachedSolverGaussSeidel.h
+++ b/Engine/Source/Runtime/Experimental/Chaos/Public/Chaos/Joint/PBDJointCachedSolverGaussSeidel.h
@@ -262,9 +262,7 @@ struct FAxisConstraintDatas
 			const FPBDJointSolverSettings& SolverSettings,
 			const FPBDJointSettings& JointSettings);
 
-		void EnableProjection();
-
-		void SetInvMassScales(
+		void SetShockPropagationScales(
 			const FReal InvMScale0,
 			const FReal InvMScale1,
 			const FReal Dt);
@@ -591,10 +589,7 @@ struct FAxisConstraintDatas
 		FVec3 InitConnectorXs[MaxConstrainedBodies];		// World-space joint connector positions
 		FRotation3 InitConnectorRs[MaxConstrainedBodies];	// World-space joint connector rotations
 
-		// Conditioned InvM and InvI
-		FReal InvMScales[MaxConstrainedBodies];
-		FReal ConditionedInvMs[MaxConstrainedBodies];
-		FVec3 ConditionedInvILs[MaxConstrainedBodies];		// Local-space
+		// Inverse Mass and Inertia
 		FReal InvMs[MaxConstrainedBodies];
 		FMatrix33 InvIs[MaxConstrainedBodies];				// World-space
 
diff --git a/Engine/Source/Runtime/Experimental/Chaos/Public/Chaos/Joint/PBDJointSolverGaussSeidel.h b/Engine/Source/Runtime/Experimental/Chaos/Public/Chaos/Joint/PBDJointSolverGaussSeidel.h
--- a/Engine/Source/Runtime/Experimental/Chaos/Public/Chaos/Joint/PBDJointSolverGaussSeidel.h
+++ b/Engine/Source/Runtime/Experimental/Chaos/Public/Chaos/Joint/PBDJointSolverGaussSeidel.h
@@ -183,7 +183,7 @@ namespace Chaos
 			const FPBDJointSolverSettings& SolverSettings,
 			const FPBDJointSettings& JointSettings);
 
-		void SetInvMassScales(
+		void SetShockPropagationScales(
 			const FReal InvMScale0, 
 			const FReal InvMScale1);
 
diff --git a/Engine/Source/Runtime/Experimental/Chaos/Public/Chaos/PBDJointConstraints.h b/Engine/Source/Runtime/Experimental/Chaos/Public/Chaos/PBDJointConstraints.h
--- a/Engine/Source/Runtime/Experimental/Chaos/Public/Chaos/PBDJointConstraints.h
+++ b/Engine/Source/Runtime/Experimental/Chaos/Public/Chaos/PBDJointConstraints.h
@@ -317,8 +317,6 @@ namespace Chaos
 		void ApplyBreakThreshold(const FReal Dt, int32 ConstraintIndex, const FVec3& LinearImpulse, const FVec3& AngularImpulse);
 		void ApplyPlasticityLimits(const int32 ConstraintIndex);
 
-		void UpdateConstraintProjection(const int32 It, const int32 NumIts, const FPBDIslandSolverData& SolverData);
-
 		FPBDJointSolverSettings Settings;
 
 		TArray<FPBDJointSettings> ConstraintSettings;
