Optimize BVH by increasing the number of children by leaf
- Add BVH Overlap unit test

diff --git a/Engine/Source/Programs/HeadlessChaos/Private/HeadlessChaosTestOverlap.cpp b/Engine/Source/Programs/HeadlessChaos/Private/HeadlessChaosTestOverlap.cpp
--- a/Engine/Source/Programs/HeadlessChaos/Private/HeadlessChaosTestOverlap.cpp
+++ b/Engine/Source/Programs/HeadlessChaos/Private/HeadlessChaosTestOverlap.cpp
@@ -17,24 +17,24 @@ namespace ChaosTest
 		// Trimesh is simple pyramid
 		using namespace Chaos;
 		FTriangleMeshImplicitObject::ParticlesType TrimeshParticles(
-		{
-			{-10.0, 0.0, 0.0},
-			{10.0, 0.0, 0.0},
-			{0.0, 10.0, 0.0},
-			{0.0, -10.0, 0.0},
-			{0.0,  0.0, 10.0},
-		});
+			{
+				{-10.0, 0.0, 0.0},
+				{10.0, 0.0, 0.0},
+				{0.0, 10.0, 0.0},
+				{0.0, -10.0, 0.0},
+				{0.0,  0.0, 10.0},
+			});
 
 		TArray<TVec3<int32>> Indices(
-		{
-			{0, 1, 2},
-			{0, 3, 1},
-			{0, 1, 2},
-			{0, 2, 4},
-			{1, 2, 4},
-			{0, 3, 4},
-			{1, 3, 4}
-		});
+			{
+				{0, 1, 2},
+				{0, 3, 1},
+				{0, 1, 2},
+				{0, 2, 4},
+				{1, 2, 4},
+				{0, 3, 4},
+				{1, 3, 4}
+			});
 
 		TArray<uint16> Materials;
 		for (int32 i = 0; i < Indices.Num(); ++i)
@@ -43,172 +43,232 @@ namespace ChaosTest
 		}
 
 		TUniquePtr<FTriangleMeshImplicitObject> TriangleMesh = MakeUnique<FTriangleMeshImplicitObject>(MoveTemp(TrimeshParticles), MoveTemp(Indices), MoveTemp(Materials));
-		
-		// Capsule test
-		const FVec3 X1 = { 0.0, 0.0, -2.0 };
-		const FVec3 X2 = { 0.0, 0.0, 2.0 };
-		const FReal Radius = 1.0;
-		const FCapsule Capsule = FCapsule(X1, X2, Radius);
-
-		{
-			const FRigidTransform3 QueryTM(FVec3(0.0, 0.0, 0.0), FQuat::Identity);
-			bool bResult = TriangleMesh->OverlapGeom(Capsule, QueryTM, 0.0);
-			EXPECT_EQ(bResult, true);
-		}
-		{
-			const FRigidTransform3 QueryTM(FVec3(10.0, 10.0, 0.0), FQuat::Identity);
-			bool bResult = TriangleMesh->OverlapGeom(Capsule, QueryTM, 0.0);
-			EXPECT_EQ(bResult, false);
-		}
-		{
-			const FRigidTransform3 QueryTM(FVec3(-10.0, -10.0, 0.0), FQuat::Identity);
-			bool bResult = TriangleMesh->OverlapGeom(Capsule, QueryTM, 0.0);
-			EXPECT_EQ(bResult, false);
-		}
-		{
-			const FRigidTransform3 QueryTM(FVec3(10.0, 0.0, 0.0), FQuat::Identity);
-			bool bResult = TriangleMesh->OverlapGeom(Capsule, QueryTM, 0.0);
-			EXPECT_EQ(bResult, true);
-		}
-		{
-			const FRigidTransform3 QueryTM(FVec3(0.0, 10.0, 0.0), FQuat::Identity);
-			bool bResult = TriangleMesh->OverlapGeom(Capsule, QueryTM, 0.0);
-			EXPECT_EQ(bResult, true);
-		}
-		{
-			const FRigidTransform3 QueryTM(FVec3(0.0, -10.0, 0.0), FQuat::Identity);
-			bool bResult = TriangleMesh->OverlapGeom(Capsule, QueryTM, 0.0);
-			EXPECT_EQ(bResult, true);
-		}
 		{
-			const FRigidTransform3 QueryTM(FVec3(-10.0, 0.0, 0.0), FQuat::Identity);
-			bool bResult = TriangleMesh->OverlapGeom(Capsule, QueryTM, 0.0);
-			EXPECT_EQ(bResult, true);
-		}
-		{
-			const FRigidTransform3 QueryTM(FVec3(0.0, 0.0, 12.0), FQuat::Identity);
-			bool bResult = TriangleMesh->OverlapGeom(Capsule, QueryTM, 0.0);
-			EXPECT_EQ(bResult, true);
-		}
-		{
-			const FRigidTransform3 QueryTM(FVec3(0.0, 0.0, 12.0), FQuat(FVec3(1.0, 0.0, 0.0), 3.1415926 / 4.0));
-			bool bResult = TriangleMesh->OverlapGeom(Capsule, QueryTM, 0.0);
-			EXPECT_EQ(bResult, false);
-		}
-		{
-			// Inside Mesh
-			const FRigidTransform3 QueryTM(FVec3(0.0, 0.0, 4.0), FQuat::Identity);
-			bool bResult = TriangleMesh->OverlapGeom(Capsule, QueryTM, 0.0);
-			EXPECT_EQ(bResult, false);
-		}
+			// Capsule test
+			const FVec3 X1 = { 0.0, 0.0, -2.0 };
+			const FVec3 X2 = { 0.0, 0.0, 2.0 };
+			const FReal Radius = 1.0;
+			const FCapsule Capsule = FCapsule(X1, X2, Radius);
 
-
-		// Box test
-		const TBox<FReal, 3> Box = TBox<FReal, 3>({-1.0, -1.0, -2.0 }, { 1.0, 1.0, 2.0 });
-		{
-			const FRigidTransform3 QueryTM(FVec3(0.0, 0.0, 0.0), FQuat::Identity);
-			bool bResult = TriangleMesh->OverlapGeom(Box, QueryTM, 0.0);
-			EXPECT_EQ(bResult, true);
-		}
-		{
-			const FRigidTransform3 QueryTM(FVec3(10.0, 10.0, 0.0), FQuat::Identity);
-			bool bResult = TriangleMesh->OverlapGeom(Box, QueryTM, 0.0);
-			EXPECT_EQ(bResult, false);
+			{
+				const FRigidTransform3 QueryTM(FVec3(0.0, 0.0, 0.0), FQuat::Identity);
+				bool bResult = TriangleMesh->OverlapGeom(Capsule, QueryTM, 0.0);
+				EXPECT_EQ(bResult, true);
+			}
+			{
+				const FRigidTransform3 QueryTM(FVec3(10.0, 10.0, 0.0), FQuat::Identity);
+				bool bResult = TriangleMesh->OverlapGeom(Capsule, QueryTM, 0.0);
+				EXPECT_EQ(bResult, false);
+			}
+			{
+				const FRigidTransform3 QueryTM(FVec3(-10.0, -10.0, 0.0), FQuat::Identity);
+				bool bResult = TriangleMesh->OverlapGeom(Capsule, QueryTM, 0.0);
+				EXPECT_EQ(bResult, false);
+			}
+			{
+				const FRigidTransform3 QueryTM(FVec3(10.0, 0.0, 0.0), FQuat::Identity);
+				bool bResult = TriangleMesh->OverlapGeom(Capsule, QueryTM, 0.0);
+				EXPECT_EQ(bResult, true);
+			}
+			{
+				const FRigidTransform3 QueryTM(FVec3(0.0, 10.0, 0.0), FQuat::Identity);
+				bool bResult = TriangleMesh->OverlapGeom(Capsule, QueryTM, 0.0);
+				EXPECT_EQ(bResult, true);
+			}
+			{
+				const FRigidTransform3 QueryTM(FVec3(0.0, -10.0, 0.0), FQuat::Identity);
+				bool bResult = TriangleMesh->OverlapGeom(Capsule, QueryTM, 0.0);
+				EXPECT_EQ(bResult, true);
+			}
+			{
+				const FRigidTransform3 QueryTM(FVec3(-10.0, 0.0, 0.0), FQuat::Identity);
+				bool bResult = TriangleMesh->OverlapGeom(Capsule, QueryTM, 0.0);
+				EXPECT_EQ(bResult, true);
+			}
+			{
+				const FRigidTransform3 QueryTM(FVec3(0.0, 0.0, 12.0), FQuat::Identity);
+				bool bResult = TriangleMesh->OverlapGeom(Capsule, QueryTM, 0.0);
+				EXPECT_EQ(bResult, true);
+			}
+			{
+				const FRigidTransform3 QueryTM(FVec3(0.0, 0.0, 12.0), FQuat(FVec3(1.0, 0.0, 0.0), 3.1415926 / 4.0));
+				bool bResult = TriangleMesh->OverlapGeom(Capsule, QueryTM, 0.0);
+				EXPECT_EQ(bResult, false);
+			}
+			{
+				// Inside Mesh
+				const FRigidTransform3 QueryTM(FVec3(0.0, 0.0, 4.0), FQuat::Identity);
+				bool bResult = TriangleMesh->OverlapGeom(Capsule, QueryTM, 0.0);
+				EXPECT_EQ(bResult, false);
+			}
 		}
 		{
-			const FRigidTransform3 QueryTM(FVec3(-10.0, -10.0, 0.0), FQuat::Identity);
-			bool bResult = TriangleMesh->OverlapGeom(Box, QueryTM, 0.0);
-			EXPECT_EQ(bResult, false);
-		}
-		{
-			const FRigidTransform3 QueryTM(FVec3(10.0, 0.0, 0.0), FQuat::Identity);
-			bool bResult = TriangleMesh->OverlapGeom(Box, QueryTM, 0.0);
-			EXPECT_EQ(bResult, true);
-		}
-		{
-			const FRigidTransform3 QueryTM(FVec3(0.0, 10.0, 0.0), FQuat::Identity);
-			bool bResult = TriangleMesh->OverlapGeom(Box, QueryTM, 0.0);
-			EXPECT_EQ(bResult, true);
-		}
-		{
-			const FRigidTransform3 QueryTM(FVec3(0.0, -10.0, 0.0), FQuat::Identity);
-			bool bResult = TriangleMesh->OverlapGeom(Box, QueryTM, 0.0);
-			EXPECT_EQ(bResult, true);
-		}
-		{
-			const FRigidTransform3 QueryTM(FVec3(-10.0, 0.0, 0.0), FQuat::Identity);
-			bool bResult = TriangleMesh->OverlapGeom(Box, QueryTM, 0.0);
-			EXPECT_EQ(bResult, true);
-		}
-		{
-			const FRigidTransform3 QueryTM(FVec3(0.0, 0.0, 12.0), FQuat::Identity);
-			bool bResult = TriangleMesh->OverlapGeom(Box, QueryTM, 0.0);
-			EXPECT_EQ(bResult, true);
-		}
-		{
-			const FRigidTransform3 QueryTM(FVec3(0.0, 0.0, 12.0), FQuat(FVec3(1.0, 0.0, 0.0), 3.1415926 / 4.0));
-			bool bResult = TriangleMesh->OverlapGeom(Box, QueryTM, 0.0);
-			EXPECT_EQ(bResult, false);
-		}
-		{
-			// Inside Mesh
-			const FRigidTransform3 QueryTM(FVec3(0.0, 0.0, 4.0), FQuat::Identity);
-			bool bResult = TriangleMesh->OverlapGeom(Box, QueryTM, 0.0);
-			EXPECT_EQ(bResult, false);
+			// Box test
+			const TBox<FReal, 3> Box = TBox<FReal, 3>({ -1.0, -1.0, -2.0 }, { 1.0, 1.0, 2.0 });
+			{
+				const FRigidTransform3 QueryTM(FVec3(0.0, 0.0, 0.0), FQuat::Identity);
+				bool bResult = TriangleMesh->OverlapGeom(Box, QueryTM, 0.0);
+				EXPECT_EQ(bResult, true);
+			}
+			{
+				const FRigidTransform3 QueryTM(FVec3(10.0, 10.0, 0.0), FQuat::Identity);
+				bool bResult = TriangleMesh->OverlapGeom(Box, QueryTM, 0.0);
+				EXPECT_EQ(bResult, false);
+			}
+			{
+				const FRigidTransform3 QueryTM(FVec3(-10.0, -10.0, 0.0), FQuat::Identity);
+				bool bResult = TriangleMesh->OverlapGeom(Box, QueryTM, 0.0);
+				EXPECT_EQ(bResult, false);
+			}
+			{
+				const FRigidTransform3 QueryTM(FVec3(10.0, 0.0, 0.0), FQuat::Identity);
+				bool bResult = TriangleMesh->OverlapGeom(Box, QueryTM, 0.0);
+				EXPECT_EQ(bResult, true);
+			}
+			{
+				const FRigidTransform3 QueryTM(FVec3(0.0, 10.0, 0.0), FQuat::Identity);
+				bool bResult = TriangleMesh->OverlapGeom(Box, QueryTM, 0.0);
+				EXPECT_EQ(bResult, true);
+			}
+			{
+				const FRigidTransform3 QueryTM(FVec3(0.0, -10.0, 0.0), FQuat::Identity);
+				bool bResult = TriangleMesh->OverlapGeom(Box, QueryTM, 0.0);
+				EXPECT_EQ(bResult, true);
+			}
+			{
+				const FRigidTransform3 QueryTM(FVec3(-10.0, 0.0, 0.0), FQuat::Identity);
+				bool bResult = TriangleMesh->OverlapGeom(Box, QueryTM, 0.0);
+				EXPECT_EQ(bResult, true);
+			}
+			{
+				const FRigidTransform3 QueryTM(FVec3(0.0, 0.0, 12.0), FQuat::Identity);
+				bool bResult = TriangleMesh->OverlapGeom(Box, QueryTM, 0.0);
+				EXPECT_EQ(bResult, true);
+			}
+			{
+				const FRigidTransform3 QueryTM(FVec3(0.0, 0.0, 12.0), FQuat(FVec3(1.0, 0.0, 0.0), 3.1415926 / 4.0));
+				bool bResult = TriangleMesh->OverlapGeom(Box, QueryTM, 0.0);
+				EXPECT_EQ(bResult, false);
+			}
+			{
+				// Inside Mesh
+				const FRigidTransform3 QueryTM(FVec3(0.0, 0.0, 4.0), FQuat::Identity);
+				bool bResult = TriangleMesh->OverlapGeom(Box, QueryTM, 0.0);
+				EXPECT_EQ(bResult, false);
+			}
 		}
 
-		// Sphere test
-		const TSphere<FReal, 3> Sphere = TSphere<FReal, 3>({ 0.0, 0.0, 0.0 }, 2.0);
-		{
-			const FRigidTransform3 QueryTM(FVec3(0.0, 0.0, 0.0), FQuat::Identity);
-			bool bResult = TriangleMesh->OverlapGeom(Sphere, QueryTM, 0.0);
-			EXPECT_EQ(bResult, true);
-		}
-		{
-			const FRigidTransform3 QueryTM(FVec3(10.0, 10.0, 0.0), FQuat::Identity);
-			bool bResult = TriangleMesh->OverlapGeom(Sphere, QueryTM, 0.0);
-			EXPECT_EQ(bResult, false);
-		}
-		{
-			const FRigidTransform3 QueryTM(FVec3(-10.0, -10.0, 0.0), FQuat::Identity);
-			bool bResult = TriangleMesh->OverlapGeom(Sphere, QueryTM, 0.0);
-			EXPECT_EQ(bResult, false);
-		}
-		{
-			const FRigidTransform3 QueryTM(FVec3(10.0, 0.0, 0.0), FQuat::Identity);
-			bool bResult = TriangleMesh->OverlapGeom(Sphere, QueryTM, 0.0);
-			EXPECT_EQ(bResult, true);
-		}
-		{
-			const FRigidTransform3 QueryTM(FVec3(0.0, 10.0, 0.0), FQuat::Identity);
-			bool bResult = TriangleMesh->OverlapGeom(Sphere, QueryTM, 0.0);
-			EXPECT_EQ(bResult, true);
-		}
-		{
-			const FRigidTransform3 QueryTM(FVec3(0.0, -10.0, 0.0), FQuat::Identity);
-			bool bResult = TriangleMesh->OverlapGeom(Sphere, QueryTM, 0.0);
-			EXPECT_EQ(bResult, true);
-		}
-		{
-			const FRigidTransform3 QueryTM(FVec3(-10.0, 0.0, 0.0), FQuat::Identity);
-			bool bResult = TriangleMesh->OverlapGeom(Sphere, QueryTM, 0.0);
-			EXPECT_EQ(bResult, true);
-		}
-		{
-			const FRigidTransform3 QueryTM(FVec3(0.0, 0.0, 12.0), FQuat::Identity);
-			bool bResult = TriangleMesh->OverlapGeom(Sphere, QueryTM, 0.0);
-			EXPECT_EQ(bResult, true);
-		}
-		{
-			const FRigidTransform3 QueryTM(FVec3(0.0, 0.0, 12.0), FQuat(FVec3(1.0, 0.0, 0.0), 3.1415926 / 4.0));
-			bool bResult = TriangleMesh->OverlapGeom(Sphere, QueryTM, 0.0);
-			EXPECT_EQ(bResult, true);
+		// Box Scaled
+		{
+			TSharedPtr<TBox<FReal, 3>> BigBoxSafe = MakeShared<TBox<FReal, 3>>(FVec3(-100.0, -100.0, -100.0), FVec3(100.0, 100.0, 100.0));
+			FVec3 TriMeshScale = { 0.01, 0.01, 0.02 };
+			TImplicitObjectScaled<TBox<FReal, 3>> ScaledBox = TImplicitObjectScaled<TBox<FReal, 3>>(BigBoxSafe, TriMeshScale);
+			{
+				const FRigidTransform3 QueryTM(FVec3(0.0, 0.0, 0.0), FQuat::Identity);
+				bool bResult = TriangleMesh->OverlapGeom(ScaledBox, QueryTM, 0.0, nullptr, TriMeshScale);
+				EXPECT_EQ(bResult, true);
+			}
+			{
+				const FRigidTransform3 QueryTM(FVec3(10.0, 10.0, 0.0), FQuat::Identity);
+				bool bResult = TriangleMesh->OverlapGeom(ScaledBox, QueryTM, 0.0, nullptr, TriMeshScale);
+				EXPECT_EQ(bResult, false);
+			}
+			{
+				const FRigidTransform3 QueryTM(FVec3(-10.0, -10.0, 0.0), FQuat::Identity);
+				bool bResult = TriangleMesh->OverlapGeom(ScaledBox, QueryTM, 0.0, nullptr, TriMeshScale);
+				EXPECT_EQ(bResult, false);
+			}
+			{
+				const FRigidTransform3 QueryTM(FVec3(10.0, 0.0, 0.0), FQuat::Identity);
+				bool bResult = TriangleMesh->OverlapGeom(ScaledBox, QueryTM, 0.0, nullptr, TriMeshScale);
+				EXPECT_EQ(bResult, true);
+			}
+			{
+				const FRigidTransform3 QueryTM(FVec3(0.0, 10.0, 0.0), FQuat::Identity);
+				bool bResult = TriangleMesh->OverlapGeom(ScaledBox, QueryTM, 0.0, nullptr, TriMeshScale);
+				EXPECT_EQ(bResult, true);
+			}
+			{
+				const FRigidTransform3 QueryTM(FVec3(0.0, -10.0, 0.0), FQuat::Identity);
+				bool bResult = TriangleMesh->OverlapGeom(ScaledBox, QueryTM, 0.0, nullptr, TriMeshScale);
+				EXPECT_EQ(bResult, true);
+			}
+			{
+				const FRigidTransform3 QueryTM(FVec3(-10.0, 0.0, 0.0), FQuat::Identity);
+				bool bResult = TriangleMesh->OverlapGeom(ScaledBox, QueryTM, 0.0, nullptr, TriMeshScale);
+				EXPECT_EQ(bResult, true);
+			}
+			{
+				const FRigidTransform3 QueryTM(FVec3(0.0, 0.0, 12.0), FQuat::Identity);
+				bool bResult = TriangleMesh->OverlapGeom(ScaledBox, QueryTM, 0.0, nullptr, TriMeshScale);
+				EXPECT_EQ(bResult, true);
+			}
+			{
+				const FRigidTransform3 QueryTM(FVec3(0.0, 0.0, 12.0), FQuat(FVec3(1.0, 0.0, 0.0), 3.1415926 / 4.0));
+				bool bResult = TriangleMesh->OverlapGeom(ScaledBox, QueryTM, 0.0, nullptr, TriMeshScale);
+				EXPECT_EQ(bResult, false);
+			}
+			{
+				// Inside Mesh
+				const FRigidTransform3 QueryTM(FVec3(0.0, 0.0, 4.0), FQuat::Identity);
+				bool bResult = TriangleMesh->OverlapGeom(ScaledBox, QueryTM, 0.0, nullptr, TriMeshScale);
+				EXPECT_EQ(bResult, false);
+			}
 		}
 		{
-			// Inside Mesh
-			const FRigidTransform3 QueryTM(FVec3(0.0, 0.0, 4.0), FQuat::Identity);
-			bool bResult = TriangleMesh->OverlapGeom(Sphere, QueryTM, 0.0);
-			EXPECT_EQ(bResult, false);
+			// Sphere test
+			const TSphere<FReal, 3> Sphere = TSphere<FReal, 3>({ 0.0, 0.0, 0.0 }, 2.0);
+			{
+				const FRigidTransform3 QueryTM(FVec3(0.0, 0.0, 0.0), FQuat::Identity);
+				bool bResult = TriangleMesh->OverlapGeom(Sphere, QueryTM, 0.0);
+				EXPECT_EQ(bResult, true);
+			}
+			{
+				const FRigidTransform3 QueryTM(FVec3(10.0, 10.0, 0.0), FQuat::Identity);
+				bool bResult = TriangleMesh->OverlapGeom(Sphere, QueryTM, 0.0);
+				EXPECT_EQ(bResult, false);
+			}
+			{
+				const FRigidTransform3 QueryTM(FVec3(-10.0, -10.0, 0.0), FQuat::Identity);
+				bool bResult = TriangleMesh->OverlapGeom(Sphere, QueryTM, 0.0);
+				EXPECT_EQ(bResult, false);
+			}
+			{
+				const FRigidTransform3 QueryTM(FVec3(10.0, 0.0, 0.0), FQuat::Identity);
+				bool bResult = TriangleMesh->OverlapGeom(Sphere, QueryTM, 0.0);
+				EXPECT_EQ(bResult, true);
+			}
+			{
+				const FRigidTransform3 QueryTM(FVec3(0.0, 10.0, 0.0), FQuat::Identity);
+				bool bResult = TriangleMesh->OverlapGeom(Sphere, QueryTM, 0.0);
+				EXPECT_EQ(bResult, true);
+			}
+			{
+				const FRigidTransform3 QueryTM(FVec3(0.0, -10.0, 0.0), FQuat::Identity);
+				bool bResult = TriangleMesh->OverlapGeom(Sphere, QueryTM, 0.0);
+				EXPECT_EQ(bResult, true);
+			}
+			{
+				const FRigidTransform3 QueryTM(FVec3(-10.0, 0.0, 0.0), FQuat::Identity);
+				bool bResult = TriangleMesh->OverlapGeom(Sphere, QueryTM, 0.0);
+				EXPECT_EQ(bResult, true);
+			}
+			{
+				const FRigidTransform3 QueryTM(FVec3(0.0, 0.0, 12.0), FQuat::Identity);
+				bool bResult = TriangleMesh->OverlapGeom(Sphere, QueryTM, 0.0);
+				EXPECT_EQ(bResult, true);
+			}
+			{
+				const FRigidTransform3 QueryTM(FVec3(0.0, 0.0, 12.0), FQuat(FVec3(1.0, 0.0, 0.0), 3.1415926 / 4.0));
+				bool bResult = TriangleMesh->OverlapGeom(Sphere, QueryTM, 0.0);
+				EXPECT_EQ(bResult, true);
+			}
+			{
+				// Inside Mesh
+				const FRigidTransform3 QueryTM(FVec3(0.0, 0.0, 4.0), FQuat::Identity);
+				bool bResult = TriangleMesh->OverlapGeom(Sphere, QueryTM, 0.0);
+				EXPECT_EQ(bResult, false);
+			}
 		}
 	}
 }
\ No newline at end of file
diff --git a/Engine/Source/Runtime/Experimental/Chaos/Private/Chaos/TriangleMeshImplicitObject.cpp b/Engine/Source/Runtime/Experimental/Chaos/Private/Chaos/TriangleMeshImplicitObject.cpp
--- a/Engine/Source/Runtime/Experimental/Chaos/Private/Chaos/TriangleMeshImplicitObject.cpp
+++ b/Engine/Source/Runtime/Experimental/Chaos/Private/Chaos/TriangleMeshImplicitObject.cpp
@@ -1496,7 +1496,7 @@ void FTriangleMeshImplicitObject::RebuildBVImp(const TArray<TVec3<IdxType>>& Ele
 	{
 		BVEntries.Add({this, Tri});
 	}
-	TreeBVH.Reinitialize(BVEntries);
+	TreeBVH.Reinitialize(BVEntries, 22);
 }
 
 FTriangleMeshImplicitObject::~FTriangleMeshImplicitObject() = default;
diff --git a/Engine/Source/Runtime/Experimental/Chaos/Public/Chaos/AABBTree.h b/Engine/Source/Runtime/Experimental/Chaos/Public/Chaos/AABBTree.h
--- a/Engine/Source/Runtime/Experimental/Chaos/Public/Chaos/AABBTree.h
+++ b/Engine/Source/Runtime/Experimental/Chaos/Public/Chaos/AABBTree.h
@@ -653,8 +653,13 @@ public:
 	}
 
 	template <typename ParticleView>
-	void Reinitialize(const ParticleView& Particles)
+	void Reinitialize(const ParticleView& Particles, int32 InMaxChildrenInLeaf = DefaultMaxChildrenInLeaf, int32 InMaxTreeDepth = DefaultMaxTreeDepth, T InMaxPayloadBounds = DefaultMaxPayloadBounds, int32 InMaxNumToProcess = DefaultMaxNumToProcess, bool bInDynamicTree = false)
 	{
+		bDynamicTree = bInDynamicTree;
+		MaxChildrenInLeaf = InMaxChildrenInLeaf;
+		MaxTreeDepth = InMaxTreeDepth;
+		MaxPayloadBounds = InMaxPayloadBounds;
+		MaxNumToProcess = InMaxNumToProcess;
 		bShouldRebuild = true;
 		GenerateTree(Particles);
 	}
