Add a BVH Overlap unit test

diff --git a/Engine/Source/Programs/HeadlessChaos/Private/HeadlessChaosTestOverlap.cpp b/Engine/Source/Programs/HeadlessChaos/Private/HeadlessChaosTestOverlap.cpp
--- a/Engine/Source/Programs/HeadlessChaos/Private/HeadlessChaosTestOverlap.cpp
+++ b/Engine/Source/Programs/HeadlessChaos/Private/HeadlessChaosTestOverlap.cpp
@@ -275,6 +275,33 @@ namespace ChaosTest
 				EXPECT_EQ(bResult, false);
 			}
 		}
+		// Box Scaled test non uniform transform with rotation
+		{
+			TSharedPtr<TBox<FReal, 3>> BigBoxSafe = MakeShared<TBox<FReal, 3>>(FVec3(-100.0, -100.0, -100.0), FVec3(100.0, 100.0, 100.0));
+			FVec3 TriMeshScale = { 0.01, 0.01, 0.05 };
+			TImplicitObjectScaled<TBox<FReal, 3>> ScaledBox = TImplicitObjectScaled<TBox<FReal, 3>>(BigBoxSafe, TriMeshScale);
+
+			{
+				const FRigidTransform3 QueryTM(FVec3(5.5, 0.0, 10.0), FQuat::Identity);
+				bool bResult = TriangleMesh->OverlapGeom(ScaledBox, QueryTM, 0.0, nullptr, TriMeshScale);
+				EXPECT_EQ(bResult, true);
+			}
+			{
+				const FRigidTransform3 QueryTM(FVec3(7.0, 0.0, 10.0), FQuat::Identity);
+				bool bResult = TriangleMesh->OverlapGeom(ScaledBox, QueryTM, 0.0, nullptr, TriMeshScale);
+				EXPECT_EQ(bResult, false);
+			}
+			{
+				const FRigidTransform3 QueryTM(FVec3(7.0, 0.0, 10.0), FQuat(FVec3(0.0, 1.0, 0.0), -3.1415926 / 4.0));
+				bool bResult = TriangleMesh->OverlapGeom(ScaledBox, QueryTM, 0.0, nullptr, TriMeshScale);
+				EXPECT_EQ(bResult, false);
+			}
+			{
+				const FRigidTransform3 QueryTM(FVec3(7.0, 0.0, 10.0), FQuat(FVec3(0.0, 1.0, 0.0), 3.1415926 / 4.0));
+				bool bResult = TriangleMesh->OverlapGeom(ScaledBox, QueryTM, 0.0, nullptr, TriMeshScale);
+				EXPECT_EQ(bResult, true);
+			}
+		}
 		{
 			// Sphere test
 			const TSphere<FReal, 3> Sphere = TSphere<FReal, 3>({ 0.0, 0.0, 0.0 }, 2.0);
