Headless chaos : fix compilation errors

diff --git a/Engine/Source/Programs/HeadlessChaos/Private/GeometryCollection/GeometryCollectionTestClustering.cpp b/Engine/Source/Programs/HeadlessChaos/Private/GeometryCollection/GeometryCollectionTestClustering.cpp
--- a/Engine/Source/Programs/HeadlessChaos/Private/GeometryCollection/GeometryCollectionTestClustering.cpp
+++ b/Engine/Source/Programs/HeadlessChaos/Private/GeometryCollection/GeometryCollectionTestClustering.cpp
@@ -607,8 +607,8 @@ namespace GeometryCollectionTest
 		Params.ClusterGroupIndex = 0;
 		FGeometryCollectionWrapper* Collection1 = TNewSimulationObject<GeometryType::GeometryCollectionWithSuppliedRestCollection>::Init(Params)->template As<FGeometryCollectionWrapper>();
 		TSharedPtr<FGeometryDynamicCollection> DynamicCollection1 = Collection1->DynamicCollection;
-		DynamicCollection1->GetAttribute<int32>("DynamicState", FGeometryCollection::TransformGroup)[1] = (uint8)EObjectStateTypeEnum::Chaos_Object_Kinematic;
-		DynamicCollection1->GetAttribute<int32>("DynamicState", FGeometryCollection::TransformGroup)[0] = (uint8)EObjectStateTypeEnum::Chaos_Object_Kinematic;
+		DynamicCollection1->ModifyAttribute<int32>("DynamicState", FGeometryCollection::TransformGroup)[1] = (uint8)EObjectStateTypeEnum::Chaos_Object_Kinematic;
+		DynamicCollection1->ModifyAttribute<int32>("DynamicState", FGeometryCollection::TransformGroup)[0] = (uint8)EObjectStateTypeEnum::Chaos_Object_Kinematic;
 
 		UnitTest.AddSimulationObject(Collection1);
 
@@ -855,7 +855,7 @@ namespace GeometryCollectionTest
 		Params.MaxClusterLevel = 1;
 		FGeometryCollectionWrapper* Collection = TNewSimulationObject<GeometryType::GeometryCollectionWithSuppliedRestCollection>::Init(Params)->template As<FGeometryCollectionWrapper>();
 
-		Collection->DynamicCollection->template GetAttribute<int32>("DynamicState", FGeometryCollection::TransformGroup)[1] = (uint8)EObjectStateTypeEnum::Chaos_Object_Kinematic;
+		Collection->DynamicCollection->template ModifyAttribute<int32>("DynamicState", FGeometryCollection::TransformGroup)[1] = (uint8)EObjectStateTypeEnum::Chaos_Object_Kinematic;
 
 		UnitTest.AddSimulationObject(Collection);
 		UnitTest.Initialize();
@@ -1154,7 +1154,7 @@ namespace GeometryCollectionTest
 		Params.MaxClusterLevel = 1;
 		FGeometryCollectionWrapper* Collection = TNewSimulationObject<GeometryType::GeometryCollectionWithSuppliedRestCollection>::Init(Params)->template As<FGeometryCollectionWrapper>();
 
-		Collection->DynamicCollection->template GetAttribute<int32>("DynamicState", FGeometryCollection::TransformGroup)[1] = (uint8)EObjectStateTypeEnum::Chaos_Object_Static;
+		Collection->DynamicCollection->template ModifyAttribute<int32>("DynamicState", FGeometryCollection::TransformGroup)[1] = (uint8)EObjectStateTypeEnum::Chaos_Object_Static;
 
 		UnitTest.AddSimulationObject(Collection);
 		UnitTest.Initialize();
@@ -1877,7 +1877,7 @@ namespace GeometryCollectionTest
 
 		TSharedPtr<FGeometryDynamicCollection> DynamicCollection1 = Collection1->DynamicCollection;
 		TSharedPtr<FGeometryDynamicCollection> DynamicCollection2 = Collection2->DynamicCollection;
-		DynamicCollection1->GetAttribute<int32>("DynamicState", FGeometryCollection::TransformGroup)[1] = (uint8)EObjectStateTypeEnum::Chaos_Object_Kinematic;
+		DynamicCollection1->ModifyAttribute<int32>("DynamicState", FGeometryCollection::TransformGroup)[1] = (uint8)EObjectStateTypeEnum::Chaos_Object_Kinematic;
 
 		UnitTest.Initialize();
 
@@ -1968,7 +1968,7 @@ namespace GeometryCollectionTest
 		Params.ClusterGroupIndex = 0;		
 		FGeometryCollectionWrapper* Collection1 = TNewSimulationObject<GeometryType::GeometryCollectionWithSuppliedRestCollection>::Init(Params)->template As<FGeometryCollectionWrapper>();
 		TSharedPtr<FGeometryDynamicCollection> DynamicCollection1 = Collection1->DynamicCollection;
-		DynamicCollection1->GetAttribute<int32>("DynamicState", FGeometryCollection::TransformGroup)[1] = (uint8)EObjectStateTypeEnum::Chaos_Object_Kinematic;
+		DynamicCollection1->ModifyAttribute<int32>("DynamicState", FGeometryCollection::TransformGroup)[1] = (uint8)EObjectStateTypeEnum::Chaos_Object_Kinematic;
 
 		UnitTest.AddSimulationObject(Collection1);
 
@@ -2036,7 +2036,7 @@ namespace GeometryCollectionTest
 		FGeometryCollectionWrapper* Collection = TNewSimulationObject<GeometryType::GeometryCollectionWithSuppliedRestCollection>::Init(Params)->template As<FGeometryCollectionWrapper>();
 
 		TSharedPtr<FGeometryDynamicCollection> DynamicCollection = Collection->DynamicCollection;
-		DynamicCollection->GetAttribute<int32>("DynamicState", FGeometryCollection::TransformGroup)[1] = (uint8)EObjectStateTypeEnum::Chaos_Object_Kinematic;
+		DynamicCollection->ModifyAttribute<int32>("DynamicState", FGeometryCollection::TransformGroup)[1] = (uint8)EObjectStateTypeEnum::Chaos_Object_Kinematic;
 
 		UnitTest.AddSimulationObject(Collection);
 		UnitTest.Initialize();
@@ -2146,7 +2146,7 @@ namespace GeometryCollectionTest
 		FalloffField->Falloff = EFieldFalloffType::Field_FallOff_None;
 
 		TSharedPtr<FGeometryDynamicCollection> DynamicCollection = Collection->DynamicCollection;
-		DynamicCollection->GetAttribute<int32>("DynamicState", FGeometryCollection::TransformGroup)[1] = (uint8)EObjectStateTypeEnum::Chaos_Object_Kinematic;
+		DynamicCollection->ModifyAttribute<int32>("DynamicState", FGeometryCollection::TransformGroup)[1] = (uint8)EObjectStateTypeEnum::Chaos_Object_Kinematic;
 
 		UnitTest.AddSimulationObject(Collection);
 		UnitTest.Initialize();
diff --git a/Engine/Source/Programs/HeadlessChaos/Private/GeometryCollection/GeometryCollectionTestCollisionResolution.cpp b/Engine/Source/Programs/HeadlessChaos/Private/GeometryCollection/GeometryCollectionTestCollisionResolution.cpp
--- a/Engine/Source/Programs/HeadlessChaos/Private/GeometryCollection/GeometryCollectionTestCollisionResolution.cpp
+++ b/Engine/Source/Programs/HeadlessChaos/Private/GeometryCollection/GeometryCollectionTestCollisionResolution.cpp
@@ -141,14 +141,14 @@ namespace GeometryCollectionTest
 		// Hard code masstolocal on rest collection to identity
 		{
 			TManagedArray<FTransform>& MassToLocal =
-				SimplicialSphereCollection->RestCollection->template GetAttribute<FTransform>(
+				SimplicialSphereCollection->RestCollection->template ModifyAttribute<FTransform>(
 					TEXT("MassToLocal"), FTransformCollection::TransformGroup);
 			check(MassToLocal.Num() == 1);
 			MassToLocal[0] = FTransform::Identity;
 		}
 		{
 			TManagedArray<FTransform>& MassToLocal =
-				ImplicitSphereCollection->RestCollection->template GetAttribute<FTransform>(
+				ImplicitSphereCollection->RestCollection->template ModifyAttribute<FTransform>(
 					TEXT("MassToLocal"), FTransformCollection::TransformGroup);
 			check(MassToLocal.Num() == 1);
 			MassToLocal[0] = FTransform::Identity;
@@ -367,12 +367,12 @@ namespace GeometryCollectionTest
 
 		// Hard code masstolocal on rest collection to identity
 		{
-			TManagedArray<FTransform>& MassToLocal = SimplicialSphereCollection->RestCollection->template GetAttribute<FTransform>(TEXT("MassToLocal"), FTransformCollection::TransformGroup);
+			TManagedArray<FTransform>& MassToLocal = SimplicialSphereCollection->RestCollection->template ModifyAttribute<FTransform>(TEXT("MassToLocal"), FTransformCollection::TransformGroup);
 			check(MassToLocal.Num() == 1);
 			MassToLocal[0] = FTransform::Identity;
 		}
 		{
-			TManagedArray<FTransform>& MassToLocal =ImplicitSphereCollection->RestCollection->template GetAttribute<FTransform>(TEXT("MassToLocal"), FTransformCollection::TransformGroup);
+			TManagedArray<FTransform>& MassToLocal =ImplicitSphereCollection->RestCollection->template ModifyAttribute<FTransform>(TEXT("MassToLocal"), FTransformCollection::TransformGroup);
 			check(MassToLocal.Num() == 1);
 			MassToLocal[0] = FTransform::Identity;
 		}
diff --git a/Engine/Source/Programs/HeadlessChaos/Private/GeometryCollection/GeometryCollectionTestInitilization.cpp b/Engine/Source/Programs/HeadlessChaos/Private/GeometryCollection/GeometryCollectionTestInitilization.cpp
--- a/Engine/Source/Programs/HeadlessChaos/Private/GeometryCollection/GeometryCollectionTestInitilization.cpp
+++ b/Engine/Source/Programs/HeadlessChaos/Private/GeometryCollection/GeometryCollectionTestInitilization.cpp
@@ -119,8 +119,8 @@ GTEST_TEST(AllTraits,GeometryCollection_Initilization_TransformedGeometryCollect
 	//
 	TArray<FTransform> RestTransforms;
 	int32 NumVertices = Collection->DynamicCollection->NumElements(FGeometryCollection::VerticesGroup);
-	TManagedArray<FVector3f>& Vertices = Collection->RestCollection->template GetAttribute<FVector3f>("Vertex",FGeometryCollection::VerticesGroup);
-	TManagedArray<int32>& BoneMap = Collection->RestCollection->template GetAttribute<int32>("BoneMap",FGeometryCollection::VerticesGroup);
+	const TManagedArray<FVector3f>& Vertices = Collection->RestCollection->template GetAttribute<FVector3f>("Vertex",FGeometryCollection::VerticesGroup);
+	const TManagedArray<int32>& BoneMap = Collection->RestCollection->template GetAttribute<int32>("BoneMap",FGeometryCollection::VerticesGroup);
 	GeometryCollectionAlgo::GlobalMatrices(Collection->RestCollection->Transform,Collection->RestCollection->Parent,RestTransforms);
 	FVector CenterOfMass(0);  
 	for (int VertexIndex = 0; VertexIndex < NumVertices; VertexIndex++)
@@ -138,7 +138,7 @@ GTEST_TEST(AllTraits,GeometryCollection_Initilization_TransformedGeometryCollect
 	UnitTest.Advance();
 
 	{ // test results
-		TManagedArray<FTransform>& MassToLocal = Collection->RestCollection->template GetAttribute<FTransform>("MassToLocal",FGeometryCollection::TransformGroup);
+		const TManagedArray<FTransform>& MassToLocal = Collection->RestCollection->template GetAttribute<FTransform>("MassToLocal",FGeometryCollection::TransformGroup);
 
 
 		EXPECT_EQ(UnitTest.Solver->GetParticles().GetGeometryCollectionParticles().Size(), 1);
diff --git a/Engine/Source/Programs/HeadlessChaos/Private/GeometryCollection/GeometryCollectionTestProximity.cpp b/Engine/Source/Programs/HeadlessChaos/Private/GeometryCollection/GeometryCollectionTestProximity.cpp
--- a/Engine/Source/Programs/HeadlessChaos/Private/GeometryCollection/GeometryCollectionTestProximity.cpp
+++ b/Engine/Source/Programs/HeadlessChaos/Private/GeometryCollection/GeometryCollectionTestProximity.cpp
@@ -38,7 +38,7 @@ namespace GeometryCollectionTest
 		ProximityUtility.UpdateProximity();
 
 		// Breaking Data
-		TManagedArray<TSet<int32>>& Proximity = Collection->GetAttribute<TSet<int32>>("Proximity", FGeometryCollection::GeometryGroup);
+		const TManagedArray<TSet<int32>>& Proximity = Collection->GetAttribute<TSet<int32>>("Proximity", FGeometryCollection::GeometryGroup);
 
 		EXPECT_TRUE((Proximity)[0].Contains(1));
 		EXPECT_TRUE((Proximity)[1].Contains(0));
@@ -88,7 +88,7 @@ namespace GeometryCollectionTest
 		ProximityUtility.UpdateProximity();
 
 		// Breaking Data
-		TManagedArray<TSet<int32>>& Proximity = Collection->GetAttribute<TSet<int32>>("Proximity", FGeometryCollection::GeometryGroup);
+		const TManagedArray<TSet<int32>>& Proximity = Collection->GetAttribute<TSet<int32>>("Proximity", FGeometryCollection::GeometryGroup);
 
 		// Proximity = [(3,4,1), (0,4,5,2), (1,5), (0,4), (0,1,3,5), (1,2,4)]
 
@@ -216,7 +216,7 @@ namespace GeometryCollectionTest
 		ProximityUtility.UpdateProximity();
 
 		// Breaking Data
-		TManagedArray<TSet<int32>>& Proximity = Collection->GetAttribute<TSet<int32>>("Proximity", FGeometryCollection::GeometryGroup);
+		const TManagedArray<TSet<int32>>& Proximity = Collection->GetAttribute<TSet<int32>>("Proximity", FGeometryCollection::GeometryGroup);
 
 		// Proximity = [(3,4,1), (0,4,5,2), (1,5), (0,4), (0,1,3,5), (1,2,4)]
 
@@ -344,7 +344,7 @@ namespace GeometryCollectionTest
 		ProximityUtility.UpdateProximity();
 
 		// Breaking Data
-		TManagedArray<TSet<int32>>& Proximity = Collection->GetAttribute<TSet<int32>>("Proximity", FGeometryCollection::GeometryGroup);
+		const TManagedArray<TSet<int32>>& Proximity = Collection->GetAttribute<TSet<int32>>("Proximity", FGeometryCollection::GeometryGroup);
 
 		// Proximity = [(3,4,1), (0,4,5,2), (1,5), (0,4), (0,1,3,5), (1,2,4)]
 
@@ -473,7 +473,7 @@ namespace GeometryCollectionTest
 		ProximityUtility.UpdateProximity();
 
 		// Breaking Data
-		TManagedArray<TSet<int32>>& Proximity = Collection->GetAttribute<TSet<int32>>("Proximity", FGeometryCollection::GeometryGroup);
+		const TManagedArray<TSet<int32>>& Proximity = Collection->GetAttribute<TSet<int32>>("Proximity", FGeometryCollection::GeometryGroup);
 
 		// Proximity = [(3,4,1), (0,4,5,2), (1,5), (0,4), (0,1,3,5), (1,2,4)]
 
@@ -587,7 +587,7 @@ namespace GeometryCollectionTest
 		ProximityUtility.UpdateProximity();
 
 		// Breaking Data
-		TManagedArray<TSet<int32>>& Proximity = Collection->GetAttribute<TSet<int32>>("Proximity", FGeometryCollection::GeometryGroup);
+		const TManagedArray<TSet<int32>>& Proximity = Collection->GetAttribute<TSet<int32>>("Proximity", FGeometryCollection::GeometryGroup);
 
 		// Proximity = [(3,4,1), (0,4,5,2), (1,5), (0,4), (0,1,3,5), (1,2,4)]
 
@@ -701,7 +701,7 @@ namespace GeometryCollectionTest
 		ProximityUtility.UpdateProximity();
 
 		// Breaking Data
-		TManagedArray<TSet<int32>>& Proximity = Collection->GetAttribute<TSet<int32>>("Proximity", FGeometryCollection::GeometryGroup);
+		const TManagedArray<TSet<int32>>& Proximity = Collection->GetAttribute<TSet<int32>>("Proximity", FGeometryCollection::GeometryGroup);
 
 		// Proximity = [(3,4,1), (0,4,5,2), (1,5), (0,4), (0,1,3,5), (1,2,4)]
 
@@ -808,7 +808,7 @@ namespace GeometryCollectionTest
 		ProximityUtility.UpdateProximity();
 
 		// Breaking Data
-		TManagedArray<TSet<int32>>& Proximity = Collection->GetAttribute<TSet<int32>>("Proximity", FGeometryCollection::GeometryGroup);
+		const TManagedArray<TSet<int32>>& Proximity = Collection->GetAttribute<TSet<int32>>("Proximity", FGeometryCollection::GeometryGroup);
 
 		// Proximity = [(3,4,1), (0,4,5,2), (1,5), (0,4), (0,1,3,5), (1,2,4)]
 
diff --git a/Engine/Source/Programs/HeadlessChaos/Private/GeometryCollection/GeometryCollectionTestSimulationSolver.cpp b/Engine/Source/Programs/HeadlessChaos/Private/GeometryCollection/GeometryCollectionTestSimulationSolver.cpp
--- a/Engine/Source/Programs/HeadlessChaos/Private/GeometryCollection/GeometryCollectionTestSimulationSolver.cpp
+++ b/Engine/Source/Programs/HeadlessChaos/Private/GeometryCollection/GeometryCollectionTestSimulationSolver.cpp
@@ -289,7 +289,7 @@ using namespace ChaosTest;
 		FFramework UnitTest;
 		for (int i=0; i<10; i++)
 		{
-			TManagedArray<float>& Mass = Collection[i]->RestCollection->template GetAttribute<float>("Mass", FTransformCollection::TransformGroup);
+			TManagedArray<float>& Mass = Collection[i]->RestCollection->template ModifyAttribute<float>("Mass", FTransformCollection::TransformGroup);
 			Mass[0] = i + 1;
 			UnitTest.AddSimulationObject(Collection[i]);
 		}
diff --git a/Engine/Source/Runtime/Experimental/Chaos/Public/GeometryCollection/GeometryCollectionUtility.h b/Engine/Source/Runtime/Experimental/Chaos/Public/GeometryCollection/GeometryCollectionUtility.h
--- a/Engine/Source/Runtime/Experimental/Chaos/Public/GeometryCollection/GeometryCollectionUtility.h
+++ b/Engine/Source/Runtime/Experimental/Chaos/Public/GeometryCollection/GeometryCollectionUtility.h
@@ -186,7 +186,7 @@ void GeometryCollection::AttributeTransfer(const FGeometryCollection * FromColle
 {
 	// #todo(dmp): later on we will support different attribute groups for transfer		
 	const TManagedArray<T> &FromAttribute = FromCollection->GetAttribute<T>(FromAttributeName, FGeometryCollection::VerticesGroup);
-	TManagedArray<T> &ToAttribute = ToCollection->GetAttribute<T>(ToAttributeName, FGeometryCollection::VerticesGroup);
+	TManagedArray<T> &ToAttribute = ToCollection->ModifyAttribute<T>(ToAttributeName, FGeometryCollection::VerticesGroup);
 
 	const TManagedArray<FVector3f> &FromVertex = FromCollection->Vertex;
 	TManagedArray<FVector3f> &ToVertex = ToCollection->Vertex;
