Non-uniform scale bug fix when calculating bounding volumes in Mesh when finding most opposing faces.

diff --git a/Engine/Source/Programs/HeadlessChaos/Private/HeadlessChaosTestMostOpposing.cpp b/Engine/Source/Programs/HeadlessChaos/Private/HeadlessChaosTestMostOpposing.cpp
--- a/Engine/Source/Programs/HeadlessChaos/Private/HeadlessChaosTestMostOpposing.cpp
+++ b/Engine/Source/Programs/HeadlessChaos/Private/HeadlessChaosTestMostOpposing.cpp
@@ -141,6 +141,39 @@ namespace ChaosTest
 			EXPECT_EQ(Tri->FindMostOpposingFaceScaled(FVec3(0, 0, 0), FVec3(0, 0, -1), INDEX_NONE, 100.0, MirrorScale), 2);		
 
 		}
+
+		// Simple Non-uniform scale (Actual bug regression test)
+		{
+			// 2 triangles
+			FTriangleMeshImplicitObject::ParticlesType Particles;
+			Particles.AddParticles(6);
+			// in z-y plane
+			Particles.X(0) = FVec3(0, 0, 0);
+			Particles.X(1) = FVec3(0, 100, 0);
+			Particles.X(2) = FVec3(0, 0, 100);
+
+			// In x-z plane
+			Particles.X(3) = FVec3(0, 0, 0);
+			Particles.X(4) = FVec3(100, 0, 0);
+			Particles.X(5) = FVec3(0, 0, 100);
+
+
+			TArray<TVec3<int32>> Indices;
+			Indices.Emplace(0, 1, 2);
+			Indices.Emplace(3, 4, 5);
+			TArray<uint16> DummyMaterials;
+			TUniquePtr<FTriangleMeshImplicitObject> Tri = MakeUnique<FTriangleMeshImplicitObject>(MoveTemp(Particles), MoveTemp(Indices), MoveTemp(DummyMaterials));
+
+			const FVec3 Scale(10, 1, 1);
+			TImplicitObjectScaledGeneric<FReal, 3> ScaledTri(MakeSerializable(Tri), nullptr, Scale);
+
+			// Make sure we find the correct face when at 200 units away from origin on x-axes. The scaling will ensure that the point is on the face. using 1 unit search distance
+			EXPECT_EQ(ScaledTri.FindMostOpposingFace(FVec3(200, 0, 0), FVec3(0,1,0), INDEX_NONE, 1.0f), 1);
+			// Change direction and increase search distance, to hit other face
+			EXPECT_EQ(ScaledTri.FindMostOpposingFace(FVec3(200, 0, 0), FVec3(-1, 0, 0), INDEX_NONE, 201.0f), 0);
+			// Reduce search distance so that we don't hit the best face
+			EXPECT_EQ(ScaledTri.FindMostOpposingFace(FVec3(200, 0, 0), FVec3(-1, 0, 0), INDEX_NONE, 199.0f), 1);
+		}
 	}
 
 
diff --git a/Engine/Source/Runtime/Experimental/Chaos/Private/Chaos/TriangleMeshImplicitObject.cpp b/Engine/Source/Runtime/Experimental/Chaos/Private/Chaos/TriangleMeshImplicitObject.cpp
--- a/Engine/Source/Runtime/Experimental/Chaos/Private/Chaos/TriangleMeshImplicitObject.cpp
+++ b/Engine/Source/Runtime/Experimental/Chaos/Private/Chaos/TriangleMeshImplicitObject.cpp
@@ -1277,7 +1277,11 @@ int32 FTriangleMeshImplicitObject::FindMostOpposingFace(const TArray<TVec3<IdxTy
 	//todo: this is horribly slow, need adjacency information
 	const FReal SearchDist2 = SearchDist * SearchDist;
 
-	FAABB3 QueryBounds(Position - FVec3(SearchDist), Position + FVec3(SearchDist));
+	const FVec3 InvScale = 1 / Scale;
+	const FVec3 ScaledPosition = Position * InvScale;
+	const FVec3 ScaledSearchDist = SearchDist * InvScale;
+	const FVec3 AbsScaledSearchDist(FMath::Abs(ScaledSearchDist[0]), FMath::Abs(ScaledSearchDist[1]), FMath::Abs(ScaledSearchDist[2]));
+	const FAABB3 QueryBounds(ScaledPosition - AbsScaledSearchDist, ScaledPosition + AbsScaledSearchDist);
 
 	const TArray<int32> PotentialIntersections = FastBVH.FindAllIntersections(QueryBounds);
 	const FReal Epsilon = 1e-4f;
