Chaos : fix compile issue with CL 20488364

diff --git a/Engine/Plugins/Experimental/ChaosCaching/Source/ChaosCaching/Private/Chaos/Adapters/GeometryCollectionComponentCacheAdapter.cpp b/Engine/Plugins/Experimental/ChaosCaching/Source/ChaosCaching/Private/Chaos/Adapters/GeometryCollectionComponentCacheAdapter.cpp
--- a/Engine/Plugins/Experimental/ChaosCaching/Source/ChaosCaching/Private/Chaos/Adapters/GeometryCollectionComponentCacheAdapter.cpp
+++ b/Engine/Plugins/Experimental/ChaosCaching/Source/ChaosCaching/Private/Chaos/Adapters/GeometryCollectionComponentCacheAdapter.cpp
@@ -338,7 +338,7 @@ namespace Chaos
 					}
 					if (ClusterHandle)
 					{
-						Solver->GetEvolution()->GetRigidClustering().ReleaseClusterParticlesNoInternalCluster(ClusterHandle->CastToClustered(), nullptr, true);
+						Solver->GetEvolution()->GetRigidClustering().ReleaseClusterParticlesNoInternalCluster(ClusterHandle->CastToClustered(), true);
 					}
 				}
 			}
diff --git a/Engine/Source/Programs/HeadlessChaos/Private/GeometryCollection/GeometryCollectionTestClustering.cpp b/Engine/Source/Programs/HeadlessChaos/Private/GeometryCollection/GeometryCollectionTestClustering.cpp
--- a/Engine/Source/Programs/HeadlessChaos/Private/GeometryCollection/GeometryCollectionTestClustering.cpp
+++ b/Engine/Source/Programs/HeadlessChaos/Private/GeometryCollection/GeometryCollectionTestClustering.cpp
@@ -307,7 +307,7 @@ namespace GeometryCollectionTest
 
 			if (Frame == 2)
 			{
-				ParticleHandles[0].AddExternalStrain(50.0f);
+				ParticleHandles[0]->AddExternalStrain(50.0f);
 				Clustering.BreakingModel();
 			}
 
@@ -2086,7 +2086,7 @@ namespace GeometryCollectionTest
 
 			if (Frame == 5)
 			{
-				UnitTest.Solver->GetEvolution()->GetRigidClustering().ReleaseClusterParticles(ParticleHandles[4]->CastToClustered(), nullptr, true);
+				UnitTest.Solver->GetEvolution()->GetRigidClustering().ReleaseClusterParticles(ParticleHandles[4]->CastToClustered(), true);
 			}
 			
 			if (Frame < 5)
