Chaos unit test fix for FindMostOpposingFaceScaled.

diff --git a/Engine/Source/Programs/HeadlessChaos/Private/HeadlessChaosTestMostOpposing.cpp b/Engine/Source/Programs/HeadlessChaos/Private/HeadlessChaosTestMostOpposing.cpp
--- a/Engine/Source/Programs/HeadlessChaos/Private/HeadlessChaosTestMostOpposing.cpp
+++ b/Engine/Source/Programs/HeadlessChaos/Private/HeadlessChaosTestMostOpposing.cpp
@@ -137,7 +137,7 @@ namespace ChaosTest
 			MirrorScale = FVec3(-1, -1, -1);
 			EXPECT_EQ(Tri->FindMostOpposingFaceScaled(FVec3(0, 0, 0), FVec3(0, 0, 1), INDEX_NONE, 100.0, MirrorScale), 2);
 
-			MirrorScale = FVec3(-1, -1, 0);
+			MirrorScale = FVec3(-1, -1, 1);
 			EXPECT_EQ(Tri->FindMostOpposingFaceScaled(FVec3(0, 0, 0), FVec3(0, 0, -1), INDEX_NONE, 100.0, MirrorScale), 2);		
 
 		}
