Chaos : Add more Cloth Collection tests to deal with Start and End corner cases,  e.g. at remove time, and when adding new elements.

#okforgithuib public

diff --git a/Engine/Source/Programs/HeadlessChaos/Private/HeadlessChaosTestCloth.cpp b/Engine/Source/Programs/HeadlessChaos/Private/HeadlessChaosTestCloth.cpp
--- a/Engine/Source/Programs/HeadlessChaos/Private/HeadlessChaosTestCloth.cpp
+++ b/Engine/Source/Programs/HeadlessChaos/Private/HeadlessChaosTestCloth.cpp
@@ -347,25 +347,129 @@ namespace ChaosTest {
 				check(ViewPosition <= ManagedArrayCollection.NumElements(ViewGroup));
 				EXPECT_TRUE(ManagedArrayCollection.InsertElements(1, ViewPosition, ViewGroup) == ViewPosition);
 
-				const int32 ValuePosition = (ViewPosition > 0) ? ValueEnd[ViewPosition - 1] + 1: 0;
-				EXPECT_TRUE(ManagedArrayCollection.InsertElements(NumValues, ValuePosition, DataGroup) == ValuePosition);
-			
 				static int32 ViewId = 0;
 				Id[ViewPosition] = ViewId++;
-				ValueStart[ViewPosition] = ValuePosition;
-				ValueEnd[ViewPosition] = ValuePosition + NumValues - 1;
 
-				for (int32 ValueIndex = ValueStart[ViewPosition]; ValueIndex <= ValueEnd[ViewPosition]; ++ValueIndex)
+				if (NumValues)
+				{
+					const int32 ValuePosition = (ViewPosition > 0) ? ValueEnd[ViewPosition - 1] + 1: 0;
+					EXPECT_TRUE(ManagedArrayCollection.InsertElements(NumValues, ValuePosition, DataGroup) == ValuePosition);
+			
+					ValueStart[ViewPosition] = ValuePosition;
+					ValueEnd[ViewPosition] = ValuePosition + NumValues - 1;
+
+					for (int32 ValueIndex = ValueStart[ViewPosition]; ValueIndex <= ValueEnd[ViewPosition]; ++ValueIndex)
+					{
+						Value[ValueIndex] = Id[ViewPosition];
+					}
+				}
+				else
+				{
+					ValueStart[ViewPosition] = INDEX_NONE;
+					ValueEnd[ViewPosition] = INDEX_NONE;
+				}
+			};
+
+		auto RemoveViews = [&ManagedArrayCollection, &ValueStart, &ValueEnd, &DataGroup, &ViewGroup](int32 ViewPosition, int32 NumViews)
+			{
+				const int32 ViewStart = ViewPosition;
+				const int32 ViewEnd = ViewPosition + NumViews - 1;
+
+				for (int32 ViewIndex = ViewStart; ViewIndex <= ViewEnd; ++ViewIndex)
 				{
-					Value[ValueIndex] = Id[ViewPosition];
+					// Remove data
+					const int32 Position = ValueStart[ViewIndex];
+					if (Position != INDEX_NONE)
+					{
+						const int32 NumValues = ValueEnd[ViewIndex] - ValueStart[ViewIndex] + 1;
+						ManagedArrayCollection.RemoveElements(DataGroup, NumValues, Position);
+					}
 				}
+				// Remove views
+				ManagedArrayCollection.RemoveElements(ViewGroup, NumViews, ViewPosition);
 			};
 
+		auto SetViewSize = [&ManagedArrayCollection, &Value, &ValueStart, &ValueEnd, &Id, &DataGroup, &ViewGroup](int32 ViewIndex, int32 InNumValues)
+			{
+				check(InNumValues >= 0);
+
+				int32& Start = ValueStart[ViewIndex];
+				int32& End = ValueEnd[ViewIndex];
+				check(Start != INDEX_NONE || End == INDEX_NONE);
+
+				const int32 NumValues = (Start == INDEX_NONE) ? 0 : End - Start + 1;
+
+				if (const int32 Delta = InNumValues - NumValues)
+				{
+					if (Delta > 0)
+					{
+						// Find a previous valid index range to insert after when the range is empty
+						auto ComputeEnd = [&ValueEnd](int32 ViewIndex)->int32
+						{
+							for (int32 Index = ViewIndex; Index >= 0; --Index)
+							{
+								if (ValueEnd[Index] != INDEX_NONE)
+								{
+									return ValueEnd[Index];
+								}
+							}
+							return INDEX_NONE;
+						};
+
+						// Grow the array
+						const int32 Position = ComputeEnd(ViewIndex) + 1;
+						ManagedArrayCollection.InsertElements(Delta, Position, DataGroup);
+
+						// Update Start/End
+						if (!NumValues)
+						{
+							Start = Position;
+						}
+						End = Start + InNumValues - 1;
+
+						// Fill the test values with the id check
+						for (int32 ValueIndex = Start; ValueIndex <= End; ++ValueIndex)
+						{
+							Value[ValueIndex] = Id[ViewIndex];
+						}
+
+					}
+					else
+					{
+						// Shrink the array
+						const int32 Position = Start + InNumValues;
+						ManagedArrayCollection.RemoveElements(DataGroup, -Delta, Position);
+
+						// Update Start/End
+						if (InNumValues)
+						{
+							End = Position - 1;
+						}
+						else
+						{
+							End = Start = INDEX_NONE;
+						}
+					}
+				}
+				const int32 Size = ValueEnd[ViewIndex] - ValueStart[ViewIndex] + 1;
+				check(
+					(InNumValues == 0 && ValueStart[ViewIndex] == INDEX_NONE && ValueEnd[ViewIndex] == INDEX_NONE) ||
+					(InNumValues == Size && ValueStart[ViewIndex] != INDEX_NONE && ValueEnd[ViewIndex] != INDEX_NONE));
+		};
+
 		auto HasKeptIntegrity = [&Value , &ValueStart, &ValueEnd, &Id]()->bool
 			{
 				EXPECT_TRUE(ValueEnd.Num() == ValueStart.Num() && Id.Num() == ValueStart.Num());
 				for (int32 ViewIndex = 0; ViewIndex < ValueStart.Num(); ++ViewIndex)
 				{
+					if (ValueStart[ViewIndex] == INDEX_NONE || ValueEnd[ViewIndex] == INDEX_NONE)
+					{
+						if (ValueStart[ViewIndex] != ValueEnd[ViewIndex])
+						{
+							return false;
+						}
+						continue;
+					}
 					for (int32 ValueIndex = ValueStart[ViewIndex]; ValueIndex <= ValueEnd[ViewIndex]; ++ValueIndex)
 					{
 						if (Value[ValueIndex] != Id[ViewIndex])
@@ -377,31 +481,109 @@ namespace ChaosTest {
 				return true;
 			};
 
+		auto HasKeptOrder = [&Value, &ValueStart, &ValueEnd, &Id]()->bool
+		{
+			int32 PrevValueEnd = -1;
+			for (int32 ViewIndex = 0; ViewIndex < ValueStart.Num(); ++ViewIndex)
+			{
+				if (ValueStart[ViewIndex] == INDEX_NONE)
+				{
+					check(ValueEnd[ViewIndex] == INDEX_NONE);
+					continue;
+				}
+				if (ValueStart[ViewIndex] != PrevValueEnd + 1)
+				{
+					return false;
+				}
+				PrevValueEnd = ValueEnd[ViewIndex];
+			}
+			return PrevValueEnd == Value.Num() - 1;
+		};
+
 		// Insert 5 views from the start of the array
 		for (int32 Index = 0; Index < 5; ++Index)
 		{
 			InsertView(0, FMath::Rand() / 32);
 		}
 		EXPECT_TRUE(HasKeptIntegrity());
+		EXPECT_TRUE(HasKeptOrder());
 
 		// Add 1 view in the middle of the array
 		InsertView(ManagedArrayCollection.NumElements(ViewGroup) / 2, FMath::Rand() / 32);
 		EXPECT_TRUE(HasKeptIntegrity());
+		EXPECT_TRUE(HasKeptOrder());
 
 		// Add 1 view at the end of the array
 		InsertView(ManagedArrayCollection.NumElements(ViewGroup), FMath::Rand() / 32);
 		EXPECT_TRUE(HasKeptIntegrity());
+		EXPECT_TRUE(HasKeptOrder());
 
 		// Remove 1 view from the start of the array
-		ManagedArrayCollection.RemoveElements(ViewGroup, 1, 0);
+		RemoveViews(0, 1);
 		EXPECT_TRUE(HasKeptIntegrity());
+		EXPECT_TRUE(HasKeptOrder());
 
 		// Remove 1 view from the end of the array
-		ManagedArrayCollection.RemoveElements(ViewGroup, 1, ManagedArrayCollection.NumElements(ViewGroup) - 1);
+		RemoveViews(ManagedArrayCollection.NumElements(ViewGroup) - 1, 1);
 		EXPECT_TRUE(HasKeptIntegrity());
+		EXPECT_TRUE(HasKeptOrder());
 
 		// Remove 2 views from the middle
-		ManagedArrayCollection.RemoveElements(ViewGroup, 2, ManagedArrayCollection.NumElements(ViewGroup) / 2);
+		RemoveViews(ManagedArrayCollection.NumElements(ViewGroup) / 2, 2);
+		EXPECT_TRUE(HasKeptIntegrity());
+		EXPECT_TRUE(HasKeptOrder());
+
+		// Remove all remaining 3 views
+		RemoveViews(0, 3);
+		EXPECT_TRUE(ManagedArrayCollection.NumElements(ViewGroup) == 0);
+		EXPECT_TRUE(ManagedArrayCollection.NumElements(DataGroup) == 0);
+
+		// Insert 5 empty views
+		for (int32 Index = 0; Index < 5; ++Index)
+		{
+			InsertView(0, 0);
+		}
+		EXPECT_TRUE(HasKeptIntegrity());
+		EXPECT_TRUE(HasKeptOrder());
+
+		// Resize view 1
+		SetViewSize(1, 11);
+		EXPECT_TRUE(HasKeptIntegrity());
+		EXPECT_TRUE(HasKeptOrder());
+
+		// Resize view 2
+		SetViewSize(2, 22);
+		EXPECT_TRUE(HasKeptIntegrity());
+		EXPECT_TRUE(HasKeptOrder());
+
+		// Resize view 3
+		SetViewSize(3, 33);
+		EXPECT_TRUE(HasKeptIntegrity());
+		EXPECT_TRUE(HasKeptOrder());
+
+		// Remove view 2
+		RemoveViews(2, 1);
+		EXPECT_TRUE(HasKeptIntegrity());
+		EXPECT_TRUE(HasKeptOrder());
+		
+		// Resize view 1 to 0
+		SetViewSize(1, 0);
+		EXPECT_TRUE(HasKeptIntegrity());
+		EXPECT_TRUE(HasKeptOrder());
+
+		// Resize view 3 to 33
+		SetViewSize(3, 33);
+		EXPECT_TRUE(HasKeptIntegrity());
+		EXPECT_TRUE(HasKeptOrder());
+
+		// Resize view 3 to 0
+		SetViewSize(3, 0);
+		EXPECT_TRUE(HasKeptIntegrity());
+		EXPECT_TRUE(HasKeptOrder());
+
+		// Resize view 0
+		SetViewSize(0, 1);
 		EXPECT_TRUE(HasKeptIntegrity());
+		EXPECT_TRUE(HasKeptOrder());
 	}
 } // namespace ChaosTest
\ No newline at end of file
