CIS: Fix HeadlessChaos error caused by API change
#rnx

diff --git a/Engine/Source/Programs/HeadlessChaos/Private/HeadlessChaosTestSolverBody.cpp b/Engine/Source/Programs/HeadlessChaos/Private/HeadlessChaosTestSolverBody.cpp
--- a/Engine/Source/Programs/HeadlessChaos/Private/HeadlessChaosTestSolverBody.cpp
+++ b/Engine/Source/Programs/HeadlessChaos/Private/HeadlessChaosTestSolverBody.cpp
@@ -97,6 +97,6 @@ namespace ChaosTest
 
 		SolverBodyContainer.Reset(1);
 
-		SolverBodyContainer.FindOrAdd(Particles[0]);
+		SolverBodyContainer.FindOrAdd(Particles[0], 1./60.);
 	}
 }
