Headless chaos : fix unit test

diff --git a/Engine/Source/Programs/HeadlessChaos/Private/GeometryCollection/GeometryCollectionTestFramework.cpp b/Engine/Source/Programs/HeadlessChaos/Private/GeometryCollection/GeometryCollectionTestFramework.cpp
--- a/Engine/Source/Programs/HeadlessChaos/Private/GeometryCollection/GeometryCollectionTestFramework.cpp
+++ b/Engine/Source/Programs/HeadlessChaos/Private/GeometryCollection/GeometryCollectionTestFramework.cpp
@@ -173,6 +173,9 @@ namespace GeometryCollectionTest
 			SimulationParams.RemoveOnFractureEnabled = Params.RemoveOnFractureEnabled;
 			SimulationParams.CollisionGroup = Params.CollisionGroup;
 			SimulationParams.ClusterGroupIndex = Params.ClusterGroupIndex;
+			SimulationParams.LinearDamping = 0;
+			SimulationParams.AngularDamping = 0;
+			SimulationParams.UseCCD = false;
 
 			Chaos::FErrorReporter ErrorReporter;
 			BuildSimulationData(ErrorReporter, *RestCollection.Get(), SimulationParams.Shared);
