Add some initial unit tests for probe collisions.
1. Test that non-ccd collision between two shapes doesn't affect velocity if one of them is a probe
2. Test that ccd collision between two shapes doesn't affect velocity if one of them is a probe (and that the resulting constraint is not marked as CCD)

diff --git a/Engine/Source/Programs/HeadlessChaos/Private/HeadlessChaosTestProbes.cpp b/Engine/Source/Programs/HeadlessChaos/Private/HeadlessChaosTestProbes.cpp
new file mode 100644
--- /dev/null
+++ b/Engine/Source/Programs/HeadlessChaos/Private/HeadlessChaosTestProbes.cpp
@@ -0,0 +1,122 @@
+// Copyright Epic Games, Inc. All Rights Reserved.
+
+#include "Chaos/Matrix.h"
+#include "Chaos/Utilities.h"
+#include "Chaos/AABB.h"
+#include "Chaos/Core.h"
+#include "Chaos/PBDRigidsEvolutionGBF.h"
+#include "Chaos/Box.h"
+#include "Chaos/Sphere.h"
+#include "ChaosSolversModule.h"
+#include "HeadlessChaos.h"
+#include "HeadlessChaosTestUtility.h"
+#include "Modules/ModuleManager.h"
+#include "PBDRigidsSolver.h"
+
+namespace ChaosTest
+{
+	using namespace Chaos;
+
+	GTEST_TEST(ProbeTests, ProbeBodyConstraint)
+	{
+		const FReal BoxHalfSize = 50; // cm
+		const FReal Separation = 10; // cm
+		const FReal InitialSpeed = BoxHalfSize * .5; // cm/s... slow enough not to need CCD
+		const FReal Dt = FReal(2) * Separation / InitialSpeed; // s... twice the time needed to close the separation
+		FParticleUniqueIndicesMultithreaded UniqueIndices;
+		FPBDRigidsSOAs Particles(UniqueIndices);
+		THandleArray<FChaosPhysicsMaterial> PhysicalMaterials;
+		FPBDRigidsEvolutionGBF Evolution(Particles, PhysicalMaterials);
+		InitEvolutionSettings(Evolution);
+
+		// Create particles
+		TGeometryParticleHandle<FReal, 3>* Static = Evolution.CreateStaticParticles(1)[0];
+		TPBDRigidParticleHandle<FReal, 3>* Dynamic = Evolution.CreateDynamicParticles(1)[0];
+
+		// Create box geometry 
+		TUniquePtr<FImplicitObject> SmallBox(new TBox<FReal, 3>(FVec3(-BoxHalfSize, -BoxHalfSize, -BoxHalfSize), FVec3(BoxHalfSize, BoxHalfSize, BoxHalfSize)));
+		Static->SetGeometry(MakeSerializable(SmallBox));
+		AppendDynamicParticleConvexBox(*Dynamic, FVec3(BoxHalfSize), 0.0f);
+		Dynamic->SetGravityEnabled(false);
+
+		// Make the dynamic shape a probe
+		Dynamic->ShapesArray()[0]->SetIsProbe(true);
+
+		// Positions
+		Static->X() = FVec3(0, 0, 0);
+		Dynamic->X() = FVec3(-(2 * BoxHalfSize) - Separation, 0, 0);
+		Dynamic->V() = FVec3(InitialSpeed, 0, 0);
+		Dynamic->SetCCDEnabled(false);
+
+		// The position of the static has changed and statics don't automatically update bounds, so update explicitly
+		Static->UpdateWorldSpaceState(TRigidTransform<FReal, 3>(Static->X(), Static->R()), FVec3(0));
+
+		// Make sure the particles would collide if Dynamic wasn't a probe
+		::ChaosTest::SetParticleSimDataToCollide({ Static,Dynamic });
+
+		// 1 step should be enough to cause a collision event
+		Evolution.AdvanceOneTimeStep(Dt);
+		Evolution.EndFrame(Dt);
+
+		// Make sure a constraint got created
+		EXPECT_GT(Evolution.GetCollisionConstraints().NumConstraints(), 0);
+
+		// Make sure that the velocity of the dynamic body wasn't affected
+		EXPECT_EQ(Dynamic->V(), FVec3(InitialSpeed, 0, 0));
+	}
+
+	GTEST_TEST(ProbeTests, ProbeBodyConstraintWithCCD)
+	{
+		const FReal BoxHalfSize = 50; // cm
+		const FReal Separation = 10; // cm
+		const FReal InitialSpeed = Separation + (3 * BoxHalfSize); // cm/s... fast enough to trigger ccd and close the distance in 1 second
+		const FReal Dt = 1;
+		FParticleUniqueIndicesMultithreaded UniqueIndices;
+		FPBDRigidsSOAs Particles(UniqueIndices);
+		THandleArray<FChaosPhysicsMaterial> PhysicalMaterials;
+		FPBDRigidsEvolutionGBF Evolution(Particles, PhysicalMaterials);
+		InitEvolutionSettings(Evolution);
+
+		// Create particles
+		TGeometryParticleHandle<FReal, 3>* Static = Evolution.CreateStaticParticles(1)[0];
+		TPBDRigidParticleHandle<FReal, 3>* Dynamic = Evolution.CreateDynamicParticles(1)[0];
+
+		// Create box geometry 
+		TUniquePtr<FImplicitObject> SmallBox(new TBox<FReal, 3>(FVec3(-BoxHalfSize, -BoxHalfSize, -BoxHalfSize), FVec3(BoxHalfSize, BoxHalfSize, BoxHalfSize)));
+		Static->SetGeometry(MakeSerializable(SmallBox));
+		AppendDynamicParticleConvexBox(*Dynamic, FVec3(BoxHalfSize), 0.0f);
+		Dynamic->SetGravityEnabled(false);
+
+		// Make the dynamic shape a probe
+		Dynamic->ShapesArray()[0]->SetIsProbe(true);
+
+		// Positions
+		Static->X() = FVec3(0, 0, 0);
+		Dynamic->X() = FVec3(-(2 * BoxHalfSize) - Separation, 0, 0);
+		Dynamic->V() = FVec3(InitialSpeed, 0, 0);
+		Dynamic->SetCCDEnabled(true);
+
+		// The position of the static has changed and statics don't automatically update bounds, so update explicitly
+		Static->UpdateWorldSpaceState(TRigidTransform<FReal, 3>(Static->X(), Static->R()), FVec3(0));
+
+		// Make sure the particles would collide if Dynamic wasn't a probe
+		::ChaosTest::SetParticleSimDataToCollide({ Static,Dynamic });
+
+		// 1 step should be enough to cause a collision event
+		Evolution.AdvanceOneTimeStep(Dt);
+		Evolution.EndFrame(Dt);
+
+		// Make sure a constraint got created
+		EXPECT_GT(Evolution.GetCollisionConstraints().NumConstraints(), 0);
+
+		// Even though the body was moving fast enough to hit the static box
+		// with CCD, make sure the constraint doesn't actually think it's CCD
+		for (FPBDCollisionConstraint* const Constraint : Evolution.GetCollisionConstraints().GetConstraints())
+		{
+			EXPECT_EQ(Constraint->GetCCDType(), ECollisionCCDType::Disabled);
+		}
+
+		// Make sure that the velocity of the dynamic body wasn't affected
+		EXPECT_EQ(Dynamic->V(), FVec3(InitialSpeed, 0, 0));
+	}
+}
