Tests for GJK numerical issue (disabled for now)

diff --git a/Engine/Source/Programs/HeadlessChaos/Private/HeadlessChaosTestGJK.cpp b/Engine/Source/Programs/HeadlessChaos/Private/HeadlessChaosTestGJK.cpp
--- a/Engine/Source/Programs/HeadlessChaos/Private/HeadlessChaosTestGJK.cpp
+++ b/Engine/Source/Programs/HeadlessChaos/Private/HeadlessChaosTestGJK.cpp
@@ -1379,8 +1379,8 @@ namespace ChaosTest
 	}
 
 	// When we have a capsule and box that are reported as initially-overlapping because they are within
-	// the GJK epsilon opf each other (but actually positively separated), verify that we get a zero time of impact.
-	// Previously the slihjtly-positive separation would result in a negative penetration and a positiove TOI.
+	// the GJK epsilon of each other (but actually positively separated), verify that we get a zero time of impact.
+	// Previously the slightly-positive separation would result in a negative penetration and a positive TOI.
 	// Bug fix: CL 10942094.
 	// NOTE: this issue no longer manifests with this example because GJK no longer reports this case as
 	// overlapping> The GJK epsilon no longer takes part in the distance calculation when the near point
@@ -2020,4 +2020,35 @@ namespace ChaosTest
 		EXPECT_NEAR(OutPosition.Y, -0.73, 1e-1);
 		EXPECT_NEAR(OutPosition.Z, 14.63, 1e-1);
 	}
+
+	// Regression test: Disabled until fixed
+	GTEST_TEST(GJKTests, DISABLED_GJK_LargeScaledBoxBoxTest)
+	{
+
+		TArray<FConvex::FVec3Type> ConvexParticles;
+		ConvexParticles.SetNum(8);
+
+		// This is a box with some small deviations
+		ConvexParticles[0] = { 500.000000, -500.000031, 2.84217094e-14 };
+		ConvexParticles[1] = { 500.000000, 499.999969, -50.0000153 };
+		ConvexParticles[2] = { 500.000000, -500.000031, -50.0000153 };
+		ConvexParticles[3] = { -500.000183, 499.999969, -50.0000153 };
+		ConvexParticles[4] = { -500.000183, -500.000031, 2.84217094e-14 };
+		ConvexParticles[5] = { -500.000183, -500.000031, -50.0000153 };
+		ConvexParticles[6] = { -500.000183, 499.999969, -2.84217094e-14 };
+		ConvexParticles[7] = { 500.000000, 499.999969, -2.84217094e-14 };
+
+		TUniquePtr<Chaos::FConvex> BigBox = MakeUnique<Chaos::FConvex>(ConvexParticles, 0.0f);
+
+		// These two boxes are clearly intersecting each other
+
+		Chaos::TBox<Chaos::FReal, 3> SmallBox({ -3200, -3200, -3200 }, { 3200, 3200, 3200 }, 0);
+
+		TImplicitObjectScaled<Chaos::FConvex> BigBoxScaled(MakeSerializable(BigBox), nullptr, FVec3(50, 50, 1));
+		const TVector<FReal, 3> Translation{16000, 16000, -500};
+
+		TRigidTransform<Chaos::FReal, 3> BToATM( Translation , TRotation<FReal, 3>::Identity);
+		EXPECT_TRUE(GJKIntersection(BigBoxScaled, SmallBox, BToATM, FReal(0), Chaos::TVector<FReal, 3>(-16000, -16000, 500)));		
+
+	}
 }
\ No newline at end of file
