Cube gets flung / bounces on landscape with 'Tick Physics Async' enabled

diff --git a/Engine/Source/Programs/HeadlessChaos/Private/HeadlessChaosTestEPA.cpp b/Engine/Source/Programs/HeadlessChaos/Private/HeadlessChaosTestEPA.cpp
--- a/Engine/Source/Programs/HeadlessChaos/Private/HeadlessChaosTestEPA.cpp
+++ b/Engine/Source/Programs/HeadlessChaos/Private/HeadlessChaosTestEPA.cpp
@@ -18,6 +18,8 @@ namespace ChaosTest
 {
 	using namespace Chaos;
 
+	constexpr FReal EpaEps = 1e-6;
+
 	void ValidFace(const FVec3* Verts, const TEPAWorkingArray<TEPAEntry<FReal>>& TetFaces, int32 Idx)
 	{
 		const TEPAEntry<FReal>& Entry = TetFaces[Idx];
@@ -78,9 +80,9 @@ namespace ChaosTest
 
 		EXPECT_GE(Entry.Distance, 0);	//positive distance since origin is inside tet
 
-		EXPECT_NEAR(Entry.DistanceToPlane(Verts[Entry.IdxBuffer[0]]), 0, 1e-6);
-		EXPECT_NEAR(Entry.DistanceToPlane(Verts[Entry.IdxBuffer[1]]), 0, 1e-6);
-		EXPECT_NEAR(Entry.DistanceToPlane(Verts[Entry.IdxBuffer[2]]), 0, 1e-6);
+		EXPECT_NEAR(Entry.DistanceToPlane(Verts[Entry.IdxBuffer[0]]), 0, EpaEps);
+		EXPECT_NEAR(Entry.DistanceToPlane(Verts[Entry.IdxBuffer[1]]), 0, EpaEps);
+		EXPECT_NEAR(Entry.DistanceToPlane(Verts[Entry.IdxBuffer[2]]), 0, EpaEps);
 	}
 
 	FVec3 ErrorSupport(const FVec3& V)
@@ -196,7 +198,7 @@ namespace ChaosTest
 				FVec3 Dir, WitnessA, WitnessB;
 
 				//Try EPA. Note that we are IGNORING the positive x vert to ensure a triangle right on the origin boundary works
-				EPA(VertsA, VertsB, ASupportNoPositiveX, EmptySupport, Penetration, Dir, WitnessA, WitnessB);
+				EPA(VertsA, VertsB, ASupportNoPositiveX, EmptySupport, Penetration, Dir, WitnessA, WitnessB, EpaEps);
 				EXPECT_NEAR(Penetration, 0, 1e-4);
 				EXPECT_VECTOR_NEAR(Dir, FVec3(1,0,0), 1e-4);
 				EXPECT_VECTOR_NEAR(WitnessA, FVec3(0), 1e-4);
@@ -224,7 +226,7 @@ namespace ChaosTest
 				FVec3 Dir, WitnessA, WitnessB;
 
 				//Try EPA. Note that we are IGNORING the positive x vert to ensure a triangle right on the origin boundary works
-				EPA(VertsA, VertsB, ASupportNoPositiveX, EmptySupport, Penetration, Dir, WitnessA, WitnessB);
+				EPA(VertsA, VertsB, ASupportNoPositiveX, EmptySupport, Penetration, Dir, WitnessA, WitnessB, EpaEps);
 				EXPECT_NEAR(Penetration, 0, 1e-4);
 				EXPECT_VECTOR_NEAR(Dir, FVec3(1, 0, 0), 1e-4);
 				EXPECT_VECTOR_NEAR(WitnessA, FVec3(0), 1e-4);
@@ -249,7 +251,7 @@ namespace ChaosTest
 				//touching so penetration 0, normal is 0,0,1
 				FReal Penetration;
 				FVec3 Dir, WitnessA, WitnessB;
-				EXPECT_EQ(EPA(VertsA, VertsB, ASupportNoX, EmptySupport, Penetration, Dir, WitnessA, WitnessB), EEPAResult::BadInitialSimplex);
+				EXPECT_EQ(EPA(VertsA, VertsB, ASupportNoX, EmptySupport, Penetration, Dir, WitnessA, WitnessB, EpaEps), EEPAResult::BadInitialSimplex);
 				EXPECT_EQ(Penetration, 0);
 				//EXPECT_VECTOR_NEAR(Dir, FVec3(0, 0, 1), 1e-7);
 				EXPECT_VECTOR_NEAR(WitnessA, FVec3(0), 1e-7);
@@ -406,7 +408,7 @@ namespace ChaosTest
 
 			FReal Penetration;
 			FVec3 Dir, WitnessA, WitnessB;
-			EXPECT_EQ(EPA(Tetrahedron, Zeros, SupportA, ZeroSupport, Penetration, Dir, WitnessA, WitnessB), EEPAResult::Ok);
+			EXPECT_EQ(EPA(Tetrahedron, Zeros, SupportA, ZeroSupport, Penetration, Dir, WitnessA, WitnessB, EpaEps), EEPAResult::Ok);
 			EXPECT_NEAR(Penetration, 0.5, 1e-4);
 			EXPECT_NEAR(Dir[0], -1, 1e-4);
 			EXPECT_NEAR(Dir[1], 0, 1e-4);
@@ -436,7 +438,7 @@ namespace ChaosTest
 
 			FReal Penetration;
 			FVec3 Dir, WitnessA, WitnessB;
-			EXPECT_EQ(EPA(Tetrahedron, Zeros, Support, ZeroSupport, Penetration, Dir, WitnessA, WitnessB), EEPAResult::MaxIterations);
+			EXPECT_EQ(EPA(Tetrahedron, Zeros, Support, ZeroSupport, Penetration, Dir, WitnessA, WitnessB, EpaEps), EEPAResult::MaxIterations);
 			EXPECT_GT(Penetration, 9);
 			EXPECT_LE(Penetration, 10);
 			EXPECT_GT(WitnessA.Size(), 9);	//don't know exact point, but should be 9 away from origin
@@ -460,7 +462,7 @@ namespace ChaosTest
 
 			FReal Penetration;
 			FVec3 Dir, WitnessA, WitnessB;
-			EXPECT_EQ(EPA(Tetrahedron, Zeros, Support, ZeroSupport, Penetration, Dir, WitnessA, WitnessB), EEPAResult::Ok);
+			EXPECT_EQ(EPA(Tetrahedron, Zeros, Support, ZeroSupport, Penetration, Dir, WitnessA, WitnessB, EpaEps), EEPAResult::MaxIterations);
 			EXPECT_NEAR(Penetration, 3, 1e-1);
 			EXPECT_NEAR(Dir[2], 0, 1e-1);	//don't know direction, but it should be in xy plane
 			EXPECT_NEAR(WitnessA.Size(), 3, 1e-1);	//don't know exact point, but should be 3 away from origin
@@ -482,7 +484,8 @@ namespace ChaosTest
 
 			FReal Penetration;
 			FVec3 Dir, WitnessA, WitnessB;
-			EXPECT_EQ(EPA(Tetrahedron, Zeros, Support, ZeroSupport, Penetration, Dir, WitnessA, WitnessB), EEPAResult::Ok);
+			const EEPAResult Result = EPA(Tetrahedron, Zeros, Support, ZeroSupport, Penetration, Dir, WitnessA, WitnessB, EpaEps);
+			EXPECT_TRUE(Result == EEPAResult::Ok);
 			EXPECT_NEAR(Penetration, 1, 1e-1);
 			EXPECT_NEAR(Dir[0], 0, 1e-1);
 			EXPECT_NEAR(Dir[1], 0, 1e-1);
@@ -521,7 +524,7 @@ namespace ChaosTest
 
 			FReal Penetration;
 			FVec3 Dir, WitnessA, WitnessB;
-			EXPECT_EQ(EPA(Tetrahedron, Zeros, Support, ZeroSupport, Penetration, Dir, WitnessA, WitnessB), EEPAResult::Ok);
+			EXPECT_EQ(EPA(Tetrahedron, Zeros, Support, ZeroSupport, Penetration, Dir, WitnessA, WitnessB, EpaEps), EEPAResult::Ok);
 			EXPECT_FLOAT_EQ(Penetration, 1);
 			EXPECT_NEAR(WitnessA.Size(), 1, 1e-1);	//don't know exact point, but should be 1 away from origin
 		}
@@ -552,7 +555,7 @@ namespace ChaosTest
 
 			FReal Penetration;
 			FVec3 Dir, WitnessA, WitnessB;
-			EXPECT_EQ(EPA(Tetrahedron, Zeros, Support, ZeroSupport, Penetration, Dir, WitnessA, WitnessB), EEPAResult::Ok);
+			EXPECT_EQ(EPA(Tetrahedron, Zeros, Support, ZeroSupport, Penetration, Dir, WitnessA, WitnessB, EpaEps), EEPAResult::Ok);
 			EXPECT_LT(Penetration, 0); // Negative penetration
 			EXPECT_NEAR(Dir.X, 1.0f, 0.001f);
 		}
@@ -584,7 +587,7 @@ namespace ChaosTest
 
 			FReal Penetration;
 			FVec3 Dir, WitnessA, WitnessB;
-			EXPECT_EQ(EPA(Tetrahedron, Zeros, Support, ZeroSupport, Penetration, Dir, WitnessA, WitnessB), EEPAResult::Ok);
+			EXPECT_EQ(EPA(Tetrahedron, Zeros, Support, ZeroSupport, Penetration, Dir, WitnessA, WitnessB, EpaEps), EEPAResult::Ok);
 			EXPECT_LT(Penetration, 0); // Negative penetration
 			EXPECT_NEAR(Dir.X, 1.0f, 0.001f);
 		}
diff --git a/Engine/Source/Runtime/Experimental/Chaos/Public/Chaos/EPA.h b/Engine/Source/Runtime/Experimental/Chaos/Public/Chaos/EPA.h
--- a/Engine/Source/Runtime/Experimental/Chaos/Public/Chaos/EPA.h
+++ b/Engine/Source/Runtime/Experimental/Chaos/Public/Chaos/EPA.h
@@ -345,7 +345,7 @@ void EPAComputeVisibilityBorder(TEPAWorkingArray<TEPAEntry<T>>& Entries, int32 E
 		TEPAEntry<T>& Entry = Entries[FloodEntry.EntryIdx];
 		if (!Entry.bObsolete)
 		{
-			if (Entry.DistanceToPlane(W) < 0)
+			if (Entry.DistanceToPlane(W) < UE_SMALL_NUMBER)
 			{
 				//W can't see this triangle so mark the edge as a border
 				OutBorderEdges.Add(FloodEntry);
