Fix Height Field Raycast bug on flat HeightField

diff --git a/Engine/Source/Programs/HeadlessChaos/Private/HeadlessChaosTestHeightfield.cpp b/Engine/Source/Programs/HeadlessChaos/Private/HeadlessChaosTestHeightfield.cpp
--- a/Engine/Source/Programs/HeadlessChaos/Private/HeadlessChaosTestHeightfield.cpp
+++ b/Engine/Source/Programs/HeadlessChaos/Private/HeadlessChaosTestHeightfield.cpp
@@ -355,6 +355,28 @@ namespace ChaosTest {
 		
 	}
 
+	void RaycastOnFlatHeightField()
+	{
+		int32 Rows = 64;
+		int32 Columns = 64;
+		FVec3 Scale(100.0, 100.0, 100.0);
+		TArray<FReal> Heights;
+		Heights.AddZeroed(Rows * Columns);
+
+		TArray<FReal> HeightsCopy = Heights;
+		FHeightField Heightfield(MoveTemp(HeightsCopy), TArray<uint8>(), Rows, Columns, Scale);
+		const auto& Bounds = Heightfield.BoundingBox();	//Current API forces us to do this to cache the bounds
+
+		FReal TOI;
+		FVec3 Position, Normal;
+		int32 FaceIdx = 0;
+	
+		const FVec3 Start(8224.6524537283822, 2631.7542424011549, 2265.2052028112184);
+		const FVec3 Dir(-0.92887444870972680, -0.11019885226781448, -0.35362193299208372);
+		EXPECT_TRUE(Heightfield.Raycast(Start, Dir, 2097152.0000000000, 0, TOI, Position, Normal, FaceIdx));
+		EXPECT_TRUE(Heightfield.Raycast(Start, Dir, 2097152.0000000000, UE_KINDA_SMALL_NUMBER, TOI, Position, Normal, FaceIdx));
+	}
+
 	void EditHeights()
 	{
 		const int32 Columns = 10;
@@ -901,6 +923,7 @@ namespace ChaosTest {
 	TEST(ChaosTests, Heightfield)
 	{
 		ChaosTest::Raycast();
+		ChaosTest::RaycastOnFlatHeightField();
 		ChaosTest::SweepTest();
 		ChaosTest::OverlapTest();
 		EditHeights();
diff --git a/Engine/Source/Runtime/Experimental/Chaos/Private/Chaos/HeightField.cpp b/Engine/Source/Runtime/Experimental/Chaos/Private/Chaos/HeightField.cpp
--- a/Engine/Source/Runtime/Experimental/Chaos/Private/Chaos/HeightField.cpp
+++ b/Engine/Source/Runtime/Experimental/Chaos/Private/Chaos/HeightField.cpp
@@ -1065,7 +1065,7 @@ namespace Chaos
 		FReal RayExitTime;
 		if(CachedBounds.RaycastFast(StartPoint, Dir, InvDir, bParallel, Length, InvCurrentLength, RayEntryTime, RayExitTime))
 		{
-			CurrentLength = RayExitTime + UE_SMALL_NUMBER; // to account for precision errors 
+			CurrentLength = RayExitTime + 1e-2; // to account for precision errors 
 			FVec3 NextStart = StartPoint + (Dir * RayEntryTime);
 
 			const FVec2 Scale2D(GeomData.Scale[0],GeomData.Scale[1]);
diff --git a/Engine/Source/Runtime/PhysicsCore/Private/SQAccelerator.cpp b/Engine/Source/Runtime/PhysicsCore/Private/SQAccelerator.cpp
--- a/Engine/Source/Runtime/PhysicsCore/Private/SQAccelerator.cpp
+++ b/Engine/Source/Runtime/PhysicsCore/Private/SQAccelerator.cpp
@@ -299,7 +299,7 @@ private:
 
 					const FVec3 DirLocal = ActorTM.InverseTransformVectorNoScale(Dir);
 					const FVec3 StartLocal = ActorTM.InverseTransformPositionNoScale(StartPoint);
-					bHit = Geom->Raycast(StartLocal, DirLocal, CurData->CurrentLength, /*Thickness=*/0.f, Distance, LocalPosition, LocalNormal, FaceIdx);
+					bHit = Geom->Raycast(StartLocal, DirLocal, CurData->CurrentLength, /*Thickness=*/0.0, Distance, LocalPosition, LocalNormal, FaceIdx);
 					if (bHit)
 					{
 						WorldPosition = ActorTM.TransformPositionNoScale(LocalPosition);
