Chaos - fix joint drive stiffness/damping (was using single value for all linear drives, and another single value for all angular drives)

diff --git a/Engine/Source/Programs/HeadlessChaos/Private/HeadlessChaosTestJointPlasticity.cpp b/Engine/Source/Programs/HeadlessChaos/Private/HeadlessChaosTestJointPlasticity.cpp
--- a/Engine/Source/Programs/HeadlessChaos/Private/HeadlessChaosTestJointPlasticity.cpp
+++ b/Engine/Source/Programs/HeadlessChaos/Private/HeadlessChaosTestJointPlasticity.cpp
@@ -30,8 +30,8 @@ namespace ChaosTest {
 
 		Test.JointSettings[0].bAngularSLerpPositionDriveEnabled = true;
 		Test.JointSettings[0].bAngularSLerpVelocityDriveEnabled = true;
-		Test.JointSettings[0].AngularDriveDamping = 500;
-		Test.JointSettings[0].AngularDriveStiffness = 500000.f;
+		Test.JointSettings[0].AngularDriveDamping = FVec3(500);
+		Test.JointSettings[0].AngularDriveStiffness = FVec3(500000.f);
 		
 		Test.JointSettings[0].AngularPlasticityLimit = PlasticityAngle * (PI/180.);
 
@@ -86,8 +86,8 @@ namespace ChaosTest {
 
 		Test.JointSettings[0].bAngularSLerpPositionDriveEnabled = true;
 		Test.JointSettings[0].bAngularSLerpVelocityDriveEnabled = true;
-		Test.JointSettings[0].AngularDriveDamping = 50;
-		Test.JointSettings[0].AngularDriveStiffness = 10000.f;
+		Test.JointSettings[0].AngularDriveDamping = FVec3(50);
+		Test.JointSettings[0].AngularDriveStiffness = FVec3(10000.f);
 
 		Test.JointSettings[0].AngularPlasticityLimit = PlasticityAngle * (PI / 180.);
 
diff --git a/Engine/Source/Programs/HeadlessChaos/Private/HeadlessChaosTestJointSolver.cpp b/Engine/Source/Programs/HeadlessChaos/Private/HeadlessChaosTestJointSolver.cpp
--- a/Engine/Source/Programs/HeadlessChaos/Private/HeadlessChaosTestJointSolver.cpp
+++ b/Engine/Source/Programs/HeadlessChaos/Private/HeadlessChaosTestJointSolver.cpp
@@ -356,8 +356,8 @@ namespace ChaosTest
 		SolverTestA.JointSettings.AngularMotionTypes[(int32)EJointAngularConstraintIndex::Swing2] = EJointMotionType::Free;
 		SolverTestA.JointSettings.bAngularSLerpPositionDriveEnabled = true;
 		SolverTestA.JointSettings.bAngularSLerpVelocityDriveEnabled = true;
-		SolverTestA.JointSettings.AngularDriveStiffness = Stiffness;
-		SolverTestA.JointSettings.AngularDriveDamping = Damping;
+		SolverTestA.JointSettings.AngularDriveStiffness = FVec3(Stiffness);
+		SolverTestA.JointSettings.AngularDriveDamping = FVec3(Damping);
 		SolverTestA.JointSettings.AngularDrivePositionTarget = FRotation3::FromIdentity();
 		SolverTestA.JointSettings.AngularDriveForceMode = ForceMode;
 
@@ -537,8 +537,8 @@ namespace ChaosTest
 		SolverTestA.JointSettings.AngularMotionTypes[(int32)EJointAngularConstraintIndex::Swing2] = EJointMotionType::Free;
 		SolverTestA.JointSettings.bAngularSLerpPositionDriveEnabled = true;
 		SolverTestA.JointSettings.bAngularSLerpVelocityDriveEnabled = true;
-		SolverTestA.JointSettings.AngularDriveStiffness = 80.0f;
-		SolverTestA.JointSettings.AngularDriveDamping = 1.0f;
+		SolverTestA.JointSettings.AngularDriveStiffness = FVec3(80.0f);
+		SolverTestA.JointSettings.AngularDriveDamping = FVec3(1.0f);
 		SolverTestA.JointSettings.AngularDrivePositionTarget = FRotation3::FromIdentity();
 
 		// Particle 0 is Kinematic
@@ -585,7 +585,7 @@ namespace ChaosTest
 		}
 
 		FReal GravityAngAcc = SolverTestA.Body1.InvI().M[1][1] * SolverTestA.Mass1 * FVec3::CrossProduct(SolverTestA.Body1.P(), SolverTestA.Gravity).Y;
-		FReal SpringAngAcc = -SolverTestA.JointSettings.AngularDriveStiffness * OutAngles1.Y;
+		FReal SpringAngAcc = -SolverTestA.JointSettings.AngularDriveStiffness.X * OutAngles1.Y;
 		EXPECT_NEAR(SpringAngAcc, -GravityAngAcc, 5.0f);
 
 		// Verify that X and Z angles are zero, and that Y is almost zero
@@ -595,4 +595,153 @@ namespace ChaosTest
 		//EXPECT_LT(OutAngles1.Y, FMath::DegreesToRadians(1));
 		//EXPECT_NEAR(OutAngles1.Z, 0.0f, KINDA_SMALL_NUMBER);
 	}
+
+	// Linear Drive Stiffness and Damping test
+	// A dynamic body attached to a kinematic boidy by a spring with different stiffness and damping on each axis.
+	GTEST_TEST(JointSolverTests, TestJointSolver_KinematicDynamic_LinearDrive)
+	{
+		FJointSolverTest SolverTest;
+
+		FReal Dt = 0.02f;
+		int32 NumIts = 100;
+		SolverTest.Gravity = FVec3(0, 0, -1000);
+		SolverTest.Mass1 = 100.0f;
+		SolverTest.Inertia1 = FVec3(10000, 10000, 10000);
+
+		SolverTest.JointSettings.LinearMotionTypes = { EJointMotionType::Free, EJointMotionType::Free, EJointMotionType::Free };
+		SolverTest.JointSettings.bSoftLinearLimitsEnabled = false;
+		SolverTest.JointSettings.LinearDriveStiffness = FVec3(1000, 0, 100);
+		SolverTest.JointSettings.LinearDriveDamping = FVec3(70, 0, 20);
+		SolverTest.JointSettings.LinearDriveForceMode = EJointForceMode::Acceleration;
+		SolverTest.JointSettings.bLinearPositionDriveEnabled = TVec3<bool>(true);
+		SolverTest.JointSettings.bLinearVelocityDriveEnabled = TVec3<bool>(true);
+
+		// Particle 0 is Kinematic
+		// Particle 1 is Dynamic
+		SolverTest.Body0.SetP(FVec3(0, 0, 0));
+		SolverTest.Body1.SetP(FVec3(100, 100, 100));
+
+		SolverTest.Init();
+
+		int32 It = 0;
+		while (It++ < NumIts)
+		{
+			SolverTest.Tick(Dt);
+		}
+
+		// Dynamic body should be close to X = 0, Y = 100, Z = G/K
+		const FReal Tolerance = UE_KINDA_SMALL_NUMBER;
+		EXPECT_NEAR(SolverTest.Body1.P().X, 0, Tolerance);
+		EXPECT_NEAR(SolverTest.Body1.P().Y, 100, Tolerance);
+		EXPECT_NEAR(SolverTest.Body1.P().Z, SolverTest.Gravity.Z / SolverTest.JointSettings.LinearDriveStiffness.Z, Tolerance);
+	}
+
+	// A dynamic body connected to a kinematic body by a twist spring.
+	// Body starts rotated about twist and swing and the twist angle should decrease to 0, leaving the swing untouched.
+	GTEST_TEST(JointSolverTests, TestJointSolver_KinematicDynamic_TwistDrive)
+	{
+		FJointSolverTest SolverTest;
+
+		FReal Dt = 0.02f;
+		int32 NumIts = 50;
+		FReal Stiffness = 1000;
+		FReal Damping = 70;
+		FVec3 Connector = FVec3(100, 0, 0);
+		FReal InitSwingAngle = FMath::DegreesToRadians(45);
+		FReal InitTwistAngle = FMath::DegreesToRadians(45);
+		FRotation3 InitRot = FRotation3::FromAxisAngle(FVec3(0,1,0), InitSwingAngle) * FRotation3::FromAxisAngle(FVec3(1, 0, 0), InitTwistAngle);
+
+		SolverTest.Gravity = FVec3(0);
+		SolverTest.Mass1 = 100.0f;
+		SolverTest.Inertia1 = FVec3(10000, 10000, 10000);
+		SolverTest.Body1.SetQ(InitRot);
+
+		// Set up a damped Twist drive. Swing drive is also active but with 0 stiffness and damping
+		SolverTest.JointSettings.AngularMotionTypes[(int32)EJointAngularConstraintIndex::Twist] = EJointMotionType::Free;
+		SolverTest.JointSettings.AngularMotionTypes[(int32)EJointAngularConstraintIndex::Swing1] = EJointMotionType::Free;
+		SolverTest.JointSettings.AngularMotionTypes[(int32)EJointAngularConstraintIndex::Swing2] = EJointMotionType::Free;
+		SolverTest.JointSettings.bAngularTwistPositionDriveEnabled = true;
+		SolverTest.JointSettings.bAngularTwistVelocityDriveEnabled = true;
+		SolverTest.JointSettings.bAngularSwingPositionDriveEnabled = true;
+		SolverTest.JointSettings.bAngularSwingVelocityDriveEnabled = true;
+		SolverTest.JointSettings.AngularDriveStiffness = FVec3(Stiffness, 0, 0);
+		SolverTest.JointSettings.AngularDriveDamping = FVec3(Damping, 0, 0);
+		SolverTest.JointSettings.AngularDrivePositionTarget = FRotation3::FromIdentity();
+		SolverTest.JointSettings.AngularDriveForceMode = EJointForceMode::Acceleration;
+
+		SolverTest.Init();
+
+		FReal TwistAngle, Swing1Angle, Swing2Angle;
+		FPBDJointUtilities::GetSwingTwistAngles(FRotation3::FromIdentity(), SolverTest.Body1.Q(), TwistAngle, Swing1Angle, Swing2Angle);
+		EXPECT_NEAR(TwistAngle, InitTwistAngle, UE_KINDA_SMALL_NUMBER);
+		EXPECT_NEAR(Swing1Angle, 0, UE_KINDA_SMALL_NUMBER);
+		EXPECT_NEAR(Swing2Angle, InitSwingAngle, UE_KINDA_SMALL_NUMBER);
+
+		int32 It = 0;
+		while (It++ < NumIts)
+		{
+			SolverTest.Tick(Dt);
+
+			FPBDJointUtilities::GetSwingTwistAngles(FRotation3::FromIdentity(), SolverTest.Body1.Q(), TwistAngle, Swing1Angle, Swing2Angle);
+		}
+
+		const FReal Tolerance = UE_KINDA_SMALL_NUMBER;
+		EXPECT_NEAR(TwistAngle, 0, Tolerance);
+		EXPECT_NEAR(Swing1Angle, 0, Tolerance);
+		EXPECT_NEAR(Swing2Angle, InitSwingAngle, Tolerance);
+	}
+
+	// A dynamic body connected to a kinematic body by a swing spring.
+	// Body starts rotated about twist and swing and the swing angle should decrease to 0, leaving the twist untouched.
+	GTEST_TEST(JointSolverTests, TestJointSolver_KinematicDynamic_Swing2Drive)
+	{
+		FJointSolverTest SolverTest;
+
+		FReal Dt = 0.02f;
+		int32 NumIts = 50;
+		FReal Stiffness = 1000;
+		FReal Damping = 70;
+		FReal InitSwingAngle = FMath::DegreesToRadians(45);
+		FReal InitTwistAngle = FMath::DegreesToRadians(0);
+		FRotation3 InitRot = FRotation3::FromAxisAngle(FVec3(0, 1, 0), InitSwingAngle) * FRotation3::FromAxisAngle(FVec3(1, 0, 0), InitTwistAngle);
+
+		SolverTest.Gravity = FVec3(0);
+		SolverTest.Mass1 = 100.0f;
+		SolverTest.Inertia1 = FVec3(10000, 10000, 10000);
+		SolverTest.Body1.SetQ(InitRot);
+
+		// Set up a damped Twist drive. Swing drive is also active but with 0 stiffness and damping
+		SolverTest.JointSettings.AngularMotionTypes[(int32)EJointAngularConstraintIndex::Twist] = EJointMotionType::Free;
+		SolverTest.JointSettings.AngularMotionTypes[(int32)EJointAngularConstraintIndex::Swing1] = EJointMotionType::Free;
+		SolverTest.JointSettings.AngularMotionTypes[(int32)EJointAngularConstraintIndex::Swing2] = EJointMotionType::Free;
+		SolverTest.JointSettings.bAngularTwistPositionDriveEnabled = true;
+		SolverTest.JointSettings.bAngularTwistVelocityDriveEnabled = true;
+		SolverTest.JointSettings.bAngularSwingPositionDriveEnabled = true;
+		SolverTest.JointSettings.bAngularSwingVelocityDriveEnabled = true;
+		SolverTest.JointSettings.AngularDriveStiffness = FVec3(0, Stiffness, Stiffness);
+		SolverTest.JointSettings.AngularDriveDamping = FVec3(0, Damping, Damping);
+		SolverTest.JointSettings.AngularDrivePositionTarget = FRotation3::FromIdentity();
+		SolverTest.JointSettings.AngularDriveForceMode = EJointForceMode::Acceleration;
+
+		SolverTest.Init();
+
+		FReal TwistAngle, Swing1Angle, Swing2Angle;
+		FPBDJointUtilities::GetSwingTwistAngles(FRotation3::FromIdentity(), SolverTest.Body1.Q(), TwistAngle, Swing1Angle, Swing2Angle);
+		EXPECT_NEAR(TwistAngle, InitTwistAngle, UE_KINDA_SMALL_NUMBER);
+		EXPECT_NEAR(Swing1Angle, 0, UE_KINDA_SMALL_NUMBER);
+		EXPECT_NEAR(Swing2Angle, InitSwingAngle, UE_KINDA_SMALL_NUMBER);
+
+		int32 It = 0;
+		while (It++ < NumIts)
+		{
+			SolverTest.Tick(Dt);
+
+			FPBDJointUtilities::GetSwingTwistAngles(FRotation3::FromIdentity(), SolverTest.Body1.Q(), TwistAngle, Swing1Angle, Swing2Angle);
+		}
+
+		EXPECT_NEAR(TwistAngle, InitTwistAngle, UE_KINDA_SMALL_NUMBER);
+		EXPECT_NEAR(Swing1Angle, 0, UE_KINDA_SMALL_NUMBER);
+		EXPECT_NEAR(Swing2Angle, 0, UE_KINDA_SMALL_NUMBER);
+	}
+
 }
\ No newline at end of file
diff --git a/Engine/Source/Runtime/Engine/Private/PhysicsEngine/Experimental/PhysInterface_Chaos.cpp b/Engine/Source/Runtime/Engine/Private/PhysicsEngine/Experimental/PhysInterface_Chaos.cpp
--- a/Engine/Source/Runtime/Engine/Private/PhysicsEngine/Experimental/PhysInterface_Chaos.cpp
+++ b/Engine/Source/Runtime/Engine/Private/PhysicsEngine/Experimental/PhysInterface_Chaos.cpp
@@ -397,8 +397,8 @@ void FPhysInterface_Chaos::UpdateLinearDrive_AssumesLocked(const FPhysicsConstra
 			}
 
 			Constraint->SetLinearDriveForceMode(Chaos::EJointForceMode::Acceleration);
-			Constraint->SetLinearDriveStiffness(Chaos::ConstraintSettings::LinearDriveStiffnessScale() * FMath::Max3(InDriveParams.XDrive.Stiffness, InDriveParams.YDrive.Stiffness, InDriveParams.ZDrive.Stiffness));
-			Constraint->SetLinearDriveDamping(Chaos::ConstraintSettings::LinearDriveDampingScale() * FMath::Max3(InDriveParams.XDrive.Damping, InDriveParams.YDrive.Damping, InDriveParams.ZDrive.Damping));
+			Constraint->SetLinearDriveStiffness(Chaos::ConstraintSettings::LinearDriveStiffnessScale() * Chaos::FVec3(InDriveParams.XDrive.Stiffness, InDriveParams.YDrive.Stiffness, InDriveParams.ZDrive.Stiffness));
+			Constraint->SetLinearDriveDamping(Chaos::ConstraintSettings::LinearDriveDampingScale() * Chaos::FVec3(InDriveParams.XDrive.Damping, InDriveParams.YDrive.Damping, InDriveParams.ZDrive.Damping));
 		}
 	}
 }
@@ -465,8 +465,16 @@ void FPhysInterface_Chaos::UpdateAngularDrive_AssumesLocked(const FPhysicsConstr
 			}
 
 			Constraint->SetAngularDriveForceMode(Chaos::EJointForceMode::Acceleration);
-			Constraint->SetAngularDriveStiffness(Chaos::ConstraintSettings::AngularDriveStiffnessScale() * FMath::Max3(InDriveParams.SlerpDrive.Stiffness, InDriveParams.TwistDrive.Stiffness, InDriveParams.SwingDrive.Stiffness));
-			Constraint->SetAngularDriveDamping(Chaos::ConstraintSettings::AngularDriveDampingScale() * FMath::Max3(InDriveParams.SlerpDrive.Damping, InDriveParams.TwistDrive.Damping, InDriveParams.SwingDrive.Damping));
+			if (InDriveParams.AngularDriveMode == EAngularDriveMode::TwistAndSwing)
+			{
+				Constraint->SetAngularDriveStiffness(Chaos::ConstraintSettings::AngularDriveStiffnessScale() * Chaos::FVec3(InDriveParams.TwistDrive.Stiffness, InDriveParams.SwingDrive.Stiffness, InDriveParams.SwingDrive.Stiffness));
+				Constraint->SetAngularDriveDamping(Chaos::ConstraintSettings::AngularDriveDampingScale() * Chaos::FVec3(InDriveParams.TwistDrive.Damping, InDriveParams.SwingDrive.Damping, InDriveParams.SwingDrive.Damping));
+			}
+			else
+			{
+				Constraint->SetAngularDriveStiffness(Chaos::ConstraintSettings::AngularDriveStiffnessScale() * Chaos::FVec3(InDriveParams.SlerpDrive.Stiffness));
+				Constraint->SetAngularDriveDamping(Chaos::ConstraintSettings::AngularDriveDampingScale() * Chaos::FVec3(InDriveParams.SlerpDrive.Damping));
+			}
 		}
 	}
 }
diff --git a/Engine/Source/Runtime/Engine/Private/PhysicsEngine/ImmediatePhysics/ImmediatePhysicsChaos/ImmediatePhysicsJointHandle_Chaos.cpp b/Engine/Source/Runtime/Engine/Private/PhysicsEngine/ImmediatePhysics/ImmediatePhysicsChaos/ImmediatePhysicsJointHandle_Chaos.cpp
--- a/Engine/Source/Runtime/Engine/Private/PhysicsEngine/ImmediatePhysics/ImmediatePhysicsChaos/ImmediatePhysicsJointHandle_Chaos.cpp
+++ b/Engine/Source/Runtime/Engine/Private/PhysicsEngine/ImmediatePhysics/ImmediatePhysicsChaos/ImmediatePhysicsJointHandle_Chaos.cpp
@@ -91,8 +91,8 @@ namespace ImmediatePhysics_Chaos
 		ConstraintSettings.bLinearVelocityDriveEnabled[0] = Profile.LinearDrive.XDrive.bEnableVelocityDrive;
 		ConstraintSettings.bLinearVelocityDriveEnabled[1] = Profile.LinearDrive.YDrive.bEnableVelocityDrive;
 		ConstraintSettings.bLinearVelocityDriveEnabled[2] = Profile.LinearDrive.ZDrive.bEnableVelocityDrive;
-		ConstraintSettings.LinearDriveStiffness = Chaos::ConstraintSettings::LinearDriveStiffnessScale() * Profile.LinearDrive.XDrive.Stiffness;
-		ConstraintSettings.LinearDriveDamping = Chaos::ConstraintSettings::LinearDriveDampingScale() * Profile.LinearDrive.XDrive.Damping;
+		ConstraintSettings.LinearDriveStiffness = Chaos::ConstraintSettings::LinearDriveStiffnessScale() * Chaos::FVec3(Profile.LinearDrive.XDrive.Stiffness, Profile.LinearDrive.YDrive.Stiffness, Profile.LinearDrive.ZDrive.Stiffness);
+		ConstraintSettings.LinearDriveDamping = Chaos::ConstraintSettings::LinearDriveDampingScale() * Chaos::FVec3(Profile.LinearDrive.XDrive.Damping, Profile.LinearDrive.YDrive.Damping, Profile.LinearDrive.ZDrive.Damping);
 		ConstraintSettings.LinearDriveForceMode = EJointForceMode::Acceleration;
 
 		ConstraintSettings.AngularDrivePositionTarget = FQuat(Profile.AngularDrive.OrientationTarget);
@@ -100,18 +100,20 @@ namespace ImmediatePhysics_Chaos
 
 		if (Profile.AngularDrive.AngularDriveMode == EAngularDriveMode::SLERP)
 		{
+			ConstraintSettings.AngularDriveStiffness = Chaos::ConstraintSettings::AngularDriveStiffnessScale() * FVec3(Profile.AngularDrive.SlerpDrive.Stiffness);
+			ConstraintSettings.AngularDriveDamping = Chaos::ConstraintSettings::AngularDriveDampingScale() * FVec3(Profile.AngularDrive.SlerpDrive.Damping);
 			ConstraintSettings.bAngularSLerpPositionDriveEnabled = Profile.AngularDrive.SlerpDrive.bEnablePositionDrive;
 			ConstraintSettings.bAngularSLerpVelocityDriveEnabled = Profile.AngularDrive.SlerpDrive.bEnableVelocityDrive;
 		}
 		else
 		{
+			ConstraintSettings.AngularDriveStiffness = Chaos::ConstraintSettings::AngularDriveStiffnessScale() * FVec3(Profile.AngularDrive.TwistDrive.Stiffness, Profile.AngularDrive.SwingDrive.Stiffness, Profile.AngularDrive.SwingDrive.Stiffness);
+			ConstraintSettings.AngularDriveDamping = Chaos::ConstraintSettings::AngularDriveDampingScale() * FVec3(Profile.AngularDrive.TwistDrive.Damping, Profile.AngularDrive.SwingDrive.Damping, Profile.AngularDrive.SwingDrive.Damping);
 			ConstraintSettings.bAngularTwistPositionDriveEnabled = Profile.AngularDrive.TwistDrive.bEnablePositionDrive;
 			ConstraintSettings.bAngularTwistVelocityDriveEnabled = Profile.AngularDrive.TwistDrive.bEnableVelocityDrive;
 			ConstraintSettings.bAngularSwingPositionDriveEnabled = Profile.AngularDrive.SwingDrive.bEnablePositionDrive;
 			ConstraintSettings.bAngularSwingVelocityDriveEnabled = Profile.AngularDrive.SwingDrive.bEnableVelocityDrive;
 		}
-		ConstraintSettings.AngularDriveStiffness = Chaos::ConstraintSettings::AngularDriveStiffnessScale() * Profile.AngularDrive.TwistDrive.Stiffness;
-		ConstraintSettings.AngularDriveDamping = Chaos::ConstraintSettings::AngularDriveDampingScale() * Profile.AngularDrive.TwistDrive.Damping;
 		ConstraintSettings.AngularDriveForceMode = EJointForceMode::Acceleration;
 
 		ConstraintSettings.LinearBreakForce = (Profile.bLinearBreakable) ? Chaos::ConstraintSettings::LinearBreakScale() * Profile.LinearBreakThreshold : FLT_MAX;
diff --git a/Engine/Source/Runtime/Engine/Private/PhysicsEngine/PhysicsHandleComponent.cpp b/Engine/Source/Runtime/Engine/Private/PhysicsEngine/PhysicsHandleComponent.cpp
--- a/Engine/Source/Runtime/Engine/Private/PhysicsEngine/PhysicsHandleComponent.cpp
+++ b/Engine/Source/Runtime/Engine/Private/PhysicsEngine/PhysicsHandleComponent.cpp
@@ -178,14 +178,14 @@ void UPhysicsHandleComponent::UpdateDriveSettings()
 
 				if (LocationMotionType != Chaos::EJointMotionType::Locked)
 				{
-					Constraint->SetLinearDriveStiffness(LinearStiffness);
-					Constraint->SetLinearDriveDamping(LinearDamping);
+					Constraint->SetLinearDriveStiffness(Chaos::FVec3(LinearStiffness));
+					Constraint->SetLinearDriveDamping(Chaos::FVec3(LinearDamping));
 				}
 
 				if (bRotationConstrained && RotationMotionType != Chaos::EJointMotionType::Locked)
 				{
-					Constraint->SetAngularDriveStiffness(AngularStiffness);
-					Constraint->SetAngularDriveDamping(AngularDamping);
+					Constraint->SetAngularDriveStiffness(Chaos::FVec3(AngularStiffness));
+					Constraint->SetAngularDriveDamping(Chaos::FVec3(AngularDamping));
 				}
 			}
 		});
diff --git a/Engine/Source/Runtime/Experimental/Chaos/Private/Chaos/PBDJointConstraintUtilities.cpp b/Engine/Source/Runtime/Experimental/Chaos/Private/Chaos/PBDJointConstraintUtilities.cpp
--- a/Engine/Source/Runtime/Experimental/Chaos/Private/Chaos/PBDJointConstraintUtilities.cpp
+++ b/Engine/Source/Runtime/Experimental/Chaos/Private/Chaos/PBDJointConstraintUtilities.cpp
@@ -483,7 +483,7 @@ namespace Chaos
 	{
 		if (JointSettings.bLinearPositionDriveEnabled[AxisIndex])
 		{
-			return (SolverSettings.LinearDriveStiffnessOverride >= 0.0f) ? SolverSettings.LinearDriveStiffnessOverride : JointSettings.LinearDriveStiffness;
+			return (SolverSettings.LinearDriveStiffnessOverride >= 0.0f) ? SolverSettings.LinearDriveStiffnessOverride : JointSettings.LinearDriveStiffness[AxisIndex];
 		}
 		return 0.0f;
 	}
@@ -495,7 +495,7 @@ namespace Chaos
 	{
 		if (JointSettings.bLinearVelocityDriveEnabled[AxisIndex])
 		{
-			return (SolverSettings.LinearDriveDampingOverride >= 0.0f) ? SolverSettings.LinearDriveDampingOverride : JointSettings.LinearDriveDamping;
+			return (SolverSettings.LinearDriveDampingOverride >= 0.0f) ? SolverSettings.LinearDriveDampingOverride : JointSettings.LinearDriveDamping[AxisIndex];
 		}
 		return 0.0f;
 	}
@@ -506,7 +506,7 @@ namespace Chaos
 	{
 		if (JointSettings.bAngularTwistPositionDriveEnabled)
 		{
-			return (SolverSettings.AngularDriveStiffnessOverride >= 0.0f) ? SolverSettings.AngularDriveStiffnessOverride : JointSettings.AngularDriveStiffness;
+			return (SolverSettings.AngularDriveStiffnessOverride >= 0.0f) ? SolverSettings.AngularDriveStiffnessOverride : JointSettings.AngularDriveStiffness[(int)EJointAngularConstraintIndex::Twist];
 		}
 		return 0.0f;
 	}
@@ -517,7 +517,7 @@ namespace Chaos
 	{
 		if (JointSettings.bAngularTwistVelocityDriveEnabled)
 		{
-			return (SolverSettings.AngularDriveDampingOverride >= 0.0f) ? SolverSettings.AngularDriveDampingOverride : JointSettings.AngularDriveDamping;
+			return (SolverSettings.AngularDriveDampingOverride >= 0.0f) ? SolverSettings.AngularDriveDampingOverride : JointSettings.AngularDriveDamping[(int)EJointAngularConstraintIndex::Twist];
 		}
 		return 0.0f;
 	}
@@ -528,7 +528,7 @@ namespace Chaos
 	{
 		if (JointSettings.bAngularSwingPositionDriveEnabled)
 		{
-			return (SolverSettings.AngularDriveStiffnessOverride >= 0.0f) ? SolverSettings.AngularDriveStiffnessOverride : JointSettings.AngularDriveStiffness;
+			return (SolverSettings.AngularDriveStiffnessOverride >= 0.0f) ? SolverSettings.AngularDriveStiffnessOverride : JointSettings.AngularDriveStiffness[(int)EJointAngularConstraintIndex::Swing1];
 		}
 		return 0.0f;
 	}
@@ -539,7 +539,7 @@ namespace Chaos
 	{
 		if (JointSettings.bAngularSwingVelocityDriveEnabled)
 		{
-			return (SolverSettings.AngularDriveDampingOverride >= 0.0f) ? SolverSettings.AngularDriveDampingOverride : JointSettings.AngularDriveDamping;
+			return (SolverSettings.AngularDriveDampingOverride >= 0.0f) ? SolverSettings.AngularDriveDampingOverride : JointSettings.AngularDriveDamping[(int)EJointAngularConstraintIndex::Swing1];
 		}
 		return 0.0f;
 	}
@@ -550,7 +550,7 @@ namespace Chaos
 	{
 		if (JointSettings.bAngularSLerpPositionDriveEnabled)
 		{
-			return (SolverSettings.AngularDriveStiffnessOverride >= 0.0f) ? SolverSettings.AngularDriveStiffnessOverride : JointSettings.AngularDriveStiffness;
+			return (SolverSettings.AngularDriveStiffnessOverride >= 0.0f) ? SolverSettings.AngularDriveStiffnessOverride : JointSettings.AngularDriveStiffness[(int)EJointAngularConstraintIndex::Twist];
 		}
 		return 0.0f;
 	}
@@ -561,7 +561,7 @@ namespace Chaos
 	{
 		if (JointSettings.bAngularSLerpVelocityDriveEnabled)
 		{
-			return (SolverSettings.AngularDriveDampingOverride >= 0.0f) ? SolverSettings.AngularDriveDampingOverride : JointSettings.AngularDriveDamping;
+			return (SolverSettings.AngularDriveDampingOverride >= 0.0f) ? SolverSettings.AngularDriveDampingOverride : JointSettings.AngularDriveDamping[(int)EJointAngularConstraintIndex::Twist];
 		}
 		return 0.0f;
 	}
diff --git a/Engine/Source/Runtime/Experimental/Chaos/Private/Chaos/PBDJointConstraints.cpp b/Engine/Source/Runtime/Experimental/Chaos/Private/Chaos/PBDJointConstraints.cpp
--- a/Engine/Source/Runtime/Experimental/Chaos/Private/Chaos/PBDJointConstraints.cpp
+++ b/Engine/Source/Runtime/Experimental/Chaos/Private/Chaos/PBDJointConstraints.cpp
@@ -226,8 +226,8 @@ namespace Chaos
 		, bLinearPositionDriveEnabled(TVector<bool, 3>(false, false, false))
 		, bLinearVelocityDriveEnabled(TVector<bool, 3>(false, false, false))
 		, LinearDriveForceMode(EJointForceMode::Acceleration)
-		, LinearDriveStiffness(0)
-		, LinearDriveDamping(0)
+		, LinearDriveStiffness(FVec3(0))
+		, LinearDriveDamping(FVec3(0))
 		, AngularDrivePositionTarget(FRotation3::FromIdentity())
 		, AngularDriveVelocityTarget(FVec3(0, 0, 0))
 		, bAngularSLerpPositionDriveEnabled(false)
@@ -237,8 +237,8 @@ namespace Chaos
 		, bAngularSwingPositionDriveEnabled(false)
 		, bAngularSwingVelocityDriveEnabled(false)
 		, AngularDriveForceMode(EJointForceMode::Acceleration)
-		, AngularDriveStiffness(0)
-		, AngularDriveDamping(0)
+		, AngularDriveStiffness(FVec3(0))
+		, AngularDriveDamping(FVec3(0))
 		, LinearBreakForce(FLT_MAX)
 		, LinearPlasticityLimit(FLT_MAX)
 		, LinearPlasticityType(EPlasticityType::Free)
diff --git a/Engine/Source/Runtime/Experimental/Chaos/Public/Chaos/JointProperties.inl b/Engine/Source/Runtime/Experimental/Chaos/Public/Chaos/JointProperties.inl
--- a/Engine/Source/Runtime/Experimental/Chaos/Public/Chaos/JointProperties.inl
+++ b/Engine/Source/Runtime/Experimental/Chaos/Public/Chaos/JointProperties.inl
@@ -29,8 +29,8 @@ CHAOS_INNER_JOINT_PROPERTY(JointSettings, LinearDriveForceMode, LinearDriveForce
 CHAOS_INNER_JOINT_PROPERTY(JointSettings, LinearMotionTypesX, LinearMotionTypes[0], EJointMotionType)
 CHAOS_INNER_JOINT_PROPERTY(JointSettings, LinearMotionTypesY, LinearMotionTypes[1], EJointMotionType)
 CHAOS_INNER_JOINT_PROPERTY(JointSettings, LinearMotionTypesZ, LinearMotionTypes[2], EJointMotionType)
-CHAOS_INNER_JOINT_PROPERTY(JointSettings, LinearDriveStiffness, LinearDriveStiffness, FReal)
-CHAOS_INNER_JOINT_PROPERTY(JointSettings, LinearDriveDamping, LinearDriveDamping, FReal)
+CHAOS_INNER_JOINT_PROPERTY(JointSettings, LinearDriveStiffness, LinearDriveStiffness, FVec3)
+CHAOS_INNER_JOINT_PROPERTY(JointSettings, LinearDriveDamping, LinearDriveDamping, FVec3)
 CHAOS_INNER_JOINT_PROPERTY(JointSettings, ContactTransferScale, ContactTransferScale, FReal)
 CHAOS_INNER_JOINT_PROPERTY(JointSettings, AngularSLerpPositionDriveEnabled, bAngularSLerpPositionDriveEnabled, bool)
 CHAOS_INNER_JOINT_PROPERTY(JointSettings, AngularTwistPositionDriveEnabled, bAngularTwistPositionDriveEnabled, bool)
@@ -44,8 +44,8 @@ CHAOS_INNER_JOINT_PROPERTY(JointSettings, AngularDriveForceMode, AngularDriveFor
 CHAOS_INNER_JOINT_PROPERTY(JointSettings, AngularMotionTypesX, AngularMotionTypes[0], EJointMotionType)
 CHAOS_INNER_JOINT_PROPERTY(JointSettings, AngularMotionTypesY, AngularMotionTypes[1], EJointMotionType)
 CHAOS_INNER_JOINT_PROPERTY(JointSettings, AngularMotionTypesZ, AngularMotionTypes[2], EJointMotionType)
-CHAOS_INNER_JOINT_PROPERTY(JointSettings, AngularDriveStiffness, AngularDriveStiffness, FReal)
-CHAOS_INNER_JOINT_PROPERTY(JointSettings, AngularDriveDamping, AngularDriveDamping, FReal)
+CHAOS_INNER_JOINT_PROPERTY(JointSettings, AngularDriveStiffness, AngularDriveStiffness, FVec3)
+CHAOS_INNER_JOINT_PROPERTY(JointSettings, AngularDriveDamping, AngularDriveDamping, FVec3)
 CHAOS_INNER_JOINT_PROPERTY(JointSettings, Stiffness, Stiffness, FReal)
 CHAOS_INNER_JOINT_PROPERTY(JointSettings, SoftLinearLimitsEnabled, bSoftLinearLimitsEnabled, bool)
 CHAOS_INNER_JOINT_PROPERTY(JointSettings, SoftTwistLimitsEnabled, bSoftTwistLimitsEnabled, bool)
diff --git a/Engine/Source/Runtime/Experimental/Chaos/Public/Chaos/PBDJointConstraintTypes.h b/Engine/Source/Runtime/Experimental/Chaos/Public/Chaos/PBDJointConstraintTypes.h
--- a/Engine/Source/Runtime/Experimental/Chaos/Public/Chaos/PBDJointConstraintTypes.h
+++ b/Engine/Source/Runtime/Experimental/Chaos/Public/Chaos/PBDJointConstraintTypes.h
@@ -171,8 +171,8 @@ namespace Chaos
 		TVector<bool, 3> bLinearPositionDriveEnabled;
 		TVector<bool, 3> bLinearVelocityDriveEnabled;
 		EJointForceMode LinearDriveForceMode;
-		FReal LinearDriveStiffness;
-		FReal LinearDriveDamping;
+		FVec3 LinearDriveStiffness;
+		FVec3 LinearDriveDamping;
 
 		FRotation3 AngularDrivePositionTarget;
 		FVec3 AngularDriveVelocityTarget;
@@ -184,8 +184,8 @@ namespace Chaos
 		bool bAngularSwingPositionDriveEnabled;
 		bool bAngularSwingVelocityDriveEnabled;
 		EJointForceMode AngularDriveForceMode;
-		FReal AngularDriveStiffness;
-		FReal AngularDriveDamping;
+		FVec3 AngularDriveStiffness;
+		FVec3 AngularDriveDamping;
 
 		FReal LinearBreakForce;
 		FReal LinearPlasticityLimit;
