Chaos : Fix dirty element very large sweeps in AABB wrongly reporting a valid number of cells
- used 64 bit integer to avoid truncation
- added unit test

diff --git a/Engine/Source/Programs/HeadlessChaos/Private/HeadlessChaosTestBroadphase.cpp b/Engine/Source/Programs/HeadlessChaos/Private/HeadlessChaosTestBroadphase.cpp
--- a/Engine/Source/Programs/HeadlessChaos/Private/HeadlessChaosTestBroadphase.cpp
+++ b/Engine/Source/Programs/HeadlessChaos/Private/HeadlessChaosTestBroadphase.cpp
@@ -864,6 +864,14 @@ namespace ChaosTest
 		FReal CurrentLength = 146.779785f;
 	}
 
+	namespace LargeSweep
+	{
+		const FVec3 QueryHalfExtents{34, 34, 90};
+		const FVec3 StartPoint{5000000000270, 11630, 187.15000295639038};
+		const FVec3 Dir{-1.0000000172032004, 0, 0.00000000000040509186487903264};
+		const FReal CurrentLength = 4999999913984;
+	};
+	
 	void AABBTreeDirtyGridFunctionsWithEdgeCase()
 	{
 		constexpr FReal DirtyElementGridCellSize = 1000.0;
@@ -872,13 +880,21 @@ namespace ChaosTest
 
 		TArray<TVec2<FReal>> VisitedCells;
 		// this should report only be as much as 4 hits as teh grid is 2D and the ray downward vertical and the halfextends are smaller than a cell size
-		DoForSweepIntersectCells(EdgeCase::QueryHalfExtents, EdgeCase::StartPoint, EdgeCase::Dir, EdgeCase::CurrentLength, DirtyElementGridCellSize, DirtyElementGridCellSizeInv,
-			[&VisitedCells](FReal X, FReal Y)
-			{
-				VisitedCells.Add({ X,Y });
-				EXPECT_TRUE(VisitedCells.Num() <= 4 );
-			});
-		EXPECT_TRUE(VisitedCells.Num() <= 4);
+		{
+			DoForSweepIntersectCells(EdgeCase::QueryHalfExtents, EdgeCase::StartPoint, EdgeCase::Dir, EdgeCase::CurrentLength, DirtyElementGridCellSize, DirtyElementGridCellSizeInv,
+				[&VisitedCells](FReal X, FReal Y)
+				{
+					VisitedCells.Add({ X,Y });
+					EXPECT_TRUE(VisitedCells.Num() <= 4 );
+				});
+			EXPECT_TRUE(VisitedCells.Num() <= 4);
+		}
+
+		// 50 trillions unit long sweep test reporting that there's not too many cells
+		{
+			const bool bTooManyCell = TooManySweepQueryCells(LargeSweep::QueryHalfExtents, LargeSweep::StartPoint, LargeSweep::Dir, LargeSweep::CurrentLength, DirtyElementGridCellSizeInv, DirtyElementMaxGridCellQueryCount);
+			EXPECT_TRUE(bTooManyCell);
+		}
 	}
 	
 	void BroadphaseCollectionTest()
diff --git a/Engine/Source/Runtime/Experimental/Chaos/Public/Chaos/AABBTreeDirtyGridUtils.h b/Engine/Source/Runtime/Experimental/Chaos/Public/Chaos/AABBTreeDirtyGridUtils.h
--- a/Engine/Source/Runtime/Experimental/Chaos/Public/Chaos/AABBTreeDirtyGridUtils.h
+++ b/Engine/Source/Runtime/Experimental/Chaos/Public/Chaos/AABBTreeDirtyGridUtils.h
@@ -244,10 +244,11 @@ namespace Chaos
 			return true;
 		}
 
-		int EstimatedNumberOfCells = ((int32)(QueryHalfExtents.X * 2 * DirtyElementGridCellSizeInv) + 2) * ((int32)(QueryHalfExtents.Y * 2 * DirtyElementGridCellSizeInv) + 2) +
-			((int32)(FMath::Max(QueryHalfExtents.X, QueryHalfExtents.Y) * 2 * DirtyElementGridCellSizeInv) + 2) * ((int32)(Length * DirtyElementGridCellSizeInv) + 2);
+		const uint64 EstimatedNumberOfCells =
+			((uint64)(QueryHalfExtents.X * 2 * DirtyElementGridCellSizeInv) + 2) * ((uint64)(QueryHalfExtents.Y * 2 * DirtyElementGridCellSizeInv) + 2) +
+			((uint64)(FMath::Max(QueryHalfExtents.X, QueryHalfExtents.Y) * 2 * DirtyElementGridCellSizeInv) + 2) * ((uint64)(Length * DirtyElementGridCellSizeInv) + 2);
 
-		return EstimatedNumberOfCells > DirtyElementMaxGridCellQueryCount;
+		return EstimatedNumberOfCells > (uint64)DirtyElementMaxGridCellQueryCount;
 	}
 
 
