Chaos - fix island manager when a constraint with no awake particle is added (it needs to set the constraint as sleeping too)

diff --git a/Engine/Source/Programs/HeadlessChaos/Private/HeadlessChaosTestGraphEvolution.cpp b/Engine/Source/Programs/HeadlessChaos/Private/HeadlessChaosTestGraphEvolution.cpp
--- a/Engine/Source/Programs/HeadlessChaos/Private/HeadlessChaosTestGraphEvolution.cpp
+++ b/Engine/Source/Programs/HeadlessChaos/Private/HeadlessChaosTestGraphEvolution.cpp
@@ -64,8 +64,9 @@ namespace ChaosTest
 			++TickCount;
 		}
 
-		void AdvanceUntilSleeping(int32 MaxIterations)
+		void AdvanceUntilSleeping()
 		{
+			const int32 MaxIterations = 50;
 			const int32 MaxTickCount = TickCount + MaxIterations;
 			bool bIsSleeping = false;
 			while (!bIsSleeping && (TickCount < MaxTickCount))
@@ -150,7 +151,7 @@ namespace ChaosTest
 		EXPECT_FALSE(Test.IslandManager->GetSolverIsland(0)->IsSleeping());
 
 		// Make all the particles sleep
-		Test.AdvanceUntilSleeping(50);
+		Test.AdvanceUntilSleeping();
 
 		// Island should be asleep
 		EXPECT_EQ(Test.IslandManager->NumIslands(), 1);
@@ -245,7 +246,7 @@ namespace ChaosTest
 		EXPECT_FALSE(Test.IslandManager->GetSolverIsland(0)->IsSleeping());
 
 		// Wait for the sleep state
-		Test.AdvanceUntilSleeping(50);
+		Test.AdvanceUntilSleeping();
 
 		// Island should be asleep but still contain all the particles and constraints
 		EXPECT_EQ(Test.IslandManager->NumIslands(), 1);
@@ -329,7 +330,7 @@ namespace ChaosTest
 		FGraphEvolutionTest Test(4);
 		Test.MakeFloor();
 
-		Test.AdvanceUntilSleeping(50);
+		Test.AdvanceUntilSleeping();
 
 		// Each particle in its own island (kinematic will be in all 3)
 		EXPECT_EQ(Test.IslandManager->NumIslands(), 3);
@@ -372,7 +373,7 @@ namespace ChaosTest
 		FGraphEvolutionTest Test(4);
 		Test.MakeFloor();
 
-		Test.AdvanceUntilSleeping(50);
+		Test.AdvanceUntilSleeping();
 
 		// Each particle in its own island (kinematic will be in all 3)
 		EXPECT_EQ(Test.IslandManager->NumIslands(), 3);
@@ -415,7 +416,7 @@ namespace ChaosTest
 		FGraphEvolutionTest Test(4);
 		Test.MakeFloor();
 
-		Test.AdvanceUntilSleeping(50);
+		Test.AdvanceUntilSleeping();
 
 		// Wake D
 		Test.Evolution.SetParticleObjectState(Test.ParticleHandles[3], EObjectStateType::Dynamic);
@@ -468,7 +469,7 @@ namespace ChaosTest
 		FGraphEvolutionTest Test(4);
 		Test.MakeFloor();
 
-		Test.AdvanceUntilSleeping(50);
+		Test.AdvanceUntilSleeping();
 
 		// Wake B
 		Test.Evolution.SetParticleObjectState(Test.ParticleHandles[1], EObjectStateType::Dynamic);
@@ -512,12 +513,13 @@ namespace ChaosTest
 	//		(d=dynamic, s=sleeping, k=kinematic)
 	//		As - Bs - Cs - Ds  =>	Ad - Bd - Cd - Dd
 	//
-	GTEST_TEST(GraphEvolutionTests, TestConstraintGraph_WakeIsland)
+	// @todo(chaos): Explicit waking of islands is not currently supported
+	GTEST_TEST(GraphEvolutionTests, DISABLED_TestConstraintGraph_WakeIsland)
 	{
 		FGraphEvolutionTest Test(4);
 		Test.MakeChain();
 
-		Test.AdvanceUntilSleeping(50);
+		Test.AdvanceUntilSleeping();
 
 		// All particles asleep
 		EXPECT_TRUE(FConstGenericParticleHandle(Test.ParticleHandles[0])->Sleeping());
@@ -530,7 +532,8 @@ namespace ChaosTest
 		if (Test.IslandManager->NumIslands() == 1)
 		{
 			EXPECT_TRUE(Test.IslandManager->GetSolverIsland(0)->IsSleeping());
-			Test.IslandManager->WakeIsland(Test.Particles, 0);
+			EXPECT_TRUE(false);	// Should be calling WakeIsland here
+			//Test.IslandManager->WakeIsland(Test.Particles, 0);
 		}
 
 		Test.Advance();
@@ -589,4 +592,149 @@ namespace ChaosTest
 		EXPECT_TRUE(FConstGenericParticleHandle(Test.ParticleHandles[3])->Sleeping());
 	}
 
+	// Add a constraint between a sleeping and a kinematic body.
+	// Nothing should wake and the constraint should be flagged as sleeping.
+	// This behaviour is required for streaming to work since scene creation may
+	// be amortized over multiple frames and constraints may be made betweens
+	// sleeping particles in a later tick.
+	//
+	//		(d=dynamic, s=sleeping, k=kinematic)
+	//		Ak  Bs    =>    Ak - Bs
+	//
+	GTEST_TEST(GraphEvolutionTests, TestConstraintGraph_SleepingKinematicConstraint)
+	{
+		FGraphEvolutionTest Test(2);
+
+		// Make A kinematic, B sleeping
+		Test.Evolution.SetParticleObjectState(Test.ParticleHandles[0], EObjectStateType::Kinematic);
+		Test.Evolution.SetParticleObjectState(Test.ParticleHandles[1], EObjectStateType::Sleeping);
+
+		Test.Advance();
+		EXPECT_TRUE(FConstGenericParticleHandle(Test.ParticleHandles[1])->Sleeping());
+
+		// Add a constraint A-B
+		Test.ConstraintHandles.Add(Test.Constraints.AddConstraint({ Test.ParticleHandles[0], Test.ParticleHandles[1] }));
+
+		Test.Advance();
+
+		// B still asleep
+		EXPECT_TRUE(FConstGenericParticleHandle(Test.ParticleHandles[1])->Sleeping());
+		EXPECT_TRUE(Test.ConstraintHandles[0]->IsSleeping());
+	}
+
+	// Add a constraint between two sleeping particles.
+	// Nothing should wake and the constraint should be flagged as sleeping.
+	// This behaviour is required for streaming to work since scene creation may
+	// be amortized over multiple frames and constraints may be made betweens
+	// sleeping particles in a later tick.
+	// In this case, A and B start in different sleeping island and get merged into
+	// a single still-sleeping island.
+	//
+	//		(d=dynamic, s=sleeping, k=kinematic)
+	//		As  Bs    =>    As - Bs
+	//
+	GTEST_TEST(GraphEvolutionTests, TestConstraintGraph_SleepingSleepingConstraint)
+	{
+		FGraphEvolutionTest Test(2);
+
+		// Make A and B sleeping
+		Test.Evolution.SetParticleObjectState(Test.ParticleHandles[0], EObjectStateType::Sleeping);
+		Test.Evolution.SetParticleObjectState(Test.ParticleHandles[1], EObjectStateType::Sleeping);
+
+		Test.Advance();
+		EXPECT_EQ(Test.IslandManager->NumIslands(), 2);
+		EXPECT_TRUE(FConstGenericParticleHandle(Test.ParticleHandles[0])->Sleeping());
+		EXPECT_TRUE(FConstGenericParticleHandle(Test.ParticleHandles[1])->Sleeping());
+
+		// Add a constraint A-B
+		Test.ConstraintHandles.Add(Test.Constraints.AddConstraint({ Test.ParticleHandles[0], Test.ParticleHandles[1] }));
+
+		Test.Advance();
+
+		// A and B still asleep
+		EXPECT_EQ(Test.IslandManager->NumIslands(), 1);
+		EXPECT_TRUE(FConstGenericParticleHandle(Test.ParticleHandles[0])->Sleeping());
+		EXPECT_TRUE(FConstGenericParticleHandle(Test.ParticleHandles[1])->Sleeping());
+		EXPECT_TRUE(Test.ConstraintHandles[0]->IsSleeping());
+	}
+
+	// Similar to TestConstraintGraph_SleepingKinematicConstraint, but we are adding a constraint 
+	// between sleeping and kinematic particles that are already in an existing sleeping island
+	// with multiple sleeping constraints.
+	//
+	//		(d=dynamic, s=sleeping, k=kinematic)
+	//		Ak - Bs - Cs  =>    Ak - Bs - Cs
+	//		                     ^--------^
+	//
+	GTEST_TEST(GraphEvolutionTests, TestConstraintGraph_SleepingKinematicConstraint_SameIsland)
+	{
+		FGraphEvolutionTest Test(3);
+
+		// Chains the particles and make the first one kinematix
+		Test.Evolution.SetParticleObjectState(Test.ParticleHandles[0], EObjectStateType::Kinematic);
+		Test.MakeChain();
+		
+		// Wait for sleep
+		Test.AdvanceUntilSleeping();
+		EXPECT_EQ(Test.IslandManager->NumIslands(), 1);
+		EXPECT_TRUE(FConstGenericParticleHandle(Test.ParticleHandles[1])->Sleeping());
+		EXPECT_TRUE(FConstGenericParticleHandle(Test.ParticleHandles[2])->Sleeping());
+		EXPECT_TRUE(Test.ConstraintHandles[0]->IsSleeping());
+		EXPECT_TRUE(Test.ConstraintHandles[1]->IsSleeping());
+
+		// Add a constraint A - C
+		Test.ConstraintHandles.Add(Test.Constraints.AddConstraint({ Test.ParticleHandles[0], Test.ParticleHandles[2] }));
+
+		Test.Advance();
+
+		// All still asleep, including the new constraint
+		EXPECT_EQ(Test.IslandManager->NumIslands(), 1);
+		EXPECT_TRUE(FConstGenericParticleHandle(Test.ParticleHandles[1])->Sleeping());
+		EXPECT_TRUE(FConstGenericParticleHandle(Test.ParticleHandles[2])->Sleeping());
+		EXPECT_TRUE(Test.ConstraintHandles[0]->IsSleeping());
+		EXPECT_TRUE(Test.ConstraintHandles[1]->IsSleeping());
+		EXPECT_TRUE(Test.ConstraintHandles[2]->IsSleeping());
+	}
+
+	// Similar to TestConstraintGraph_SleepingKinematicConstraint, but we are adding a constraint 
+	// between two sleeping particles in different island, but where each island already contains
+	// sleeping constraints.
+	//
+	//		(d=dynamic, s=sleeping, k=kinematic)
+	//		As - Bs   Cs - Ds  =>    As - Bs - Cs - Ds
+	//
+	GTEST_TEST(GraphEvolutionTests, TestConstraintGraph_SleepingSleepingConstraint_MergeIslands)
+	{
+		FGraphEvolutionTest Test(4);
+
+		// Add constraints A-B and C-D
+		Test.ConstraintHandles.Add(Test.Constraints.AddConstraint({ Test.ParticleHandles[0], Test.ParticleHandles[1] }));
+		Test.ConstraintHandles.Add(Test.Constraints.AddConstraint({ Test.ParticleHandles[2], Test.ParticleHandles[3] }));
+
+		// Wait for sleep
+		Test.AdvanceUntilSleeping();
+		EXPECT_EQ(Test.IslandManager->NumIslands(), 2);
+		EXPECT_TRUE(FConstGenericParticleHandle(Test.ParticleHandles[0])->Sleeping());
+		EXPECT_TRUE(FConstGenericParticleHandle(Test.ParticleHandles[1])->Sleeping());
+		EXPECT_TRUE(FConstGenericParticleHandle(Test.ParticleHandles[2])->Sleeping());
+		EXPECT_TRUE(FConstGenericParticleHandle(Test.ParticleHandles[3])->Sleeping());
+		EXPECT_TRUE(Test.ConstraintHandles[0]->IsSleeping());
+		EXPECT_TRUE(Test.ConstraintHandles[1]->IsSleeping());
+
+		// Add a constraint B - C
+		Test.ConstraintHandles.Add(Test.Constraints.AddConstraint({ Test.ParticleHandles[1], Test.ParticleHandles[2] }));
+
+		Test.Advance();
+
+		// All still asleep, including the new constraint
+		EXPECT_EQ(Test.IslandManager->NumIslands(), 1);
+		EXPECT_TRUE(FConstGenericParticleHandle(Test.ParticleHandles[0])->Sleeping());
+		EXPECT_TRUE(FConstGenericParticleHandle(Test.ParticleHandles[1])->Sleeping());
+		EXPECT_TRUE(FConstGenericParticleHandle(Test.ParticleHandles[2])->Sleeping());
+		EXPECT_TRUE(FConstGenericParticleHandle(Test.ParticleHandles[3])->Sleeping());
+		EXPECT_TRUE(Test.ConstraintHandles[0]->IsSleeping());
+		EXPECT_TRUE(Test.ConstraintHandles[1]->IsSleeping());
+		EXPECT_TRUE(Test.ConstraintHandles[2]->IsSleeping());
+	}
+
 }
diff --git a/Engine/Source/Runtime/Experimental/Chaos/Private/Chaos/Island/IslandGraph.cpp b/Engine/Source/Runtime/Experimental/Chaos/Private/Chaos/Island/IslandGraph.cpp
--- a/Engine/Source/Runtime/Experimental/Chaos/Private/Chaos/Island/IslandGraph.cpp
+++ b/Engine/Source/Runtime/Experimental/Chaos/Private/Chaos/Island/IslandGraph.cpp
@@ -837,6 +837,12 @@ void FIslandGraph<NodeType, EdgeType, IslandType, OwnerType>::UpdateGraph()
 {
 	SCOPE_CYCLE_COUNTER(STAT_MergeIslandGraph);
 
+	// Default to sleeping - the flag will get cleared if an island contains anything that should wake it
+	for (auto& GraphIsland : GraphIslands)
+	{
+		GraphIsland.bIsSleeping = true;
+	}
+
 	// Merge all the islands if necessary
 	MergeIslands();
 
@@ -926,13 +932,6 @@ void FIslandGraph<NodeType, EdgeType, IslandType, OwnerType>::InitIslands()
 			}
 		}
 	} 
-
-	// Reset of the sleeping flag for the graph islands
-	// See UpdateGraph() which sets it to false again if there are any awake nodes
-	for (auto& GraphIsland : GraphIslands)
-	{
-		GraphIsland.bIsSleeping = true;
-	}
 }
 
 template<typename NodeType, typename EdgeType, typename IslandType, typename OwnerType>
diff --git a/Engine/Source/Runtime/Experimental/Chaos/Private/Chaos/Island/IslandManager.cpp b/Engine/Source/Runtime/Experimental/Chaos/Private/Chaos/Island/IslandManager.cpp
--- a/Engine/Source/Runtime/Experimental/Chaos/Private/Chaos/Island/IslandManager.cpp
+++ b/Engine/Source/Runtime/Experimental/Chaos/Private/Chaos/Island/IslandManager.cpp
@@ -399,29 +399,12 @@ void FPBDIslandManager::InitializeGraph(const TParticleView<FPBDRigidParticles>&
 		}
 	}
 
-	// For now we are resetting all the constraints but we should keep the persistent constraints.
-	// NOTE: This only removes awake constraints from the graph and it leaves the sleeping ones. This is important for 
+	// For now we are removing all the constraints but really we should keep the persistent constraints and only remove the transients.
+	// NOTE: This only removes constraints from awake islands and it leaves the sleeping ones. This is important for 
 	// persistent collisions because we do not run collision detection on sleeping particles, so we need the graph
 	// to keep it's sleeping edges so we know what collisions to wake when an island is awakened. The constraint
-	// management system also used the sleeping state as a lock to prevent constraint destruction.
+	// management system also uses the sleeping state as a lock to prevent constraint destruction.
 	IslandGraph->InitIslands();
-
-	// If the user explicitly woke any islands, transfer that change to the graph
-	// @todo(chaos): this is not great. If WakeIsland is called before InitializeGraph we set the wake state 
-	// on the graph but it gets reset in InitIslands. Maybe we should reset the state at the end of the tick and
-	// then any user changes could be retained. For now we check to see if the solver sleep state was changed
-	// and if it was, use that to initialize the graph island.
-	for (int32 IslandIndex = 0; IslandIndex < IslandGraph->GraphIslands.GetMaxIndex(); ++IslandIndex)
-	{
-		if (IslandGraph->GraphIslands.IsValidIndex(IslandIndex))
-		{
-			if (IslandGraph->GraphIslands[IslandIndex].IslandItem->SleepingChanged())
-			{
-				IslandGraph->GraphIslands[IslandIndex].bIsSleeping = IslandGraph->GraphIslands[IslandIndex].IslandItem->IsSleeping();
-				IslandGraph->GraphIslands[IslandIndex].IslandItem->ResetSleepingChanged();
-			}
-		}
-	}
 }
 
 void FPBDIslandManager::EndTick()
@@ -562,7 +545,20 @@ void FPBDIslandManager::AddConstraint(const uint32 ContainerId, FConstraintHandl
 			const int32 NodeIndex0 = AddParticle(ConstrainedParticles[0], INDEX_NONE, false);
 			const int32 NodeIndex1 = AddParticle(ConstrainedParticles[1], INDEX_NONE, false);
 
-			IslandGraph->AddEdge(ConstraintHandle, ContainerId, NodeIndex0, NodeIndex1);
+			const int32 EdgeIndex = IslandGraph->AddEdge(ConstraintHandle, ContainerId, NodeIndex0, NodeIndex1);
+
+			// If we were added to a sleeping island, make sure the constraint is flagged as sleeping.
+			// Adding constraints between 2 sleeping particles does not wake them because we need to handle
+			// streaming which may amortize particle and constraint creation over multiple ticks.
+			const FGraphEdge& GraphEdge = IslandGraph->GraphEdges[EdgeIndex];
+			if (GraphEdge.IslandIndex != INDEX_NONE)
+			{
+				const bool bIsSleeping = IslandGraph->GraphIslands[GraphEdge.IslandIndex].bIsSleeping;
+				if (bIsSleeping)
+				{
+					ConstraintHandle->SetIsSleeping(bIsSleeping);
+				}
+			}
 		}
 	}
 }
@@ -895,11 +891,6 @@ void FPBDIslandManager::SleepIsland(FPBDRigidsSOAs& Particles, const int32 Islan
 	SetIslandSleeping(Particles, IslandIndex, true);
 }
 
-void FPBDIslandManager::WakeIsland(FPBDRigidsSOAs& Particles, const int32 IslandIndex)
-{
-	SetIslandSleeping(Particles, IslandIndex, false);
-}
-
 void FPBDIslandManager::SetIslandSleeping(FPBDRigidsSOAs& Particles, const int32 IslandIndex, const bool bIsSleeping)
 {
 	// NOTE: This Sleep function is called from the graph update to set the state of the 
diff --git a/Engine/Source/Runtime/Experimental/Chaos/Public/Chaos/Island/IslandManager.h b/Engine/Source/Runtime/Experimental/Chaos/Public/Chaos/Island/IslandManager.h
--- a/Engine/Source/Runtime/Experimental/Chaos/Public/Chaos/Island/IslandManager.h
+++ b/Engine/Source/Runtime/Experimental/Chaos/Public/Chaos/Island/IslandManager.h
@@ -176,13 +176,6 @@ public:
 	*/
 	void SleepIsland(FPBDRigidsSOAs& Particles, const int32 IslandIndex);
 
-	/**
-	  * Wake all particles and constraints in an Island.
-	  * @param Particles used in the islands
-	  * @param IslandIndex Island index to be woken up
-	  */
-	void WakeIsland(FPBDRigidsSOAs& Particles, const int32 IslandIndex);
-
 	/**
 	  * Put particles in inactive islands to sleep.
 	  * @param IslandIndex Island index
diff --git a/Engine/Source/Runtime/Experimental/Chaos/Public/Chaos/PBDRigidsEvolution.h b/Engine/Source/Runtime/Experimental/Chaos/Public/Chaos/PBDRigidsEvolution.h
--- a/Engine/Source/Runtime/Experimental/Chaos/Public/Chaos/PBDRigidsEvolution.h
+++ b/Engine/Source/Runtime/Experimental/Chaos/Public/Chaos/PBDRigidsEvolution.h
@@ -483,11 +483,6 @@ public:
 
 	CHAOS_API void DisableParticles(const TSet<FGeometryParticleHandle*>& InParticles);
 
-	CHAOS_API void WakeIsland(const int32 Island)
-	{
-		ConstraintGraph.WakeIsland(Particles, Island);
-	}
-
 	/** remove a constraint from the constraint graph (see AddConstraintsToConstraintGraph) */
 	CHAOS_API void RemoveConstraintFromConstraintGraph(FConstraintHandle* ConstraintHandle)
 	{
