Add unit test for Height Field

diff --git a/Engine/Source/Programs/HeadlessChaos/Private/HeadlessChaosTestHeightfield.cpp b/Engine/Source/Programs/HeadlessChaos/Private/HeadlessChaosTestHeightfield.cpp
--- a/Engine/Source/Programs/HeadlessChaos/Private/HeadlessChaosTestHeightfield.cpp
+++ b/Engine/Source/Programs/HeadlessChaos/Private/HeadlessChaosTestHeightfield.cpp
@@ -374,9 +374,68 @@ namespace ChaosTest {
 		const FVec3 Start(8224.6524537283822, 2631.7542424011549, 2265.2052028112184);
 		const FVec3 Dir(-0.92887444870972680, -0.11019885226781448, -0.35362193299208372);
 		EXPECT_TRUE(Heightfield.Raycast(Start, Dir, 2097152.0000000000, 0, TOI, Position, Normal, FaceIdx));
-		EXPECT_TRUE(Heightfield.Raycast(Start, Dir, 2097152.0000000000, UE_KINDA_SMALL_NUMBER, TOI, Position, Normal, FaceIdx));
 	}
 
+	void RaycastVariousWalkOnHeightField()
+	{
+		constexpr int32 Rows = 64;
+		constexpr int32 Columns = Rows;
+		FVec3 Scale(1.0, 1.0, 1.0);
+		TArray<FReal> Heights;
+		Heights.AddZeroed(Rows * Columns);
+
+		// Add a mountain on the diagonal
+		for (int32 Index = 0; Index < Columns; ++Index)
+		{
+			Heights[Index * Columns + Index] = 20;
+		}
+		constexpr int32 BigMountainIndex = Columns / 2;
+		Heights[BigMountainIndex * Columns + BigMountainIndex] = 40;
+
+		TArray<FReal> HeightsCopy = Heights;
+
+		FHeightField Heightfield(MoveTemp(HeightsCopy), TArray<uint8>(), Rows, Columns, Scale);
+		const auto& Bounds = Heightfield.BoundingBox();	//Current API forces us to do this to cache the bounds
+
+		FReal TOI;
+		FVec3 Position, Normal;
+		int32 FaceIdx = 0;
+
+		// Hit the diagonal, used directly the three walks (LowRes / Fast / Slow)
+		{
+			const FVec3 Start(2.0, 0.0, 10.0);
+			FVec3 Dir(0.0, 1.0, 0.0);
+			Dir.Normalize();
+			const bool bResult = Heightfield.Raycast(Start, Dir, 100.0, 0, TOI, Position, Normal, FaceIdx);
+			EXPECT_TRUE(bResult);
+		}
+		// Miss the diagonal, used the three walks but miss them  (LowRes / Fast / Slow) (worse case)
+		{
+			const FVec3 Start(0.6, 0.0, 17.0);
+			FVec3 Dir(1.0, 1.0, 0.0);
+			Dir.Normalize();
+			const bool bResult = Heightfield.Raycast(Start, Dir, 100.0, 0, TOI, Position, Normal, FaceIdx);
+			EXPECT_FALSE(bResult);
+		}
+		// Hit the diagonal in the middle used Low Res for few steps then Fast walk for few step and Slow walks (optimal use case)
+		{
+			const FVec3 Start(64.0, 0.0, 10.0);
+			FVec3 Dir(-1.0, 1.0, 0.0);
+			Dir.Normalize();
+			const bool bResult = Heightfield.Raycast(Start, Dir, 100.0, 0, TOI, Position, Normal, FaceIdx);
+			EXPECT_TRUE(bResult);
+		}
+		// Miss the diagonal in the middle used Low Res for few steps then Fast walk and Slow walks, then miss, then use LowRes until the end
+		{
+			const FVec3 Start(63.0, 0.0, 36.0);
+			FVec3 Dir(-1.0, 1.0, 0.0);
+			Dir.Normalize();
+			const bool bResult = Heightfield.Raycast(Start, Dir, 100.0, 0, TOI, Position, Normal, FaceIdx);
+			EXPECT_FALSE(bResult);
+		}
+	}
+
+
 	void EditHeights()
 	{
 		const int32 Columns = 10;
@@ -924,6 +983,7 @@ namespace ChaosTest {
 	{
 		ChaosTest::Raycast();
 		ChaosTest::RaycastOnFlatHeightField();
+		ChaosTest::RaycastVariousWalkOnHeightField();
 		ChaosTest::SweepTest();
 		ChaosTest::OverlapTest();
 		EditHeights();
