Chaos - Constraint Solver Refactor
- Constraint Solver now solves IslandGroups using sets of dependent tasks
- Gather and Scatter use all available workers regardless of number of Islands
- Significantly better multi-core performance in large-island tests
- Added IslandGroupManager
- Added ConstraintGroupSolver (replaced SolverDatas)
- Added ConstraintContainerSolver (replaced ConstraintRule)
- Removed a lot of dead code

#okforgitbug public

diff --git a/Engine/Source/Editor/PhysicsAssetEditor/Private/PhysicsAssetEditorAnimInstanceProxy.cpp b/Engine/Source/Editor/PhysicsAssetEditor/Private/PhysicsAssetEditorAnimInstanceProxy.cpp
--- a/Engine/Source/Editor/PhysicsAssetEditor/Private/PhysicsAssetEditorAnimInstanceProxy.cpp
+++ b/Engine/Source/Editor/PhysicsAssetEditor/Private/PhysicsAssetEditorAnimInstanceProxy.cpp
@@ -89,15 +89,6 @@ bool FPhysicsAssetEditorAnimInstanceProxy::Evaluate_WithRoot(FPoseContext& Outpu
 			SolverSettings.PositionIterations,
 			SolverSettings.VelocityIterations,
 			SolverSettings.ProjectionIterations);
-
-		Simulation->SetLegacySolverSettings(
-			SolverIterations.SolverIterations,
-			SolverIterations.JointIterations,
-			SolverIterations.CollisionIterations,
-			SolverIterations.SolverPushOutIterations,
-			SolverIterations.JointPushOutIterations,
-			SolverIterations.CollisionPushOutIterations
-		);
 	}
 
 	if (CurrentAsset != nullptr)
diff --git a/Engine/Source/Programs/HeadlessChaos/Private/HeadlessChaosApp.cpp b/Engine/Source/Programs/HeadlessChaos/Private/HeadlessChaosApp.cpp
--- a/Engine/Source/Programs/HeadlessChaos/Private/HeadlessChaosApp.cpp
+++ b/Engine/Source/Programs/HeadlessChaos/Private/HeadlessChaosApp.cpp
@@ -119,12 +119,6 @@ TEST(CollisionTests, Collisions) {
 	SUCCEED();
 }
 
-TEST(CollisionTests, PGS) {
-	ChaosTest::CollisionPGS();
-	ChaosTest::CollisionPGS2();
-	SUCCEED();
-}
-
 TEST(Clustering, Clustering) {
 	ChaosTest::ImplicitCluster();
 	ChaosTest::FractureCluster();
diff --git a/Engine/Source/Programs/HeadlessChaos/Private/HeadlessChaosTestCollisions.cpp b/Engine/Source/Programs/HeadlessChaos/Private/HeadlessChaosTestCollisions.cpp
--- a/Engine/Source/Programs/HeadlessChaos/Private/HeadlessChaosTestCollisions.cpp
+++ b/Engine/Source/Programs/HeadlessChaos/Private/HeadlessChaosTestCollisions.cpp
@@ -182,25 +182,17 @@ namespace ChaosTest {
 		{
 			INVARIANT_XR_START(Box);
 			const int32 NumIts = 8;
-			for (int32 It = 0; It < NumIts; ++It)
-			{
-				Collisions.Apply(Dt, It, NumIts);
-			}
+			Collisions.Apply(Dt, NumIts);
 			INVARIANT_XR_END(Box);
 		}
 
-		Collisions.SetImplicitVelocities(Dt);
-
 		{
 			RESET_PQ(Box);
 			{
 				INVARIANT_XR_START(Box);
 				INVARIANT_VW_START(Box);
 				const int32 NumIts = 1;
-				for (int32 It = 0; It < NumIts; ++It)
-				{
-					Collisions.ApplyPushOut(Dt, It, NumIts);
-				}
+				Collisions.ApplyPushOut(Dt, NumIts);
 				INVARIANT_XR_END(Box);
 				INVARIANT_VW_END(Box);
 			}
@@ -264,26 +256,17 @@ namespace ChaosTest {
 
 		{
 			INVARIANT_XR_START(Box);
-			Collisions.Apply(Dt, 0, 1);
+			Collisions.Apply(Dt, 1);
 			INVARIANT_XR_END(Box);
 		}
 
-		Collisions.SetImplicitVelocities(Dt);
-
-		// 0 restitution so expecting 0 velocity
-		//EXPECT_TRUE(Box->V().Equals(FVec3(0)));
-		//EXPECT_TRUE(Box->W().Equals(FVec3(0)));
-
 		{
 			RESET_PQ(Box);
 			{
 				INVARIANT_XR_START(Box);
 				INVARIANT_VW_START(Box);
 				const int32 NumIts = 10;
-				for (int32 It = 0; It < NumIts; ++It)
-				{
-					Collisions.ApplyPushOut(Dt, It, NumIts);
-				}
+				Collisions.ApplyPushOut(Dt, NumIts);
 				INVARIANT_XR_END(Box);
 				INVARIANT_VW_END(Box);
 			}
@@ -340,22 +323,17 @@ namespace ChaosTest {
 
 		{
 			INVARIANT_XR_START(Box);
-			Collisions.Apply(Dt, 0, 1);
+			Collisions.Apply(Dt, 1);
 			INVARIANT_XR_END(Box);
 		}
 
-		Collisions.SetImplicitVelocities(Dt);
-
 		{
 			RESET_PQ(Box);
 			{
 				INVARIANT_XR_START(Box);
 				INVARIANT_VW_START(Box);
 				const int32 NumIts = 10;
-				for (int32 It = 0; It < NumIts; ++It)
-				{
-					Collisions.ApplyPushOut(Dt, It, NumIts);
-				}
+				Collisions.ApplyPushOut(Dt, NumIts);
 				INVARIANT_XR_END(Box);
 				INVARIANT_VW_END(Box);
 			}
@@ -418,21 +396,16 @@ namespace ChaosTest {
 
 		{
 			INVARIANT_XR_START(Box);
-			Collisions.Apply(Dt, 0, 1);
+			Collisions.Apply(Dt, 1);
 			INVARIANT_XR_END(Box);
 		}
 
-		Collisions.SetImplicitVelocities(Dt);
-
 		{
 			RESET_PQ(Box);
 			{
 				INVARIANT_XR_START(Box);
 				const int32 NumIts = 10;
-				for (int32 It = 0; It < NumIts; ++It)
-				{
-					Collisions.ApplyPushOut(Dt, It, NumIts);
-				}
+				Collisions.ApplyPushOut(Dt, NumIts);
 				INVARIANT_XR_END(Box);
 			}
 		}
@@ -503,13 +476,11 @@ namespace ChaosTest {
 		EXPECT_NEAR(Constraint.CalculateWorldContactNormal().Z, FReal(1), KINDA_SMALL_NUMBER);
 		{
 			INVARIANT_XR_START(DynamicCube);
-			Collisions.Apply(Dt, 0, 1);
+			Collisions.Apply(Dt, 1);
 			INVARIANT_XR_END(DynamicCube);
 		}
 
-		Collisions.SetImplicitVelocities(Dt);
-
-		Collisions.ApplyPushOut(Dt, 0, 1);
+		Collisions.ApplyPushOut(Dt, 1);
 
 		Collisions.ScatterOutput(Dt);
 
@@ -576,23 +547,18 @@ namespace ChaosTest {
 		{
 			INVARIANT_XR_START(Box2);
 			INVARIANT_XR_START(StaticBox);
-			Collisions.Apply(Dt, 0, 1);
+			Collisions.Apply(Dt, 1);
 			INVARIANT_XR_END(Box2);
 			INVARIANT_XR_END(StaticBox);
 		}
 
-		Collisions.SetImplicitVelocities(Dt);
-
 		RESET_PQ(Box2);
 		{
 			//INVARIANT_XR_START(Box2);
 			//INVARIANT_XR_START(StaticBox);
 			//INVARIANT_VW_START(Box2);
 			const int32 NumIts = 10;
-			for (int32 It = 0; It < NumIts; ++It)
-			{
-				Collisions.ApplyPushOut(Dt, It, NumIts);
-			}
+			Collisions.ApplyPushOut(Dt, NumIts);
 			//INVARIANT_XR_END(Box2);
 			//INVARIANT_XR_END(StaticBox);
 			//INVARIANT_VW_END(Box2);
@@ -606,115 +572,5 @@ namespace ChaosTest {
 		EXPECT_FALSE(Box2->P().Equals(StartingPoint)); // moved
 		EXPECT_FALSE(Box2->Q().Equals(FQuat::Identity)); // and rotated
 	}
-
-	void CollisionPGS()
-	{
-#if CHAOS_PARTICLEHANDLE_TODO
-		TSet<int32> ActiveIndices;
-		TArray<TSet<int32>> IslandParticles;
-		TArray<int32> IslandSleepCounts;
-		TArrayCollectionArray<bool> Collided;
-		TUniquePtr<FChaosPhysicsMaterial> PhysicsMaterial = MakeUnique<FChaosPhysicsMaterial>();
-		PhysicsMaterial->Friction = (FReal)0;
-		PhysicsMaterial->Restitution = (FReal)0;
-		TArrayCollectionArray<TSerializablePtr<FChaosPhysicsMaterial>> PhysicsMaterials;
-		TArrayCollectionArray<TUniquePtr<FChaosPhysicsMaterial>> PerParticlePhysicsMaterials;
-		Chaos::FPBDRigidParticles Particles;
-		Particles.AddArray(&Collided);
-		Particles.AddArray(&PhysicsMaterials);
-
-		int32 BoxId1 = AppendParticleBox<FReal>(Particles);
-		Particles.X(BoxId1) = FVec3(0.f, 0.f, 0.5f);
-		Particles.R(BoxId1) = FRotation3(FQuat::Identity);
-		Particles.V(BoxId1) = FVec3(0.f, 0.f, -10.f);
-		PhysicsMaterials[BoxId1]  = MakeSerializable(PhysicsMaterial);
-
-		int32 BoxId2 = AppendParticleBox<FReal>(Particles);
-		Particles.X(BoxId2) = FVec3(0.f, 0.f, 0.5f);
-		Particles.R(BoxId2) = FRotation3(FQuat::Identity);
-		Particles.V(BoxId2) = FVec3(0.f, 0.f, -10.f);
-		PhysicsMaterials[BoxId2]  = MakeSerializable(PhysicsMaterial);
-
-		int32 FloorId = AppendStaticAnalyticFloor<FReal>(Particles);
-
-		ActiveIndices.Add(BoxId1);
-		ActiveIndices.Add(BoxId2);
-		ActiveIndices.Add(FloorId);
-
-		TArray<int32> Indices = ActiveIndices.Array();
-
-		FPBDCollisionConstraintPGS CollisionConstraints(Particles, Indices, Collided, PhysicsMaterials);
-		FRigidBodyContactConstraintPGS Constraint1;
-		Constraint1.ParticleIndex = BoxId1;
-		Constraint1.LevelsetIndex = BoxId2;
-		FRigidBodyContactConstraintPGS Constraint2;
-		Constraint2.ParticleIndex = BoxId1;
-		Constraint2.LevelsetIndex = FloorId;
-		FRigidBodyContactConstraintPGS Constraint3;
-		Constraint3.ParticleIndex = BoxId2;
-		Constraint3.LevelsetIndex = FloorId;
-		CollisionConstraints.Constraints.Add(Constraint1);
-		CollisionConstraints.Constraints.Add(Constraint2);
-		CollisionConstraints.Constraints.Add(Constraint3);
-		CollisionConstraints.Apply(Particles, 1.f, { 0, 1, 2 });
-		EXPECT_LT(FMath::Abs(Particles.V(BoxId1)[2]), 1e-3);
-		EXPECT_LT(FMath::Abs(Particles.V(BoxId2)[2]), 1e-3);
-#endif
-	}
-
-	void CollisionPGS2()
-	{
-#if CHAOS_PARTICLEHANDLE_TODO
-		TSet<int32> ActiveIndices;
-		TArray<TSet<int32>> IslandParticles;
-		TArray<int32> IslandSleepCounts;
-		TArrayCollectionArray<bool> Collided;
-		TUniquePtr<FChaosPhysicsMaterial> PhysicsMaterial = MakeUnique<FChaosPhysicsMaterial>();
-		PhysicsMaterial->Friction = (FReal)0;
-		PhysicsMaterial->Restitution = (FReal)0;
-		TArrayCollectionArray<TSerializablePtr<FChaosPhysicsMaterial>> PhysicsMaterials;
-		TArrayCollectionArray<TUniquePtr<FChaosPhysicsMaterial>> PerParticlePhysicsMaterials;
-		Chaos::FPBDRigidParticles Particles;
-		Particles.AddArray(&Collided);
-		Particles.AddArray(&PhysicsMaterials);
-
-		int32 BoxId1 = AppendParticleBox<FReal>(Particles);
-		Particles.X(BoxId1) = FVec3(0.f, 0.f, 0.5f);
-		Particles.R(BoxId1) = FRotation3(FQuat::Identity);
-		Particles.V(BoxId1) = FVec3(0.f, 0.f, -10.f);
-		PhysicsMaterials[BoxId1] = MakeSerializable(PhysicsMaterial);
-
-		int32 BoxId2 = AppendParticleBox<FReal>(Particles);
-		Particles.X(BoxId2) = FVec3(0.f, 0.f, 0.5f);
-		Particles.R(BoxId2) = FRotation3(FQuat::Identity);
-		Particles.V(BoxId2) = FVec3(0.f, 0.f, -10.f);
-		PhysicsMaterials[BoxId2] = MakeSerializable(PhysicsMaterial);
-
-		int32 FloorId = AppendStaticAnalyticFloor<FReal>(Particles);
-
-		ActiveIndices.Add(BoxId1);
-		ActiveIndices.Add(BoxId2);
-		ActiveIndices.Add(FloorId);
-
-		TArray<int32> Indices = ActiveIndices.Array();
-
-		FPBDCollisionConstraintPGS CollisionConstraints(Particles, Indices, Collided, PhysicsMaterials);
-		FRigidBodyContactConstraintPGS Constraint1;
-		Constraint1.ParticleIndex = BoxId1;
-		Constraint1.LevelsetIndex = BoxId2;
-		FRigidBodyContactConstraintPGS Constraint2;
-		Constraint2.ParticleIndex = BoxId1;
-		Constraint2.LevelsetIndex = FloorId;
-		FRigidBodyContactConstraintPGS Constraint3;
-		Constraint3.ParticleIndex = BoxId2;
-		Constraint3.LevelsetIndex = FloorId;
-		CollisionConstraints.Constraints.Add(Constraint2);
-		CollisionConstraints.Constraints.Add(Constraint3);
-		CollisionConstraints.Constraints.Add(Constraint1);
-		CollisionConstraints.Apply(Particles, 1.f, { 0, 1, 2 });
-		EXPECT_LT(FMath::Abs(Particles.V(BoxId1)[2] - 0.5), 1e-3);
-		EXPECT_LT(FMath::Abs(Particles.V(BoxId2)[2] + 0.5), 1e-3);
-#endif
-	}
 }
 
diff --git a/Engine/Source/Programs/HeadlessChaos/Private/HeadlessChaosTestConstraintHandle.cpp b/Engine/Source/Programs/HeadlessChaos/Private/HeadlessChaosTestConstraintHandle.cpp
--- a/Engine/Source/Programs/HeadlessChaos/Private/HeadlessChaosTestConstraintHandle.cpp
+++ b/Engine/Source/Programs/HeadlessChaos/Private/HeadlessChaosTestConstraintHandle.cpp
@@ -2,7 +2,6 @@
 
 #include "HeadlessChaos.h"
 #include "HeadlessChaosTestUtility.h"
-#include "Chaos/PBDConstraintRule.h"
 #include "Chaos/PBDJointConstraints.h"
 #include "Chaos/PBDPositionConstraints.h"
 #include "Chaos/PBDRigidParticles.h"
@@ -12,7 +11,6 @@
 #include "Chaos/Utilities.h"
 #include "Chaos/PBDRigidsEvolutionGBF.h"
 #include "Chaos/PBDPositionConstraints.h"
-#include "Chaos/PBDConstraintRule.h"
 #include "Chaos/PBDJointConstraints.h"
 #include "Chaos/PBDRigidSpringConstraints.h"
 #include "Chaos/PBDRigidDynamicSpringConstraints.h"
diff --git a/Engine/Source/Programs/HeadlessChaos/Private/HeadlessChaosTestConstraints.cpp b/Engine/Source/Programs/HeadlessChaos/Private/HeadlessChaosTestConstraints.cpp
--- a/Engine/Source/Programs/HeadlessChaos/Private/HeadlessChaosTestConstraints.cpp
+++ b/Engine/Source/Programs/HeadlessChaos/Private/HeadlessChaosTestConstraints.cpp
@@ -5,7 +5,6 @@
 #include "HeadlessChaos.h"
 #include "HeadlessChaosTestUtility.h"
 #include "Modules/ModuleManager.h"
-#include "Chaos/PBDConstraintRule.h"
 #include "Chaos/PBDJointConstraints.h"
 #include "Chaos/PBDPositionConstraints.h"
 #include "Chaos/PBDRigidParticles.h"
@@ -15,7 +14,6 @@
 #include "Chaos/Utilities.h"
 #include "Chaos/PBDRigidsEvolutionGBF.h"
 #include "Chaos/PBDPositionConstraints.h"
-#include "Chaos/PBDConstraintRule.h"
 #include "Chaos/PBDJointConstraints.h"
 #include "Chaos/PBDSuspensionConstraints.h"
 
@@ -37,10 +35,9 @@ namespace ChaosTest {
 			TArray<FPBDRigidParticleHandle*> Dynamics = Evolution.CreateDynamicParticles(1);
 			TArray<FVec3> Positions = { FVec3(0) };
 			FPBDPositionConstraints PositionConstraints(MoveTemp(Positions), MoveTemp(Dynamics), 1.f);
-			TPBDConstraintIslandRule<FPBDPositionConstraints> ConstraintRule(PositionConstraints);
 			InitEvolutionSettings(Evolution);
 
-			Evolution.AddConstraintRule(&ConstraintRule);
+			Evolution.AddConstraintContainer(PositionConstraints);
 			Evolution.AdvanceOneTimeStep(0.1);
 			Evolution.EndFrame(0.1);
 			EXPECT_LT(Evolution.GetParticleHandles().Handle(0)->X().SizeSquared(), SMALL_NUMBER);
@@ -56,8 +53,7 @@ namespace ChaosTest {
 
 			TArray<FVec3> Positions = { FVec3(1) };
 			FPBDPositionConstraints PositionConstraints(MoveTemp(Positions), MoveTemp(Dynamics), 0.5f);
-			TPBDConstraintIslandRule<FPBDPositionConstraints> ConstraintRule(PositionConstraints);
-			Evolution.AddConstraintRule(&ConstraintRule);
+			Evolution.AddConstraintContainer(PositionConstraints);
 
 			// The effect of stiffness parameter (which is set to 0.5 above) is iteration depeendent
 			Evolution.SetNumPositionIterations(1);
@@ -103,14 +99,12 @@ namespace ChaosTest {
 
 		TArray<FPBDRigidParticleHandle*> PositionParticles = { Dynamics[0] };
 		FPBDPositionConstraints PositionConstraints(MoveTemp(PositionConstraintPositions), MoveTemp(PositionParticles), 1.f);
-		TPBDConstraintIslandRule<FPBDPositionConstraints> PositionConstraintRule(PositionConstraints);
-		Evolution.AddConstraintRule(&PositionConstraintRule);
+		Evolution.AddConstraintContainer(PositionConstraints);
 
 		TVec2<TGeometryParticleHandle<FReal, 3>*> JointParticles = { Dynamics[0], Dynamics[1] };
 		FPBDJointConstraints JointConstraints;
 		JointConstraints.AddConstraint(JointParticles, FRigidTransform3(JointConstraintPosition, FRotation3::FromIdentity()));
-		TPBDConstraintIslandRule<FPBDJointConstraints> JointConstraintRule(JointConstraints);
-		Evolution.AddConstraintRule(&JointConstraintRule);
+		Evolution.AddConstraintContainer(JointConstraints);
 
 		FReal Dt = 0.1f;
 		for (int32 TimeIndex = 0; TimeIndex < 100; ++TimeIndex)
@@ -172,13 +166,7 @@ namespace ChaosTest {
 			// hard-stop will activate because 9 breaks the min suspension limit, anything greater than 8 will do this
 			SuspensionConstraints.AddConstraint(DynamicParticle, SuspensionLocalLocationA, SuspensionSettings);
 
-			//auto* Constraint = SuspensionConstraints.GetConstraintHandle(0);
-			//Constraint->
-			//
-			//FPBDSuspensionConstraintHandle 
-
-			TPBDConstraintIslandRule<FPBDSuspensionConstraints> ConstraintRule(SuspensionConstraints);
-			Evolution.AddConstraintRule(&ConstraintRule);
+			Evolution.AddConstraintContainer(SuspensionConstraints);
 
 			Evolution.AdvanceOneTimeStep(0.1);
 			Evolution.EndFrame(0.1);
@@ -228,8 +216,7 @@ namespace ChaosTest {
 			SuspensionSettings.Target = FVec3(45, 10, 9);
 			SuspensionConstraints.AddConstraint(DynamicParticle, SuspensionLocalLocationB, SuspensionSettings);
 
-			TPBDConstraintIslandRule<FPBDSuspensionConstraints> ConstraintRule(SuspensionConstraints);
-			Evolution.AddConstraintRule(&ConstraintRule);
+			Evolution.AddConstraintContainer(SuspensionConstraints);
 
 			const FVec3& Pos = Evolution.GetParticleHandles().Handle(0)->X();
 			const FRotation3& Rot = Evolution.GetParticleHandles().Handle(0)->R();
@@ -305,8 +292,7 @@ namespace ChaosTest {
 			SuspensionConstraints.AddConstraint(DynamicParticle, SusLocalOffset[SusIndex], SuspensionSettings);
 		}
 
-		TPBDConstraintIslandRule<FPBDSuspensionConstraints> ConstraintRule(SuspensionConstraints);
-		Evolution.AddConstraintRule(&ConstraintRule);
+		Evolution.AddConstraintContainer(SuspensionConstraints);
 
 		const FVec3& Pos = Evolution.GetParticleHandles().Handle(0)->X();
 		const FRotation3& Rot = Evolution.GetParticleHandles().Handle(0)->R();
@@ -342,4 +328,4 @@ namespace ChaosTest {
 		SUCCEED();
 	}
 
-}
\ No newline at end of file
+}
diff --git a/Engine/Source/Programs/HeadlessChaos/Private/HeadlessChaosTestContactModification.cpp b/Engine/Source/Programs/HeadlessChaos/Private/HeadlessChaosTestContactModification.cpp
--- a/Engine/Source/Programs/HeadlessChaos/Private/HeadlessChaosTestContactModification.cpp
+++ b/Engine/Source/Programs/HeadlessChaos/Private/HeadlessChaosTestContactModification.cpp
@@ -130,7 +130,6 @@ namespace ChaosTest
 		auto* Solver = Module->CreateSolver(nullptr, /*AsyncDt=*/-1);
 		InitSolverSettings(Solver);
 		Solver->SetThreadingMode_External(EThreadingModeTemp::SingleThread);
-		//Solver->GetEvolution()->GetCollisionConstraints().SetSolverType(EConstraintSolverType::GbfPbd);
 
 		// create a static floor and two boxes falling onto it.
 		// One box has separation modified to float 5 units above floor.
@@ -227,7 +226,6 @@ namespace ChaosTest
 		auto* Solver = Module->CreateSolver(nullptr, /*AsyncDt=*/-1);
 		InitSolverSettings(Solver);
 		Solver->SetThreadingMode_External(EThreadingModeTemp::SingleThread);
-		//Solver->GetEvolution()->GetCollisionConstraints().SetSolverType(EConstraintSolverType::GbfPbd);
 
 		// create a static floor and two boxes falling onto it.
 		// One box has normal modified to be parallel to floor, should fall through floor due to non-upward normal.
diff --git a/Engine/Source/Programs/HeadlessChaos/Private/HeadlessChaosTestGraph.cpp b/Engine/Source/Programs/HeadlessChaos/Private/HeadlessChaosTestGraph.cpp
--- a/Engine/Source/Programs/HeadlessChaos/Private/HeadlessChaosTestGraph.cpp
+++ b/Engine/Source/Programs/HeadlessChaos/Private/HeadlessChaosTestGraph.cpp
@@ -31,7 +31,7 @@ namespace ChaosTest {
 	 * by using containers with different T_TYPEIDs.
 	 */
 	template<int32 T_TYPEID>
-	class TMockGraphConstraints : public FPBDIndexedConstraintContainer
+	class TMockGraphConstraints : public TPBDIndexedConstraintContainer<TMockGraphConstraints<T_TYPEID>>
 	{
 	public:
 		using FConstraintContainerHandle = TMockGraphConstraintHandle<T_TYPEID>;
@@ -43,12 +43,39 @@ namespace ChaosTest {
 		};
 
 		TMockGraphConstraints()
-			: FPBDIndexedConstraintContainer(FConstraintContainerHandle::StaticType())
+			: TPBDIndexedConstraintContainer<TMockGraphConstraints<T_TYPEID>>(FConstraintContainerHandle::StaticType())
 		{
-			SetContainerId(T_TYPEID);
+			TPBDIndexedConstraintContainer<TMockGraphConstraints<T_TYPEID>>::SetContainerId(T_TYPEID);
 		}
 
-		int32 NumConstraints() const { return Constraints.Num(); }
+		//
+		// FConstraintContainer Implementation
+		//
+		virtual int32 GetNumConstraints() const override final { return Constraints.Num(); }
+		virtual void ResetConstraints() override final {}
+		virtual void AddConstraintsToGraph(FPBDIslandManager& IslandManager) override final {}
+		virtual void PrepareTick() override final {}
+		virtual void UnprepareTick() override final {}
+
+		//
+		// TSimpleConstraintContainerSolver API - used by RBAN solvers
+		//
+		void AddBodies(FSolverBodyContainer& SolverBodyContainer) {}
+		void GatherInput(const FReal Dt) {}
+		void ScatterOutput(const FReal Dt) {}
+		void ApplyPositionConstraints(const FReal Dt, const int32 It, const int32 NumIts) {}
+		void ApplyVelocityConstraints(const FReal Dt, const int32 It, const int32 NumIts) {}
+		void ApplyProjectionConstraints(const FReal Dt, const int32 It, const int32 NumIts) {}
+
+		//
+		// TIndexedConstraintContainerSolver API - used by World solvers
+		//
+		void AddBodies(const TArrayView<int32>& ConstraintIndices, FSolverBodyContainer& SolverBodyContainer) {}
+		void GatherInput(const TArrayView<int32>& ConstraintIndices, const FReal Dt) {}
+		void ScatterOutput(const TArrayView<int32>& ConstraintIndices, const FReal Dt) {}
+		void ApplyPositionConstraints(const TArrayView<int32>& ConstraintIndices, const FReal Dt, const int32 It, const int32 NumIts) {}
+		void ApplyVelocityConstraints(const TArrayView<int32>& ConstraintIndices, const FReal Dt, const int32 It, const int32 NumIts) {}
+		void ApplyProjectionConstraints(const TArrayView<int32>& ConstraintIndices, const FReal Dt, const int32 It, const int32 NumIts) {}
 
 		TVec2<int32> ConstraintParticleIndices(const int32 ConstraintIndex) const
 		{
@@ -228,24 +255,28 @@ namespace ChaosTest {
 
 		// Set up the particle graph
 		FPBDConstraintGraph Graph;
+		Graph.AddConstraintContainer(ConstraintsOfType0);
+		Graph.AddConstraintContainer(ConstraintsOfType1);
+
 		Graph.InitializeGraph(SOAs.GetNonDisabledDynamicView());
 
-		Graph.ReserveConstraints(ConstraintsOfType0.NumConstraints());
-		for (int32 ConstraintIndex = 0; ConstraintIndex < ConstraintsOfType0.NumConstraints(); ++ConstraintIndex)
+		Graph.ReserveConstraints(ConstraintsOfType0.GetNumConstraints());
+		for (int32 ConstraintIndex = 0; ConstraintIndex < ConstraintsOfType0.GetNumConstraints(); ++ConstraintIndex)
 		{
 			TVec2<int32> Indices = ConstraintsOfType0.ConstraintParticleIndices(ConstraintIndex);
 			Graph.AddConstraint(0, ConstraintsOfType0.Handles[ConstraintIndex], TVec2<TGeometryParticleHandle<FReal, 3>*>(AllParticles[Indices[0]], AllParticles[Indices[1]]));
 		}
 
-		Graph.ReserveConstraints(ConstraintsOfType1.NumConstraints());
-		for (int32 ConstraintIndex = 0; ConstraintIndex < ConstraintsOfType1.NumConstraints(); ++ConstraintIndex)
+		Graph.ReserveConstraints(ConstraintsOfType1.GetNumConstraints());
+		for (int32 ConstraintIndex = 0; ConstraintIndex < ConstraintsOfType1.GetNumConstraints(); ++ConstraintIndex)
 		{
 			TVec2<int32> Indices = ConstraintsOfType1.ConstraintParticleIndices(ConstraintIndex);
 			Graph.AddConstraint(1, ConstraintsOfType1.Handles[ConstraintIndex], TVec2<TGeometryParticleHandle<FReal, 3>*>(AllParticles[Indices[0]], AllParticles[Indices[1]]));
 		}
 		
 		// Generate the constraint/particle islands
-		Graph.UpdateIslands(SOAs.GetNonDisabledDynamicView(), SOAs,2);
+		Graph.UpdateGraph();
+		Graph.UpdateIslands(SOAs);
 		
 		// Islands should be end up with the following particles (note: particle 17 is infinite mass and can appear in multiple islands)
 		TArray<TSet<TGeometryParticleHandle<FReal, 3>*>> ExpectedIslandParticles = {
@@ -279,12 +310,12 @@ namespace ChaosTest {
 		for (int32 ExpectedIslandIndex = 0; ExpectedIslandIndex < ExpectedIslandParticles.Num(); ++ExpectedIslandIndex)
 		{
 			const int32 CalculatedIslandIndex = CalculatedIslandIndices[ExpectedIslandIndex];
-			const TArray<TGeometryParticleHandle<FReal, 3>*>& CalculatedIslandParticles = Graph.GetIslandParticles(CalculatedIslandIndex);
+			TArrayView<FPBDIslandParticle> CalculatedIslandParticles = Graph.GetIsland(CalculatedIslandIndex)->GetParticles();
 
 			EXPECT_EQ(CalculatedIslandParticles.Num(), ExpectedIslandParticles[ExpectedIslandIndex].Num());
-			for (TGeometryParticleHandle<FReal, 3>* CalculatedIslandParticleIndex : CalculatedIslandParticles)
+			for (FPBDIslandParticle& CalculatedIslandParticle : CalculatedIslandParticles)
 			{
-				EXPECT_TRUE(ExpectedIslandParticles[ExpectedIslandIndex].Contains(CalculatedIslandParticleIndex));
+				EXPECT_TRUE(ExpectedIslandParticles[ExpectedIslandIndex].Contains(CalculatedIslandParticle.GetParticle()));
 			}
 		}
 	}
@@ -294,14 +325,14 @@ namespace ChaosTest {
 		for (int32 ExpectedIslandIndex = 0; ExpectedIslandIndex < ExpectedIslandParticles.Num(); ++ExpectedIslandIndex)
 		{
 			const int32 CalculatedIslandIndex = CalculatedIslandIndices[ExpectedIslandIndex];
-			const TArray<TGeometryParticleHandle<FReal, 3>*>& CalculatedIslandParticles = Graph.GetIslandParticles(CalculatedIslandIndex);
+			TArrayView<FPBDIslandParticle> CalculatedIslandParticles = Graph.GetIsland(CalculatedIslandIndex)->GetParticles();
 
 			EXPECT_EQ(CalculatedIslandParticles.Num(), ExpectedIslandParticles[ExpectedIslandIndex].Num());
-			for (TGeometryParticleHandle<FReal, 3>* CalculatedIslandParticleIndex : CalculatedIslandParticles)
+			for (FPBDIslandParticle& CalculatedIslandParticle : CalculatedIslandParticles)
 			{
-				EXPECT_TRUE(ExpectedIslandParticles[ExpectedIslandIndex].Contains(CalculatedIslandParticleIndex));
+				EXPECT_TRUE(ExpectedIslandParticles[ExpectedIslandIndex].Contains(CalculatedIslandParticle.GetParticle()));
 
-				auto* DynParticle = CalculatedIslandParticleIndex->CastToRigidParticle();
+				auto* DynParticle = CalculatedIslandParticle.GetParticle()->CastToRigidParticle();
 				if (DynParticle && DynParticle->ObjectState() == EObjectStateType::Dynamic)
 				{
 					EXPECT_EQ(DynParticle->IslandIndex(), CalculatedIslandIndex);
@@ -529,8 +560,9 @@ namespace ChaosTest {
 				}
 				IterData.ExpectedIslandParticles.Add(HandlesSet);
 			}
+
 			FPBDConstraintGraph Graph;
-			FPBDConstraintColor GraphColor;
+			//FPBDConstraintColor GraphColor;
 			const int32 ContainerId = 0;
 
 			// Set up the particle graph
@@ -538,20 +570,23 @@ namespace ChaosTest {
 
 			// add constraints
 			TMockGraphConstraints<0> ConstraintsOfType0;
+			Graph.AddConstraintContainer(ConstraintsOfType0);
+
 			for (const auto& ConstrainedParticles : IterData.ConstrainedParticles)
 			{
 				ConstraintsOfType0.AddConstraint(ConstrainedParticles);
 			}
 
-			Graph.ReserveConstraints(ConstraintsOfType0.NumConstraints());
-			for (int32 ConstraintIndex = 0; ConstraintIndex < ConstraintsOfType0.NumConstraints(); ++ConstraintIndex)
+			Graph.ReserveConstraints(ConstraintsOfType0.GetNumConstraints());
+			for (int32 ConstraintIndex = 0; ConstraintIndex < ConstraintsOfType0.GetNumConstraints(); ++ConstraintIndex)
 			{
 				TVec2<int32> Indices = ConstraintsOfType0.ConstraintParticleIndices(ConstraintIndex);
 				Graph.AddConstraint(ContainerId, ConstraintsOfType0.Handles[ConstraintIndex], TVec2<TGeometryParticleHandle<FReal, 3>*>(AllParticles[Indices[0]], AllParticles[Indices[1]]));
 			}
 
 			// Generate the constraint/particle islands
-			Graph.UpdateIslands(SOAs.GetNonDisabledDynamicView(), SOAs, 1);
+			Graph.UpdateGraph();
+			Graph.UpdateIslands(SOAs);
 
 			// Check the number of islands
 			EXPECT_EQ(IterData.ExpectedIslandParticleIndices.Num(), Graph.NumIslands());
@@ -559,15 +594,15 @@ namespace ChaosTest {
 			// Compute levels and colors
 			Graph.GetIslandGraph()->InitSorting();
 
-			Graph.GetIslandGraph()->ComputeLevels(ContainerId);
-			Graph.GetIslandGraph()->ComputeColors(ContainerId, 0);
+			Graph.GetIslandGraph()->ComputeLevels();
+			Graph.GetIslandGraph()->ComputeColors(0);
 
 			// Assign color to constraints
-			GraphColor.InitializeColor(Graph);
-			for (int32 IslandIndex = 0; IslandIndex < Graph.NumIslands(); ++IslandIndex)
-			{
-				GraphColor.ComputeColor(1/*Dt*/, IslandIndex, Graph, ContainerId);
-			}
+			//GraphColor.InitializeColor(Graph);
+			//for (int32 IslandIndex = 0; IslandIndex < Graph.NumIslands(); ++IslandIndex)
+			//{
+			//	GraphColor.ComputeColor(1/*Dt*/, IslandIndex, Graph, ContainerId);
+			//}
 
 			// get the generated island indices
 			TArray<int32> CalculatedIslandIndices;
@@ -598,27 +633,27 @@ namespace ChaosTest {
 			int Index = 0;
 			for (int32 Island : CalculatedIslandIndices)
 			{
-				const TArray<FConstraintHandleHolder>& IslandConstraints = Graph.GetIslandConstraints(Island);
+				TArrayView<FPBDIslandConstraint> IslandConstraints = Graph.GetIsland(Island)->GetConstraints(ContainerId);
 				EXPECT_EQ(IslandConstraints.Num(), IterData.ExpectedIslandEdges[Index++]);
 			}
 
 			CheckIslandIntegrity(IterData.ExpectedIslandParticles, CalculatedIslandIndices, Graph);
 
 			// check level/color integrity
-			Index = 0;
-			for (int32 Island : CalculatedIslandIndices)
-			{
-				const int32 MaxLevel = GraphColor.GetIslandMaxLevel(Island);
-				const int32 MaxColor = GraphColor.GetIslandMaxColor(Island);
-
-				EXPECT_EQ(IterData.MaxLevel[Index], MaxLevel);
-				EXPECT_EQ(IterData.MaxColor[Index], MaxColor);
-				
-				EXPECT_EQ(FMath::Max(0,IterData.MaxLevel[Index]), FMath::Max(0,Graph.GetIslandGraph()->GraphIslands[Graph.GetGraphIndex(Island)].MaxLevels));
-				EXPECT_EQ(FMath::Max(0,IterData.MaxColor[Index]), FMath::Max(0,Graph.GetIslandGraph()->GraphIslands[Graph.GetGraphIndex(Island)].MaxColors));
-				
-				++Index;
-			}
+			//Index = 0;
+			//for (int32 Island : CalculatedIslandIndices)
+			//{
+			//	const int32 MaxLevel = GraphColor.GetIslandMaxLevel(Island);
+			//	const int32 MaxColor = GraphColor.GetIslandMaxColor(Island);
+
+			//	EXPECT_EQ(IterData.MaxLevel[Index], MaxLevel);
+			//	EXPECT_EQ(IterData.MaxColor[Index], MaxColor);
+			//	
+			//	EXPECT_EQ(FMath::Max(0,IterData.MaxLevel[Index]), FMath::Max(0,Graph.GetIslandGraph()->GraphIslands[Graph.GetGraphIndex(Island)].MaxLevels));
+			//	EXPECT_EQ(FMath::Max(0,IterData.MaxColor[Index]), FMath::Max(0,Graph.GetIslandGraph()->GraphIslands[Graph.GetGraphIndex(Island)].MaxColors));
+			//	
+			//	++Index;
+			//}
 
 			Graph.EndTick();
 			Graph.Reset();
@@ -634,23 +669,24 @@ namespace ChaosTest {
 		SOAs.ClearTransientDirty();
 		Graph.InitializeGraph(SOAs.GetNonDisabledDynamicView());
 
-		Graph.ReserveConstraints(Constraints.NumConstraints());
-		for(int32 ConstraintIndex = 0; ConstraintIndex < Constraints.NumConstraints(); ++ConstraintIndex)
+		Graph.ReserveConstraints(Constraints.GetNumConstraints());
+		for(int32 ConstraintIndex = 0; ConstraintIndex < Constraints.GetNumConstraints(); ++ConstraintIndex)
 		{
 			const TVec2<int32> Indices = Constraints.ConstraintParticleIndices(ConstraintIndex);
 			Graph.AddConstraint(0,Constraints.Handles[ConstraintIndex],TVec2<TGeometryParticleHandle<FReal,3>*>(Particles[Indices[0]],Particles[Indices[1]]));
 		}
 
-		Graph.UpdateIslands(SOAs.GetNonDisabledDynamicView(),SOAs, 1);
+		Graph.UpdateGraph();
+		Graph.UpdateIslands(SOAs);
 		for (int32 IslandIndex = 0; IslandIndex < Graph.NumIslands(); ++IslandIndex)
 		{
 			const bool bSleeped = Graph.SleepInactive(IslandIndex,PhysicsMaterials,PhysicalMaterials);
 
 			if(bSleeped)
 			{
-				for(TGeometryParticleHandle<FReal,3>* Particle : Graph.GetIslandParticles(IslandIndex))
+				for(FPBDIslandParticle& Particle : Graph.GetIsland(IslandIndex)->GetParticles())
 				{
-					SOAs.DeactivateParticle(Particle);
+					SOAs.DeactivateParticle(Particle.GetParticle());
 				}
 
 				Graph.SleepIsland(SOAs, IslandIndex);
@@ -728,6 +764,7 @@ namespace ChaosTest {
 			}
 
 			FPBDConstraintGraph Graph;
+			Graph.AddConstraintContainer(Constraints);
 			for (int32 LoopIndex = 0; LoopIndex < 5 + PhysicalMaterial->SleepCounterThreshold; ++LoopIndex)
 			{
 				// @todo(chaos): redo this test - sleeping now used a damped velocity rather than current velocity
@@ -815,6 +852,8 @@ namespace ChaosTest {
 			}
 
 			FPBDConstraintGraph Graph;
+			Graph.AddConstraintContainer(Constraints);
+
 			const int32 WakeUpFrame = 5 + PhysicalMaterial->SleepCounterThreshold;
 			for (int32 LoopIndex = 0; LoopIndex < WakeUpFrame + 5; ++LoopIndex)
 			{
@@ -914,6 +953,8 @@ namespace ChaosTest {
 			}
 
 			FPBDConstraintGraph Graph;
+			Graph.AddConstraintContainer(Constraints);
+
 			const int32 MergeFrame = 5 + PhysicalMaterial->SleepCounterThreshold;
 			const int32 SleepAfterMergeFrame = MergeFrame + PhysicalMaterial->SleepCounterThreshold + 1;	//first frame after merge must wake up
 			for(int32 LoopIndex = 0; LoopIndex < SleepAfterMergeFrame + 5; ++LoopIndex)
@@ -1091,74 +1132,77 @@ namespace ChaosTest {
 		
 		// Build the connectivity graph and islands
 		FPBDConstraintGraph Graph;
-		FPBDConstraintColor GraphColor;
+		Graph.AddConstraintContainer(Constraints);
+		//FPBDConstraintColor GraphColor;
 		const int32 ContainerId = 0;
 		
 		Graph.InitializeGraph(SOAs.GetNonDisabledDynamicView());
-		Graph.ReserveConstraints(Constraints.NumConstraints());
-		for (int32 ConstraintIndex = 0; ConstraintIndex < Constraints.NumConstraints(); ++ConstraintIndex)
+		Graph.ReserveConstraints(Constraints.GetNumConstraints());
+		for (int32 ConstraintIndex = 0; ConstraintIndex < Constraints.GetNumConstraints(); ++ConstraintIndex)
 		{
 			const TVec2<int32>& Indices = Constraints.ConstraintParticleIndices(ConstraintIndex);
 			Graph.AddConstraint(ContainerId, Constraints.Handles[ConstraintIndex], TVec2<TGeometryParticleHandle<FReal, 3>*>(AllParticles[Indices[0]], AllParticles[Indices[1]]) );
 		}
-		Graph.UpdateIslands(SOAs.GetNonDisabledDynamicView(), SOAs, 1);
+		Graph.UpdateGraph();
+		Graph.UpdateIslands(SOAs);
 
 		// It's a connected grid, so only one island
 		EXPECT_EQ(Graph.NumIslands(), 1);
 		
 		// Assign color to constraints
-		GraphColor.InitializeColor(Graph);
-		for (int32 IslandIndex = 0; IslandIndex < Graph.NumIslands(); ++IslandIndex)
-		{
-			GraphColor.ComputeColor(1/*Dt*/, IslandIndex, Graph, ContainerId);
-		}
+		//GraphColor.InitializeColor(Graph);
+		//for (int32 IslandIndex = 0; IslandIndex < Graph.NumIslands(); ++IslandIndex)
+		//{
+		//	GraphColor.ComputeColor(1/*Dt*/, IslandIndex, Graph, ContainerId);
+		//}
+		// 
 		// Check colors
 		// No constraints should appear twice in the level-color-constraint map
 		// No particles should be influenced by more than one constraint in any individual level+color
-		TSet<FConstraintHandle*> ConstraintUnionSet;
-		for (int32 IslandIndex = 0; IslandIndex < Graph.NumIslands(); ++IslandIndex)
-		{
-			const typename FPBDConstraintColor::FLevelToColorToConstraintListMap& LevelToColorToConstraintListMap = GraphColor.GetIslandLevelToColorToConstraintListMap(IslandIndex);
-			const int32 MaxLevel = GraphColor.GetIslandMaxLevel(IslandIndex);
-			const int32 MaxColor = GraphColor.GetIslandMaxColor(IslandIndex);
-			for (int32 Level = 0; Level <= MaxLevel; ++Level)
-			{
-				for (int Color = 0; Color <= MaxColor; ++Color)
-				{
-					if (LevelToColorToConstraintListMap[Level].Contains(Color))
-					{
-						// Build the set of constraint and particle indices for this color
-						TSet<FConstraintHandle*> ConstraintSet = TSet<FConstraintHandle*>(LevelToColorToConstraintListMap[Level][Color]);
-
-						// See if we have any constraints that were also in a prior color
-						TSet<FConstraintHandle*> ConstraintIntersectSet = ConstraintUnionSet.Intersect(ConstraintSet);
-						EXPECT_EQ(ConstraintIntersectSet.Num(), 0);
-						ConstraintUnionSet = ConstraintUnionSet.Union(ConstraintSet);
-
-						// See if any particles are being modified by more than one constraint at this level+color
-						EXPECT_TRUE(Constraints.AreConstraintsIndependent(AllParticles, LevelToColorToConstraintListMap[Level][Color]));
-					}
-				}
-			}
-		}
+		//TSet<FConstraintHandle*> ConstraintUnionSet;
+		//for (int32 IslandIndex = 0; IslandIndex < Graph.NumIslands(); ++IslandIndex)
+		//{
+		//	const typename FPBDConstraintColor::FLevelToColorToConstraintListMap& LevelToColorToConstraintListMap = GraphColor.GetIslandLevelToColorToConstraintListMap(IslandIndex);
+		//	const int32 MaxLevel = GraphColor.GetIslandMaxLevel(IslandIndex);
+		//	const int32 MaxColor = GraphColor.GetIslandMaxColor(IslandIndex);
+		//	for (int32 Level = 0; Level <= MaxLevel; ++Level)
+		//	{
+		//		for (int Color = 0; Color <= MaxColor; ++Color)
+		//		{
+		//			if (LevelToColorToConstraintListMap[Level].Contains(Color))
+		//			{
+		//				// Build the set of constraint and particle indices for this color
+		//				TSet<FConstraintHandle*> ConstraintSet = TSet<FConstraintHandle*>(LevelToColorToConstraintListMap[Level][Color]);
+
+		//				// See if we have any constraints that were also in a prior color
+		//				TSet<FConstraintHandle*> ConstraintIntersectSet = ConstraintUnionSet.Intersect(ConstraintSet);
+		//				EXPECT_EQ(ConstraintIntersectSet.Num(), 0);
+		//				ConstraintUnionSet = ConstraintUnionSet.Union(ConstraintSet);
+
+		//				// See if any particles are being modified by more than one constraint at this level+color
+		//				EXPECT_TRUE(Constraints.AreConstraintsIndependent(AllParticles, LevelToColorToConstraintListMap[Level][Color]));
+		//			}
+		//		}
+		//	}
+		//}
 
 		// Verify that we have created a reasonable number of colors. For the greedy edge coloring algorithm this is
 		//		NumColors >= MaxNodeMultiplicity
 		//		NumColors <= MaxNodeMultiplicity * 2 - 1
 		// Each node connects to 4 neighbors with 'Multiplicity' connections, but some of them will be to static particles and we ignore those.
 		// For grid dimensions less that 4 each particle on has 2 non-static connections, otherwise there are up to 4
-		const int32 MaxMultiplicity = Multiplicity * (((NumParticlesX <= 4) || (NumParticlesY <= 4)) ? 2 : 4);
-		const int32 MinNumGreedyColors = MaxMultiplicity;
-		const int32 MaxNumGreedyColors = 2 * MaxMultiplicity - 1;
-		for (int32 IslandIndex = 0; IslandIndex < Graph.NumIslands(); ++IslandIndex)
-		{
-			const typename FPBDConstraintColor::FLevelToColorToConstraintListMap& LevelToColorToConstraintListMap = GraphColor.GetIslandLevelToColorToConstraintListMap(IslandIndex);
-			const int32 MaxIslandColor = GraphColor.GetIslandMaxColor(IslandIndex);
-			EXPECT_GE(MaxIslandColor, MinNumGreedyColors - 1);
-			// We are consistently slightly worse that the greedy algorithm, but hopefully doesn't matter
-			//EXPECT_LE(MaxIslandColor, MaxNumGreedyColors - 1);
-			EXPECT_LE(MaxIslandColor, MaxNumGreedyColors);
-		}
+		//const int32 MaxMultiplicity = Multiplicity * (((NumParticlesX <= 4) || (NumParticlesY <= 4)) ? 2 : 4);
+		//const int32 MinNumGreedyColors = MaxMultiplicity;
+		//const int32 MaxNumGreedyColors = 2 * MaxMultiplicity - 1;
+		//for (int32 IslandIndex = 0; IslandIndex < Graph.NumIslands(); ++IslandIndex)
+		//{
+		//	const typename FPBDConstraintColor::FLevelToColorToConstraintListMap& LevelToColorToConstraintListMap = GraphColor.GetIslandLevelToColorToConstraintListMap(IslandIndex);
+		//	const int32 MaxIslandColor = GraphColor.GetIslandMaxColor(IslandIndex);
+		//	EXPECT_GE(MaxIslandColor, MinNumGreedyColors - 1);
+		//	// We are consistently slightly worse that the greedy algorithm, but hopefully doesn't matter
+		//	//EXPECT_LE(MaxIslandColor, MaxNumGreedyColors - 1);
+		//	EXPECT_LE(MaxIslandColor, MaxNumGreedyColors);
+		//}
 	}
 
 	/**
@@ -1199,10 +1243,10 @@ namespace ChaosTest {
 			const int32 NumIslands = ConstraintGraph.NumIslands();
 			for (int32 IslandIndex = 0; IslandIndex < NumIslands; ++IslandIndex)
 			{
-				const TArray<FGeometryParticleHandle*>& IslandParticles = ConstraintGraph.GetIslandParticles(IslandIndex);
-				for (const FGeometryParticleHandle* Handle : IslandParticles)
+				TArrayView<const FPBDIslandParticle> IslandParticles = ConstraintGraph.GetIsland(IslandIndex)->GetParticles();
+				for (const FPBDIslandParticle& IslandParticle : IslandParticles)
 				{
-					EXPECT_TRUE(IsParticleHandleValid(Handle));
+					EXPECT_TRUE(IsParticleHandleValid(IslandParticle.GetParticle()));
 				}
 			}
 
@@ -1394,4 +1438,4 @@ namespace ChaosTest {
 	{
 		GraphDestroyParticles();
 	}
-}
+}
\ No newline at end of file
diff --git a/Engine/Source/Programs/HeadlessChaos/Private/HeadlessChaosTestGraphEvolution.cpp b/Engine/Source/Programs/HeadlessChaos/Private/HeadlessChaosTestGraphEvolution.cpp
--- a/Engine/Source/Programs/HeadlessChaos/Private/HeadlessChaosTestGraphEvolution.cpp
+++ b/Engine/Source/Programs/HeadlessChaos/Private/HeadlessChaosTestGraphEvolution.cpp
@@ -5,7 +5,6 @@
 
 #include "Chaos/Island/IslandManager.h"
 #include "Chaos/ParticleHandle.h"
-#include "Chaos/PBDConstraintRule.h"
 #include "Chaos/PBDNullConstraints.h"
 #include "Chaos/PBDRigidsEvolutionGBF.h"
 #include "Chaos/Utilities.h"
@@ -24,13 +23,12 @@ namespace ChaosTest
 			, PhysicalMaterials()
 			, Evolution(Particles, PhysicalMaterials)
 			, Constraints()
-			, ConstraintsRule(Constraints)
 			, TickCount(0)
 		{
 			IslandManager = &Evolution.GetConstraintGraph();
 
 			// Bind the constraints to the evolution
-			Evolution.AddConstraintRule(&ConstraintsRule);
+			Evolution.AddConstraintContainer(Constraints);
 
 			Evolution.GetGravityForces().SetAcceleration(FVec3(0));
 
@@ -76,7 +74,7 @@ namespace ChaosTest
 				bIsSleeping = true;
 				for (int32 IslandIndex = 0; IslandIndex < IslandManager->NumIslands(); ++IslandIndex)
 				{
-					bIsSleeping = bIsSleeping && IslandManager->GetSolverIsland(IslandIndex)->IsSleeping();
+					bIsSleeping = bIsSleeping && IslandManager->GetIsland(IslandIndex)->IsSleeping();
 				}
 			}
 
@@ -89,10 +87,7 @@ namespace ChaosTest
 		THandleArray<FChaosPhysicsMaterial> PhysicalMaterials;
 		FPBDRigidsEvolutionGBF Evolution;
 
-		// Create a constraint container and bind it to the evolution
 		FPBDNullConstraints Constraints;
-		TPBDConstraintIslandRule<FPBDNullConstraints> ConstraintsRule;
-
 		TArray<FPBDRigidParticleHandle*> ParticleHandles;
 		TArray<FPBDNullConstraintHandle*> ConstraintHandles;
 		FPBDIslandManager* IslandManager;
@@ -145,17 +140,17 @@ namespace ChaosTest
 
 		// All constraints in graph in 1 island that is awake
 		EXPECT_EQ(Test.IslandManager->NumIslands(), 1);
-		EXPECT_NE(Test.ConstraintHandles[0]->ConstraintGraphIndex(), INDEX_NONE);
-		EXPECT_NE(Test.ConstraintHandles[1]->ConstraintGraphIndex(), INDEX_NONE);
-		EXPECT_NE(Test.ConstraintHandles[2]->ConstraintGraphIndex(), INDEX_NONE);
-		EXPECT_FALSE(Test.IslandManager->GetSolverIsland(0)->IsSleeping());
+		EXPECT_NE(Test.ConstraintHandles[0]->GetConstraintGraphIndex(), INDEX_NONE);
+		EXPECT_NE(Test.ConstraintHandles[1]->GetConstraintGraphIndex(), INDEX_NONE);
+		EXPECT_NE(Test.ConstraintHandles[2]->GetConstraintGraphIndex(), INDEX_NONE);
+		EXPECT_FALSE(Test.IslandManager->GetIsland(0)->IsSleeping());
 
 		// Make all the particles sleep
 		Test.AdvanceUntilSleeping();
 
 		// Island should be asleep
 		EXPECT_EQ(Test.IslandManager->NumIslands(), 1);
-		EXPECT_TRUE(Test.IslandManager->GetSolverIsland(0)->IsSleeping());
+		EXPECT_TRUE(Test.IslandManager->GetIsland(0)->IsSleeping());
 	}
 
 	// Sleep all particles in the scene and ensure that the island manager puts the island to sleep
@@ -174,10 +169,10 @@ namespace ChaosTest
 
 		// All constraints in graph in 1 island that is awake
 		EXPECT_EQ(Test.IslandManager->NumIslands(), 1);
-		EXPECT_NE(Test.ConstraintHandles[0]->ConstraintGraphIndex(), INDEX_NONE);
-		EXPECT_NE(Test.ConstraintHandles[1]->ConstraintGraphIndex(), INDEX_NONE);
-		EXPECT_NE(Test.ConstraintHandles[2]->ConstraintGraphIndex(), INDEX_NONE);
-		EXPECT_FALSE(Test.IslandManager->GetSolverIsland(0)->IsSleeping());
+		EXPECT_NE(Test.ConstraintHandles[0]->GetConstraintGraphIndex(), INDEX_NONE);
+		EXPECT_NE(Test.ConstraintHandles[1]->GetConstraintGraphIndex(), INDEX_NONE);
+		EXPECT_NE(Test.ConstraintHandles[2]->GetConstraintGraphIndex(), INDEX_NONE);
+		EXPECT_FALSE(Test.IslandManager->GetIsland(0)->IsSleeping());
 
 		// Make all the particles sleep
 		for (FPBDRigidParticleHandle* ParticleHandle : Test.ParticleHandles)
@@ -189,7 +184,7 @@ namespace ChaosTest
 
 		// Island should be asleep
 		EXPECT_EQ(Test.IslandManager->NumIslands(), 1);
-		EXPECT_TRUE(Test.IslandManager->GetSolverIsland(0)->IsSleeping());
+		EXPECT_TRUE(Test.IslandManager->GetIsland(0)->IsSleeping());
 	}
 
 	// Start with an island containing 4 particle connected in a chain, then make the middle two kinematic.
@@ -208,9 +203,9 @@ namespace ChaosTest
 
 		// All constraints in graph in 1 island
 		EXPECT_EQ(Test.IslandManager->NumIslands(), 1);
-		EXPECT_NE(Test.ConstraintHandles[0]->ConstraintGraphIndex(), INDEX_NONE);
-		EXPECT_NE(Test.ConstraintHandles[1]->ConstraintGraphIndex(), INDEX_NONE);
-		EXPECT_NE(Test.ConstraintHandles[2]->ConstraintGraphIndex(), INDEX_NONE);
+		EXPECT_NE(Test.ConstraintHandles[0]->GetConstraintGraphIndex(), INDEX_NONE);
+		EXPECT_NE(Test.ConstraintHandles[1]->GetConstraintGraphIndex(), INDEX_NONE);
+		EXPECT_NE(Test.ConstraintHandles[2]->GetConstraintGraphIndex(), INDEX_NONE);
 
 		// Convert a particle B and C to kinematic to split the islands: A-B, C-D
 		Test.Evolution.SetParticleObjectState(Test.ParticleHandles[1], EObjectStateType::Kinematic);
@@ -219,9 +214,9 @@ namespace ChaosTest
 		Test.Advance();
 
 		// Constraint[1] is not in the graph or an island
-		EXPECT_NE(Test.ConstraintHandles[0]->ConstraintGraphIndex(), INDEX_NONE);
-		EXPECT_EQ(Test.ConstraintHandles[1]->ConstraintGraphIndex(), INDEX_NONE);	// Not in graph
-		EXPECT_NE(Test.ConstraintHandles[2]->ConstraintGraphIndex(), INDEX_NONE);
+		EXPECT_NE(Test.ConstraintHandles[0]->GetConstraintGraphIndex(), INDEX_NONE);
+		EXPECT_EQ(Test.ConstraintHandles[1]->GetConstraintGraphIndex(), INDEX_NONE);	// Not in graph
+		EXPECT_NE(Test.ConstraintHandles[2]->GetConstraintGraphIndex(), INDEX_NONE);
 		EXPECT_TRUE(Test.IslandManager->DebugCheckConstraintNotInGraph(Test.ConstraintHandles[1]));
 
 		// Should now have 2 islands
@@ -240,20 +235,20 @@ namespace ChaosTest
 
 		// All constraints in graph in 1 island that is awake
 		EXPECT_EQ(Test.IslandManager->NumIslands(), 1);
-		EXPECT_NE(Test.ConstraintHandles[0]->ConstraintGraphIndex(), INDEX_NONE);
-		EXPECT_NE(Test.ConstraintHandles[1]->ConstraintGraphIndex(), INDEX_NONE);
-		EXPECT_NE(Test.ConstraintHandles[2]->ConstraintGraphIndex(), INDEX_NONE);
-		EXPECT_FALSE(Test.IslandManager->GetSolverIsland(0)->IsSleeping());
+		EXPECT_NE(Test.ConstraintHandles[0]->GetConstraintGraphIndex(), INDEX_NONE);
+		EXPECT_NE(Test.ConstraintHandles[1]->GetConstraintGraphIndex(), INDEX_NONE);
+		EXPECT_NE(Test.ConstraintHandles[2]->GetConstraintGraphIndex(), INDEX_NONE);
+		EXPECT_FALSE(Test.IslandManager->GetIsland(0)->IsSleeping());
 
 		// Wait for the sleep state
 		Test.AdvanceUntilSleeping();
 
 		// Island should be asleep but still contain all the particles and constraints
 		EXPECT_EQ(Test.IslandManager->NumIslands(), 1);
-		EXPECT_TRUE(Test.IslandManager->GetSolverIsland(0)->IsSleeping());
-		EXPECT_NE(Test.ConstraintHandles[0]->ConstraintGraphIndex(), INDEX_NONE);
-		EXPECT_NE(Test.ConstraintHandles[1]->ConstraintGraphIndex(), INDEX_NONE);
-		EXPECT_NE(Test.ConstraintHandles[2]->ConstraintGraphIndex(), INDEX_NONE);
+		EXPECT_TRUE(Test.IslandManager->GetIsland(0)->IsSleeping());
+		EXPECT_NE(Test.ConstraintHandles[0]->GetConstraintGraphIndex(), INDEX_NONE);
+		EXPECT_NE(Test.ConstraintHandles[1]->GetConstraintGraphIndex(), INDEX_NONE);
+		EXPECT_NE(Test.ConstraintHandles[2]->GetConstraintGraphIndex(), INDEX_NONE);
 
 		// Convert a particle B and C to kinematic to split the islands: A-B, C-D
 		// NOTE: the island won't split until the island is woken
@@ -266,7 +261,7 @@ namespace ChaosTest
 		EXPECT_EQ(Test.IslandManager->NumIslands(), 1);
 
 		// The kinematic-kinematic constraint will still be in the graph, and still in its island (it will not be removed until the island wakes)
-		EXPECT_NE(Test.ConstraintHandles[1]->ConstraintGraphIndex(), INDEX_NONE);
+		EXPECT_NE(Test.ConstraintHandles[1]->GetConstraintGraphIndex(), INDEX_NONE);
 		EXPECT_NE(Test.IslandManager->GetConstraintIsland(Test.ConstraintHandles[1]), INDEX_NONE);
 
 		// Wake a dynamic particle the island and we should get the split at the next update
@@ -278,7 +273,7 @@ namespace ChaosTest
 		EXPECT_EQ(Test.IslandManager->NumIslands(), 2);
 
 		// The kinematic-kinematic constraint will still be in the graph, but not in any island
-		EXPECT_NE(Test.ConstraintHandles[1]->ConstraintGraphIndex(), INDEX_NONE);
+		EXPECT_NE(Test.ConstraintHandles[1]->GetConstraintGraphIndex(), INDEX_NONE);
 		EXPECT_EQ(Test.IslandManager->GetConstraintIsland(Test.ConstraintHandles[1]), INDEX_NONE);
 	}
 
@@ -346,7 +341,7 @@ namespace ChaosTest
 
 		// All particles in one awake island
 		EXPECT_EQ(Test.IslandManager->NumIslands(), 1);
-		EXPECT_FALSE(Test.IslandManager->GetSolverIsland(0)->IsSleeping());
+		EXPECT_FALSE(Test.IslandManager->GetIsland(0)->IsSleeping());
 		EXPECT_EQ(Test.IslandManager->FindParticleIslands(Test.ParticleHandles[0]).Num(), 1);
 		EXPECT_EQ(Test.IslandManager->FindParticleIslands(Test.ParticleHandles[1]).Num(), 1);
 		EXPECT_EQ(Test.IslandManager->FindParticleIslands(Test.ParticleHandles[2]).Num(), 1);
@@ -389,7 +384,7 @@ namespace ChaosTest
 
 		// All particles in one asleep island
 		EXPECT_EQ(Test.IslandManager->NumIslands(), 1);
-		EXPECT_TRUE(Test.IslandManager->GetSolverIsland(0)->IsSleeping());
+		EXPECT_TRUE(Test.IslandManager->GetIsland(0)->IsSleeping());
 		EXPECT_EQ(Test.IslandManager->FindParticleIslands(Test.ParticleHandles[0]).Num(), 1);
 		EXPECT_EQ(Test.IslandManager->FindParticleIslands(Test.ParticleHandles[1]).Num(), 1);
 		EXPECT_EQ(Test.IslandManager->FindParticleIslands(Test.ParticleHandles[2]).Num(), 1);
@@ -441,7 +436,7 @@ namespace ChaosTest
 
 		// All particles in one awake island (D was awake so it would wake the island)
 		EXPECT_EQ(Test.IslandManager->NumIslands(), 1);
-		EXPECT_FALSE(Test.IslandManager->GetSolverIsland(0)->IsSleeping());
+		EXPECT_FALSE(Test.IslandManager->GetIsland(0)->IsSleeping());
 		EXPECT_EQ(Test.IslandManager->FindParticleIslands(Test.ParticleHandles[0]).Num(), 1);
 		EXPECT_EQ(Test.IslandManager->FindParticleIslands(Test.ParticleHandles[1]).Num(), 1);
 		EXPECT_EQ(Test.IslandManager->FindParticleIslands(Test.ParticleHandles[2]).Num(), 1);
@@ -494,7 +489,7 @@ namespace ChaosTest
 
 		// All particles in one awake island (B was awake so it would wake the island, even though A was set to sleeping)
 		EXPECT_EQ(Test.IslandManager->NumIslands(), 1);
-		EXPECT_FALSE(Test.IslandManager->GetSolverIsland(0)->IsSleeping());
+		EXPECT_FALSE(Test.IslandManager->GetIsland(0)->IsSleeping());
 		EXPECT_EQ(Test.IslandManager->FindParticleIslands(Test.ParticleHandles[0]).Num(), 1);
 		EXPECT_EQ(Test.IslandManager->FindParticleIslands(Test.ParticleHandles[1]).Num(), 1);
 		EXPECT_EQ(Test.IslandManager->FindParticleIslands(Test.ParticleHandles[2]).Num(), 1);
@@ -531,7 +526,7 @@ namespace ChaosTest
 		// Flag the island for wake up
 		if (Test.IslandManager->NumIslands() == 1)
 		{
-			EXPECT_TRUE(Test.IslandManager->GetSolverIsland(0)->IsSleeping());
+			EXPECT_TRUE(Test.IslandManager->GetIsland(0)->IsSleeping());
 			EXPECT_TRUE(false);	// Should be calling WakeIsland here
 			//Test.IslandManager->WakeIsland(Test.Particles, 0);
 		}
@@ -542,7 +537,7 @@ namespace ChaosTest
 		EXPECT_EQ(Test.IslandManager->NumIslands(), 1);
 		if (Test.IslandManager->NumIslands() == 1)
 		{
-			EXPECT_FALSE(Test.IslandManager->GetSolverIsland(0)->IsSleeping());
+			EXPECT_FALSE(Test.IslandManager->GetIsland(0)->IsSleeping());
 		}
 		EXPECT_FALSE(FConstGenericParticleHandle(Test.ParticleHandles[0])->Sleeping());
 		EXPECT_FALSE(FConstGenericParticleHandle(Test.ParticleHandles[1])->Sleeping());
@@ -574,7 +569,7 @@ namespace ChaosTest
 		// Flag the island for sleep
 		if (Test.IslandManager->NumIslands() == 1)
 		{
-			EXPECT_FALSE(Test.IslandManager->GetSolverIsland(0)->IsSleeping());
+			EXPECT_FALSE(Test.IslandManager->GetIsland(0)->IsSleeping());
 			Test.IslandManager->SleepIsland(Test.Particles, 0);
 		}
 
@@ -584,7 +579,7 @@ namespace ChaosTest
 		EXPECT_EQ(Test.IslandManager->NumIslands(), 1);
 		if (Test.IslandManager->NumIslands() == 1)
 		{
-			EXPECT_TRUE(Test.IslandManager->GetSolverIsland(0)->IsSleeping());
+			EXPECT_TRUE(Test.IslandManager->GetIsland(0)->IsSleeping());
 		}
 		EXPECT_TRUE(FConstGenericParticleHandle(Test.ParticleHandles[0])->Sleeping());
 		EXPECT_TRUE(FConstGenericParticleHandle(Test.ParticleHandles[1])->Sleeping());
diff --git a/Engine/Source/Programs/HeadlessChaos/Private/HeadlessChaosTestJoint.h b/Engine/Source/Programs/HeadlessChaos/Private/HeadlessChaosTestJoint.h
--- a/Engine/Source/Programs/HeadlessChaos/Private/HeadlessChaosTestJoint.h
+++ b/Engine/Source/Programs/HeadlessChaos/Private/HeadlessChaosTestJoint.h
@@ -3,7 +3,6 @@
 
 #include "HeadlessChaosTestConstraints.h"
 #include "Chaos/ParticleHandle.h"
-#include "Chaos/PBDConstraintRule.h"
 #include "Chaos/PBDRigidsEvolution.h"
 #include "Chaos/PBDRigidsEvolutionGBF.h"
 #include "Chaos/PBDRigidParticles.h"
diff --git a/Engine/Source/Programs/HeadlessChaos/Private/HeadlessChaosTestJointRestitution.cpp b/Engine/Source/Programs/HeadlessChaos/Private/HeadlessChaosTestJointRestitution.cpp
--- a/Engine/Source/Programs/HeadlessChaos/Private/HeadlessChaosTestJointRestitution.cpp
+++ b/Engine/Source/Programs/HeadlessChaos/Private/HeadlessChaosTestJointRestitution.cpp
@@ -20,8 +20,6 @@ namespace ChaosTest {
 		FJointChainTest<TEvolution> Test(NumSolverIterations, Gravity);
 
 		FPBDJointSolverSettings JointSolverSettings;
-		JointSolverSettings.ApplyPairIterations = 1;
-		JointSolverSettings.ApplyPushOutPairIterations = 1;
 		Test.Evolution.GetJointConstraints().SetSettings(JointSolverSettings);
 
 		Test.InitChain(2, FVec3(1,0,0));
diff --git a/Engine/Source/Programs/HeadlessChaos/Private/HeadlessChaosTestJointSolver.cpp b/Engine/Source/Programs/HeadlessChaos/Private/HeadlessChaosTestJointSolver.cpp
--- a/Engine/Source/Programs/HeadlessChaos/Private/HeadlessChaosTestJointSolver.cpp
+++ b/Engine/Source/Programs/HeadlessChaos/Private/HeadlessChaosTestJointSolver.cpp
@@ -106,10 +106,11 @@ namespace ChaosTest
 				Body1.UpdateRotationDependentState();
 			}
 
+			Solver.SetSolverBodies(&Body0, &Body1);
+
 			// Apply Constraints
 			Solver.Init(
 				Dt,
-				{&Body0, &Body1},
 				SolverSettings,
 				JointSettings,
 				Connector0,
diff --git a/Engine/Source/Programs/HeadlessChaos/Private/HeadlessChaosTestJoints.cpp b/Engine/Source/Programs/HeadlessChaos/Private/HeadlessChaosTestJoints.cpp
--- a/Engine/Source/Programs/HeadlessChaos/Private/HeadlessChaosTestJoints.cpp
+++ b/Engine/Source/Programs/HeadlessChaos/Private/HeadlessChaosTestJoints.cpp
@@ -5,7 +5,6 @@
 
 #include "Modules/ModuleManager.h"
 #include "Chaos/ParticleHandle.h"
-#include "Chaos/PBDConstraintRule.h"
 #include "Chaos/PBDRigidsEvolution.h"
 #include "Chaos/PBDRigidsEvolutionGBF.h"
 #include "Chaos/PBDRigidParticles.h"
@@ -42,9 +41,8 @@ namespace ChaosTest {
 
 		FJointConstraintsTest(const int32 NumIterations, const FReal Gravity)
 			: Base(NumIterations, Gravity)
-			, JointsRule(Joints)
 		{
-			Evolution.AddConstraintRule(&JointsRule);
+			Evolution.AddConstraintContainer(Joints);
 		}
 
 		FPBDJointConstraintHandle* AddJoint(const TVec2<FGeometryParticleHandle*>& InConstrainedParticleIndices, const FVec3& InLocation)
@@ -77,7 +75,6 @@ namespace ChaosTest {
 
 		// Solver state
 		FPBDJointConstraints Joints;
-		TPBDConstraintIslandRule<FPBDJointConstraints> JointsRule;
 	};
 
 	/**
@@ -499,8 +496,7 @@ namespace ChaosTest {
 
 		FPBDRigidSpringConstraints JointConstraints;
 		JointConstraints.AddConstraint(ConstrainedParticles, Points, 1.0f, 0.0f, (Points[0] - Points[1]).Size());
-		Chaos::TPBDConstraintIslandRule<Chaos::FPBDRigidSpringConstraints> JointRule(JointConstraints);
-		Evolution.AddConstraintRule(&JointRule);
+		Evolution.AddConstraintContainer(JointConstraints);
 
 		const FReal Dt = 0.01f;
 		for (int32 i = 0; i < 100; ++i)
@@ -541,8 +537,7 @@ namespace ChaosTest {
 			Evolution.SetPhysicsMaterial(&Box2, MakeSerializable(PhysicalMaterial));
 
 			Chaos::FPBDRigidDynamicSpringConstraints SpringConstraints(MoveTemp(Constraints));
-			Chaos::TPBDConstraintIslandRule<Chaos::FPBDRigidDynamicSpringConstraints> SpringRule(SpringConstraints);
-			Evolution.AddConstraintRule(&SpringRule);
+			Evolution.AddConstraintContainer(SpringConstraints);
 
 			const FReal Dt = 0.01f;
 			for (int32 i = 0; i < 200; ++i)
@@ -572,8 +567,7 @@ namespace ChaosTest {
 			Evolution.SetPhysicsMaterial(&Box2, MakeSerializable(PhysicalMaterial));
 
 			Chaos::FPBDRigidDynamicSpringConstraints SpringConstraints(MoveTemp(Constraints), 400);
-			Chaos::TPBDConstraintIslandRule<Chaos::FPBDRigidDynamicSpringConstraints> SpringRule(SpringConstraints);
-			Evolution.AddConstraintRule(&SpringRule);
+			Evolution.AddConstraintContainer(SpringConstraints);
 
 			const FReal Dt = 0.01f;
 			for (int32 i = 0; i < 200; ++i)
diff --git a/Engine/Source/Programs/HeadlessChaos/Private/HeadlessChaosTestMinEvolution.cpp b/Engine/Source/Programs/HeadlessChaos/Private/HeadlessChaosTestMinEvolution.cpp
--- a/Engine/Source/Programs/HeadlessChaos/Private/HeadlessChaosTestMinEvolution.cpp
+++ b/Engine/Source/Programs/HeadlessChaos/Private/HeadlessChaosTestMinEvolution.cpp
@@ -8,7 +8,6 @@
 #include "Chaos/Evolution/PBDMinEvolution.h"
 #include "Chaos/ParticleHandle.h"
 #include "Chaos/PBDCollisionConstraints.h"
-#include "Chaos/PBDConstraintRule.h"
 #include "Chaos/PBDRigidParticles.h"
 #include "Chaos/PBDRigidSpringConstraints.h"
 
@@ -44,20 +43,19 @@ namespace ChaosTest
 		FBasicBroadPhase BroadPhase(&ActivePotentiallyCollidingPairs, nullptr, nullptr);
 		FNarrowPhase NarrowPhase(0, 0, Collisions.GetConstraintAllocator());
 		FCollisionDetector CollisionDetector(BroadPhase, NarrowPhase, Collisions);
-		TSimpleConstraintRule<FCollisionConstraints> CollisionsRule(1, Collisions);
 		// End collisions stuff
 
 		// Springs
 		FPBDRigidSpringConstraints Springs;
-		TSimpleConstraintRule<FPBDRigidSpringConstraints> SpringsRule(0, Springs);
 
 		// Evolution
 		// @todo(ccaulfield): this should start with some reasonable default iterations
 		FPBDMinEvolution Evolution(ParticlesContainer, ParticlePrevXs, ParticlePrevRs, CollisionDetector, 0);
-		Evolution.SetNumIterations(1);
-		Evolution.SetNumPushOutIterations(0);
+		Evolution.SetNumPositionIterations(6);
+		Evolution.SetNumVelocityIterations(1);
+		Evolution.SetNumProjectionIterations(1);
 
-		Evolution.AddConstraintRule(&SpringsRule);
+		Evolution.AddConstraintContainer(Springs);
 		Evolution.SetGravity(FVec3(0));
 
 		FReal Dt = 1.0f / 30.0f;
@@ -108,4 +106,4 @@ namespace ChaosTest
 		EXPECT_NEAR(Distance01, Spring->GetRestLength(), 0.1f);
 	}
 
-}
\ No newline at end of file
+}
diff --git a/Engine/Source/Programs/HeadlessChaos/Private/HeadlessChaosTestPerf.cpp b/Engine/Source/Programs/HeadlessChaos/Private/HeadlessChaosTestPerf.cpp
--- a/Engine/Source/Programs/HeadlessChaos/Private/HeadlessChaosTestPerf.cpp
+++ b/Engine/Source/Programs/HeadlessChaos/Private/HeadlessChaosTestPerf.cpp
@@ -4,7 +4,6 @@
 #include "HeadlessChaosTestUtility.h"
 
 #include "Chaos/PBDRigidsEvolutionGBF.h"
-#include "Chaos/PBDConstraintRule.h"
 #include "Chaos/PBDJointConstraints.h"
 #include "Chaos/Box.h"
 #include "Chaos/Sphere.h"
@@ -46,8 +45,7 @@ namespace ChaosTest {
 
 #if USE_CONSTRAINTS
 		FPBDJointConstraints Joints;
-		TPBDConstraintIslandRule<FPBDJointConstraints> JointsRule(Joints);
-		Evolution.AddConstraintRule(&JointsRule);
+		Evolution.AddConstraintContainer(Joints);
 
 		FPBDJointSettings JointSettings;
 		const FReal TwistLimit = 0.;
diff --git a/Engine/Source/Programs/HeadlessChaos/Private/HeadlessChaosTestShapePadding.cpp b/Engine/Source/Programs/HeadlessChaos/Private/HeadlessChaosTestShapePadding.cpp
--- a/Engine/Source/Programs/HeadlessChaos/Private/HeadlessChaosTestShapePadding.cpp
+++ b/Engine/Source/Programs/HeadlessChaos/Private/HeadlessChaosTestShapePadding.cpp
@@ -87,7 +87,7 @@ namespace ChaosTest {
 
 		FCollisionConstraintAllocator CollisionAllocator;
 
-		TUniquePtr<FPBDCollisionConstraint> Constraint = FPBDCollisionConstraint::Make(
+		FPBDCollisionConstraintPtr Constraint = CollisionAllocator.CreateCollisionConstraint(
 			Box0,
 			Box0->Geometry().Get(),
 			Box0->ShapesArray()[0].Get(),
@@ -197,9 +197,9 @@ namespace ChaosTest {
 		EXPECT_NEAR(BoxBounds1.Extents().Y, Size.Y, Tolerance);
 		EXPECT_NEAR(BoxBounds1.Extents().Z, Size.Z, Tolerance);
 
-		FCollisionConstraintAllocator CollisionConstraintAllocator;
+		FCollisionConstraintAllocator CollisionAllocator;
 
-		TUniquePtr<FPBDCollisionConstraint> Constraint = FPBDCollisionConstraint::Make(
+		FPBDCollisionConstraintPtr Constraint = CollisionAllocator.CreateCollisionConstraint(
 			Box0,
 			Box0->Geometry().Get(),
 			Box0->ShapesArray()[0].Get(),
diff --git a/Engine/Source/Programs/HeadlessChaos/Private/HeadlessChaosTestSolverBody.cpp b/Engine/Source/Programs/HeadlessChaos/Private/HeadlessChaosTestSolverBody.cpp
--- a/Engine/Source/Programs/HeadlessChaos/Private/HeadlessChaosTestSolverBody.cpp
+++ b/Engine/Source/Programs/HeadlessChaos/Private/HeadlessChaosTestSolverBody.cpp
@@ -97,6 +97,6 @@ namespace ChaosTest
 
 		SolverBodyContainer.Reset(1);
 
-		SolverBodyContainer.FindOrAdd(Particles[0], 1./60.);
+		SolverBodyContainer.FindOrAdd(Particles[0]);
 	}
 }
diff --git a/Engine/Source/Programs/HeadlessChaos/Public/HeadlessChaosCollisionConstraints.h b/Engine/Source/Programs/HeadlessChaos/Public/HeadlessChaosCollisionConstraints.h
--- a/Engine/Source/Programs/HeadlessChaos/Public/HeadlessChaosCollisionConstraints.h
+++ b/Engine/Source/Programs/HeadlessChaos/Public/HeadlessChaosCollisionConstraints.h
@@ -5,6 +5,7 @@
 #include "Math/Vector.h"
 #include "Chaos/Box.h"
 #include "Chaos/Evolution/SolverBodyContainer.h"
+#include "Chaos/Evolution/ConstraintGroupSolver.h"
 #include "Chaos/Particle/ParticleUtilities.h"
 #include "Chaos/PBDRigidParticles.h"
 #include "Chaos/PBDCollisionConstraints.h"
@@ -17,7 +18,6 @@
 #include "Chaos/Utilities.h"
 #include "Chaos/Vector.h"
 #include "Chaos/PBDRigidsSOAs.h"
-#include "Chaos/Evolution/SolverDatas.h"
 
 namespace Chaos
 {	
@@ -38,28 +38,30 @@ public:
 	FPBDCollisionConstraintAccessor()
 		: EmptyParticles(UniqueIndices)
 		, SpatialAcceleration(EmptyParticles.GetNonDisabledView())
-		, CollisionConstraints(EmptyParticles, EmptyCollided, EmptyPhysicsMaterials, EmptyUniquePhysicsMaterials, nullptr, 1, 1)
+		, CollisionConstraints(EmptyParticles, EmptyCollided, EmptyPhysicsMaterials, EmptyUniquePhysicsMaterials, nullptr)
 		, NarrowPhase((FReal)1, (FReal)0, CollisionConstraints.GetConstraintAllocator())
 		, BroadPhase(EmptyParticles)
 		, CollisionDetector(BroadPhase, NarrowPhase, CollisionConstraints)
+		, ConstraintSolver()
 	{
-		CollisionConstraints.SetSolverType(EConstraintSolverType::QuasiPbd);
 		CollisionConstraints.SetContainerId(0);
-		SolverData.AddConstraintDatas<FCollisionConstraints>(CollisionConstraints.GetContainerId());
+
+		ConstraintSolver.SetConstraintSolver(CollisionConstraints.GetContainerId(), CollisionConstraints.CreateSceneSolver(0));
 	}
 
 	FPBDCollisionConstraintAccessor(const FPBDRigidsSOAs& InParticles, TArrayCollectionArray<bool>& Collided, const TArrayCollectionArray<TSerializablePtr<FChaosPhysicsMaterial>>& PerParticleMaterials, const TArrayCollectionArray<TUniquePtr<FChaosPhysicsMaterial>>& PerParticleUniqueMaterials,
 		const int32 PushOutIterations, const int32 PushOutPairIterations) 
 		: EmptyParticles(UniqueIndices)
 		, SpatialAcceleration(InParticles.GetNonDisabledView())
-		, CollisionConstraints(InParticles, Collided, PerParticleMaterials, PerParticleUniqueMaterials, nullptr, 1, 1)
+		, CollisionConstraints(InParticles, Collided, PerParticleMaterials, PerParticleUniqueMaterials, nullptr)
 		, NarrowPhase((FReal)1, (FReal)0, CollisionConstraints.GetConstraintAllocator())
 		, BroadPhase(InParticles)
 		, CollisionDetector(BroadPhase, NarrowPhase, CollisionConstraints)
+		, ConstraintSolver()
 	{
-		CollisionConstraints.SetSolverType(EConstraintSolverType::QuasiPbd);
 		CollisionConstraints.SetContainerId(0);
-		SolverData.AddConstraintDatas<FCollisionConstraints>(CollisionConstraints.GetContainerId());
+		
+		ConstraintSolver.SetConstraintSolver(CollisionConstraints.GetContainerId(), CollisionConstraints.CreateSceneSolver(0));
 	}
 
 	virtual ~FPBDCollisionConstraintAccessor() {}
@@ -130,38 +132,28 @@ public:
 		return CollisionConstraints.GetConstraintHandles()[ConstraintIndex];
 	}
 
-	void Apply(const FReal Dt, const int32 It, const int32 NumIts)
+	void Apply(const FReal Dt, const int32 NumIts)
 	{
-		CollisionConstraints.ApplyPhase1(Dt, It, NumIts, SolverData);
+		ConstraintSolver.ApplyPositionConstraints(Dt, NumIts);
 	}
 
-	bool ApplyPushOut(const FReal Dt, int32 Iteration, int32 NumIterations)
+	void ApplyPushOut(const FReal Dt, int32 NumIts)
 	{
-		return CollisionConstraints.ApplyPhase2(Dt, Iteration, NumIterations, SolverData);
+		ConstraintSolver.ApplyVelocityConstraints(Dt, NumIts);
 	}
 
 	void GatherInput(FReal Dt)
 	{
-		SolverData.GetBodyContainer().Reset(1000);
-		CollisionConstraints.SetNumIslandConstraints(CollisionConstraints.NumConstraints(), SolverData);
-
-		for (auto Handle : CollisionConstraints.GetConstraintHandles())
-		{
-			Handle->PreGatherInput(Dt, SolverData);
-			Handle->GatherInput(Dt, INDEX_NONE, INDEX_NONE, SolverData);
-		}
+		ConstraintSolver.Reset();
+		ConstraintSolver.AddConstraintsAndBodies();
+		ConstraintSolver.GatherBodies(Dt);
+		ConstraintSolver.GatherConstraints(Dt);
 	}
 
 	void ScatterOutput(FReal Dt)
 	{
-		CollisionConstraints.ScatterOutput(Dt, SolverData);
-		SolverData.GetBodyContainer().ScatterOutput();
-		SolverData.GetBodyContainer().Reset(0);;
-	}
-
-	void SetImplicitVelocities(FReal Dt)
-	{
-		SolverData.GetBodyContainer().SetImplicitVelocities(Dt);
+		ConstraintSolver.ScatterConstraints(Dt);
+		ConstraintSolver.ScatterBodies(Dt);
 	}
 
 	FPBDCollisionConstraint EmptyConstraint;
@@ -176,6 +168,6 @@ public:
 	FNarrowPhase NarrowPhase;
 	FSpatialAccelerationBroadPhase BroadPhase;
 	FCollisionDetector CollisionDetector;
-	FPBDIslandSolverData SolverData;
+	FPBDSceneConstraintGroupSolver ConstraintSolver;
 };
 }
diff --git a/Engine/Source/Programs/HeadlessChaos/Public/HeadlessChaosTestCollisions.h b/Engine/Source/Programs/HeadlessChaos/Public/HeadlessChaosTestCollisions.h
--- a/Engine/Source/Programs/HeadlessChaos/Public/HeadlessChaosTestCollisions.h
+++ b/Engine/Source/Programs/HeadlessChaos/Public/HeadlessChaosTestCollisions.h
@@ -27,11 +27,4 @@ namespace ChaosTest {
 	void CollisionCubeCubeRestitution();
 
 	void CollisionBoxToStaticBox();
-
-	void CollisionPGS();
-
-	void CollisionPGS2();
-
-	void CollisionBroadphaseVelocityExpansion();
-
 }
\ No newline at end of file
diff --git a/Engine/Source/Programs/HeadlessChaos/Public/HeadlessChaosTestConstraints.h b/Engine/Source/Programs/HeadlessChaos/Public/HeadlessChaosTestConstraints.h
--- a/Engine/Source/Programs/HeadlessChaos/Public/HeadlessChaosTestConstraints.h
+++ b/Engine/Source/Programs/HeadlessChaos/Public/HeadlessChaosTestConstraints.h
@@ -33,6 +33,8 @@ namespace ChaosTest
 			PhysicalMaterial->DisabledAngularThreshold = 0;
 
 			Evolution.SetNumPositionIterations(NumIterations);
+			Evolution.SetNumVelocityIterations(1);
+			Evolution.SetNumProjectionIterations(1);
 			Evolution.GetGravityForces().SetAcceleration(Gravity * FVec3(0, 0, -1));
 		}
 
diff --git a/Engine/Source/Programs/HeadlessChaos/Public/HeadlessChaosTestUtility.h b/Engine/Source/Programs/HeadlessChaos/Public/HeadlessChaosTestUtility.h
--- a/Engine/Source/Programs/HeadlessChaos/Public/HeadlessChaosTestUtility.h
+++ b/Engine/Source/Programs/HeadlessChaos/Public/HeadlessChaosTestUtility.h
@@ -152,7 +152,6 @@ namespace ChaosTest {
 		Evolution.GetCollisionDetector().GetNarrowPhase().GetContext().GetSettings().bDeferNarrowPhase = false;
 		Evolution.GetCollisionDetector().GetNarrowPhase().GetContext().GetSettings().bAllowManifolds = true;
 		Evolution.GetCollisionConstraints().SetDetectorSettings(Evolution.GetCollisionDetector().GetNarrowPhase().GetContext().GetSettings());
-		Evolution.GetCollisionConstraints().SetSolverType(EConstraintSolverType::QuasiPbd);
 	}
 
 	template<typename T_SOLVER>
diff --git a/Engine/Source/Runtime/AnimGraphRuntime/Private/BoneControllers/AnimNode_RigidBody.cpp b/Engine/Source/Runtime/AnimGraphRuntime/Private/BoneControllers/AnimNode_RigidBody.cpp
--- a/Engine/Source/Runtime/AnimGraphRuntime/Private/BoneControllers/AnimNode_RigidBody.cpp
+++ b/Engine/Source/Runtime/AnimGraphRuntime/Private/BoneControllers/AnimNode_RigidBody.cpp
@@ -849,14 +849,6 @@ void FAnimNode_RigidBody::EvaluateSkeletalControl_AnyThread(FComponentSpacePoseC
 				SolverSettings.VelocityIterations,
 				SolverSettings.ProjectionIterations);
 
-			PhysicsSimulation->SetLegacySolverSettings(
-				SolverIterations.SolverIterations,
-				SolverIterations.JointIterations,
-				SolverIterations.CollisionIterations,
-				SolverIterations.SolverPushOutIterations,
-				SolverIterations.JointPushOutIterations,
-				SolverIterations.CollisionPushOutIterations);
-
 			if (!bUseDeferredSimulationTask)
 			{
 				RunPhysicsSimulation(DeltaSeconds, SimSpaceGravity);
@@ -1281,13 +1273,6 @@ void FAnimNode_RigidBody::InitPhysics(const UAnimInstance* InAnimInstance)
 			SolverSettings.ProjectionIterations);
 
 		SolverIterations = UsePhysicsAsset->SolverIterations;
-		PhysicsSimulation->SetLegacySolverSettings(
-			SolverIterations.SolverIterations,
-			SolverIterations.JointIterations,
-			SolverIterations.CollisionIterations,
-			SolverIterations.SolverPushOutIterations,
-			SolverIterations.JointPushOutIterations,
-			SolverIterations.CollisionPushOutIterations);
 	}
 }
 
diff --git a/Engine/Source/Runtime/Engine/Private/PhysicsEngine/ImmediatePhysics/ImmediatePhysicsChaos/ImmediatePhysicsSimulation_Chaos.cpp b/Engine/Source/Runtime/Engine/Private/PhysicsEngine/ImmediatePhysics/ImmediatePhysicsChaos/ImmediatePhysicsSimulation_Chaos.cpp
--- a/Engine/Source/Runtime/Engine/Private/PhysicsEngine/ImmediatePhysics/ImmediatePhysicsChaos/ImmediatePhysicsSimulation_Chaos.cpp
+++ b/Engine/Source/Runtime/Engine/Private/PhysicsEngine/ImmediatePhysics/ImmediatePhysicsChaos/ImmediatePhysicsSimulation_Chaos.cpp
@@ -19,7 +19,6 @@
 #include "Chaos/MassConditioning.h"
 #include "Chaos/ParticleHandle.h"
 #include "Chaos/PBDCollisionConstraints.h"
-#include "Chaos/PBDConstraintRule.h"
 #include "Chaos/PBDJointConstraints.h"
 #include "Chaos/PBDRigidParticles.h"
 #include "Chaos/PBDRigidsEvolutionGBF.h"
@@ -52,6 +51,7 @@ int32 ChaosImmediate_Evolution_Iterations = -1;				// Legacy
 int32 ChaosImmediate_Evolution_PushOutIterations = -1;		// Legacy
 Chaos::FRealSingle ChaosImmediate_Evolution_BoundsExtension = 0.0f;
 int32 ChaosImmediate_DisableInactiveByIndex = 1;
+int32 ChaosImmediate_Evolution_NumCollisionsPerBlock = 50;
 FAutoConsoleVariableRef CVarChaosImmPhysStepTime(TEXT("p.Chaos.ImmPhys.StepTime"), ChaosImmediate_Evolution_StepTime, TEXT("Override step time (if not zero)"));
 FAutoConsoleVariableRef CVarChaosImmPhysNumSteps(TEXT("p.Chaos.ImmPhys.NumSteps"), ChaosImmediate_Evolution_NumSteps, TEXT("Override num steps (if not zero)"));
 FAutoConsoleVariableRef CVarChaosImmPhysInitialStepTime(TEXT("p.Chaos.ImmPhys.InitialStepTime"), ChaosImmediate_Evolution_InitialStepTime, TEXT("Initial step time (then calculated from rolling average)"));
@@ -63,6 +63,7 @@ FAutoConsoleVariableRef CVarChaosImmPhysIterations(TEXT("p.Chaos.ImmPhys.Iterati
 FAutoConsoleVariableRef CVarChaosImmPhysPushOutIterations(TEXT("p.Chaos.ImmPhys.PushOutIterations"), ChaosImmediate_Evolution_PushOutIterations, TEXT("[Legacy Solver] Override number of solver push-out loops (if >= 0)"));
 FAutoConsoleVariableRef CVarChaosImmPhysBoundsExtension(TEXT("p.Chaos.ImmPhys.BoundsExtension"), ChaosImmediate_Evolution_BoundsExtension, TEXT("Bounds are grown by this fraction of their size (should be >= 0.0)"));
 FAutoConsoleVariableRef CVarChaosImmPhysDisableInactiveByIndex(TEXT("p.Chaos.ImmPhys.DisableInactiveByIndex"), ChaosImmediate_DisableInactiveByIndex, TEXT("Disable bodies that are no longer active based on the index, rather than just count."));
+FAutoConsoleVariableRef CVarChaosImmPhysNumCollisionsPerBlock(TEXT("p.Chaos.ImmPhys.NumCollisionsPerBlock"), ChaosImmediate_Evolution_NumCollisionsPerBlock, TEXT("The number of collision in a block in the collision pool. Higher values give better cache efficieny but waste memory if you do not need that many"));
 
 Chaos::FRealSingle ChaosImmediate_Evolution_SimSpaceCentrifugalAlpha = 1.0f;
 Chaos::FRealSingle ChaosImmediate_Evolution_SimSpaceCoriolisAlpha = 0.5f;
@@ -168,12 +169,6 @@ FAutoConsoleVariableRef  CVarChaosImmPhysParticleInertiaConditioningRotationRati
 Chaos::FRealSingle ChaosImmediate_MaxInvInertiaComponentRatio = 0;
 FAutoConsoleVariableRef  CVarChaosImmPhysInertiaConditioningMaxInvInertiaComponentRatio(TEXT("p.Chaos.ImmPhys.InertiaConditioning.MaxInvInertiaComponentRatio"), ChaosImmediate_MaxInvInertiaComponentRatio, TEXT("An input to inertia conditioning system. The largest inertia component must be at least least multiple of the smallest component"));
 
-//
-// Select the solver technique to use until we settle on the final one...
-//
-int32 ChaosImmediate_SolverType = (int32)Chaos::EConstraintSolverType::QuasiPbd;
-FAutoConsoleVariableRef CVarChaosImmPhysSolverType(TEXT("p.Chaos.ImmPhys.SolverType"), ChaosImmediate_SolverType, TEXT("0 = None; 1 = GbfPbd; 2 = Pbd; 3 = QuasiPbd"));
-
 // Whether to use the linear joint solver which is significantly faster than the non-linear one but less accurate. Only applies to the QuasiPBD Solver
 int32 ChaosImmediate_Joint_UseLinearSolver = -1;
 FAutoConsoleVariableRef CVarChaosImmPhysJointUseCachedSolver(TEXT("p.Chaos.ImmPhys.Joint.UseLinearSolver"), ChaosImmediate_Joint_UseLinearSolver, TEXT("Force use of linear or non-linear joint solver. (-1 to use PhysicsAsset setting)"));
@@ -304,12 +299,10 @@ namespace ImmediatePhysics_Chaos
 		FImplementation()
 			: Particles(UniqueIndices)
 			, Joints()
-			, Collisions(Particles, CollidedParticles, ParticleMaterials, PerParticleMaterials, nullptr)
+			, Collisions(Particles, CollidedParticles, ParticleMaterials, PerParticleMaterials, nullptr, ChaosImmediate_Evolution_NumCollisionsPerBlock, 2000)
 			, BroadPhase(&ActivePotentiallyCollidingPairs, nullptr, nullptr)
 			, NarrowPhase(FReal(0), FReal(0), Collisions.GetConstraintAllocator())
 			, CollisionDetector(BroadPhase, NarrowPhase, Collisions)
-			, JointsRule(0, Joints)
-			, CollisionsRule(1, Collisions)
 			, Evolution(Particles, ParticlePrevXs, ParticlePrevRs, CollisionDetector, FReal(0))
 			, NumActiveDynamicActorHandles(0)
 			, SimulationSpace()
@@ -324,8 +317,8 @@ namespace ImmediatePhysics_Chaos
 			Particles.GetParticleHandles().AddArray(&ParticlePrevXs);
 			Particles.GetParticleHandles().AddArray(&ParticlePrevRs);
 
-			Evolution.AddConstraintRule(&CollisionsRule);
-			Evolution.AddConstraintRule(&JointsRule);
+			Evolution.AddConstraintContainer(Collisions, ChaosImmediate_Collision_Priority);
+			Evolution.AddConstraintContainer(Joints, 0);
 		}
 
 		~FImplementation()
@@ -348,8 +341,6 @@ namespace ImmediatePhysics_Chaos
 		Chaos::FBasicBroadPhase BroadPhase;
 		Chaos::FNarrowPhase NarrowPhase;
 		Chaos::FBasicCollisionDetector CollisionDetector;
-		Chaos::TSimpleConstraintRule<Chaos::FPBDJointConstraints> JointsRule;
-		Chaos::TSimpleConstraintRule<Chaos::FPBDCollisionConstraints> CollisionsRule;
 		Chaos::FPBDMinEvolution Evolution;
 
 		/** Mapping from entity index to handle */
@@ -388,8 +379,6 @@ namespace ImmediatePhysics_Chaos
 
 		// RBAN collision customization
 		Implementation->Collisions.DisableHandles();
-		Implementation->Collisions.SetSolverType(EConstraintSolverType::StandardPbd);
-		Implementation->Joints.SetSolverType(EConstraintSolverType::StandardPbd);
 
 		Implementation->NarrowPhase.GetContext().GetSettings().bFilteringEnabled = false;
 		Implementation->NarrowPhase.GetContext().GetSettings().bAllowManifoldReuse = false;
@@ -790,36 +779,6 @@ namespace ImmediatePhysics_Chaos
 		}
 	}
 
-	void FSimulation::SetLegacySolverSettings(const int32 SolverIts, const int32 JointIts, const int32 CollisionIts, const int32 SolverPushOutIts, const int32 JointPushOutIts, const int32 CollisionPushOutIts)
-	{
-		if (SolverIts >= 0)
-		{
-			Implementation->Evolution.SetNumIterations(SolverIts);
-		}
-		if (SolverPushOutIts >= 0)
-		{
-			Implementation->Evolution.SetNumPushOutIterations(SolverPushOutIts);
-		}
-
-		if (JointIts >= 0)
-		{
-			Implementation->Joints.SetNumPairIterations(JointIts);
-		}
-		if (JointPushOutIts >= 0)
-		{
-			Implementation->Joints.SetNumPushOutPairIterations(JointPushOutIts);
-		}
-
-		if (CollisionIts >= 0)
-		{
-			Implementation->Collisions.SetPairIterations(CollisionIts);
-		}
-		if (CollisionPushOutIts >= 0)
-		{
-			Implementation->Collisions.SetPushOutPairIterations(CollisionPushOutIts);
-		}
-	}
-
 	void FSimulation::DebugDraw()
 	{
 		DebugDrawStaticParticles();
@@ -854,31 +813,14 @@ namespace ImmediatePhysics_Chaos
 
 		// TEMP: overrides
 		{
-			const EConstraintSolverType SolverType = (EConstraintSolverType)FMath::Clamp(ChaosImmediate_SolverType, int32(EConstraintSolverType::None), int32(EConstraintSolverType::QuasiPbd));
-			Implementation->Evolution.SetSolverType(SolverType);
-			Implementation->Collisions.SetSolverType(SolverType);
-			Implementation->Joints.SetSolverType(SolverType);
-			if (SolverType == EConstraintSolverType::QuasiPbd)
-			{
-				SetSolverSettings(
-					ChaosImmediate_Evolution_FixedStepTime,
-					ChaosImmediate_Collision_CullDistance,
-					ChaosImmediate_Collision_MaxDepenetrationVelocity,
-					ChaosImmediate_Joint_UseLinearSolver,
-					ChaosImmediate_Evolution_PositionIterations,
-					ChaosImmediate_Evolution_VelocityIterations,
-					ChaosImmediate_Evolution_ProjectionIterations);
-			}
-			else
-			{
-				SetLegacySolverSettings(
-					ChaosImmediate_Evolution_Iterations,
-					ChaosImmediate_Joint_PairIterations,
-					ChaosImmediate_Collision_PairIterations,
-					ChaosImmediate_Evolution_PushOutIterations,
-					ChaosImmediate_Joint_PushOutPairIterations,
-					0);
-			}
+			SetSolverSettings(
+				ChaosImmediate_Evolution_FixedStepTime,
+				ChaosImmediate_Collision_CullDistance,
+				ChaosImmediate_Collision_MaxDepenetrationVelocity,
+				ChaosImmediate_Joint_UseLinearSolver,
+				ChaosImmediate_Evolution_PositionIterations,
+				ChaosImmediate_Evolution_VelocityIterations,
+				ChaosImmediate_Evolution_ProjectionIterations);
 
 			FPBDJointSolverSettings JointsSettings = Implementation->Joints.GetSettings();
 			JointsSettings.SwingTwistAngleTolerance = ChaosImmediate_Joint_SwingTwistAngleTolerance;
@@ -905,31 +847,19 @@ namespace ImmediatePhysics_Chaos
 			JointsSettings.LinearDriveDampingOverride = ChaosImmediate_Joint_LinearDriveDamping;
 			JointsSettings.AngularDriveStiffnessOverride = ChaosImmediate_Joint_AngularDriveStiffness;
 			JointsSettings.AngularDriveDampingOverride = ChaosImmediate_Joint_AngularDriveDamping;
-			if (SolverType == EConstraintSolverType::QuasiPbd)
-			{
-				// NOTE: bUseLinearSolver is set via SetSolverSettings above for QPBD
-				JointsSettings.NumShockPropagationIterations = ChaosImmediate_Joint_NumShockPropagationIterations;
-			}
-			else
-			{
-				JointsSettings.NumShockPropagationIterations = 0;
-				JointsSettings.bUseLinearSolver = false;
-			}
+			JointsSettings.NumShockPropagationIterations = ChaosImmediate_Joint_NumShockPropagationIterations;
 			Implementation->Joints.SetSettings(JointsSettings);
 
-			if (SolverType == EConstraintSolverType::QuasiPbd)
-			{
-				Implementation->Collisions.SetPositionFrictionIterations(ChaosImmediate_Collision_NumPositionFrictionIterations);
-				Implementation->Collisions.SetVelocityFrictionIterations(ChaosImmediate_Collision_NumVelocityFrictionIterations);
-				Implementation->Collisions.SetPositionShockPropagationIterations(ChaosImmediate_Collision_NumPositionShockPropagationIterations);
-				Implementation->Collisions.SetVelocityShockPropagationIterations(ChaosImmediate_Collision_NumVelocityShockPropagationIterations);
-			}
+			Implementation->Collisions.SetPositionFrictionIterations(ChaosImmediate_Collision_NumPositionFrictionIterations);
+			Implementation->Collisions.SetVelocityFrictionIterations(ChaosImmediate_Collision_NumVelocityFrictionIterations);
+			Implementation->Collisions.SetPositionShockPropagationIterations(ChaosImmediate_Collision_NumPositionShockPropagationIterations);
+			Implementation->Collisions.SetVelocityShockPropagationIterations(ChaosImmediate_Collision_NumVelocityShockPropagationIterations);
 
 			Implementation->Collisions.SetRestitutionEnabled(ChaosImmediate_Collision_RestitutionEnabled != 0);
 			Implementation->Collisions.SetRestitutionThreshold(ChaosImmediate_Collision_RestitutionThresholdMultiplier * InGravity.Size());
 			Implementation->Collisions.SetCollisionsEnabled(ChaosImmediate_Collision_Enabled != 0);
 
-			Implementation->CollisionsRule.SetPriority(ChaosImmediate_Collision_Priority);
+			Implementation->Evolution.SetConstraintContainerPriority(Implementation->Collisions.GetContainerId(), ChaosImmediate_Collision_Priority);
 
 			Implementation->NarrowPhase.GetContext().GetSettings().bAllowManifoldReuse = false;
 			Implementation->NarrowPhase.GetContext().GetSettings().bDeferNarrowPhase = (ChaosImmediate_Collision_DeferNarrowPhase != 0);
diff --git a/Engine/Source/Runtime/Engine/Public/Physics/ImmediatePhysics/ImmediatePhysicsChaos/ImmediatePhysicsSimulation_Chaos.h b/Engine/Source/Runtime/Engine/Public/Physics/ImmediatePhysics/ImmediatePhysicsChaos/ImmediatePhysicsSimulation_Chaos.h
--- a/Engine/Source/Runtime/Engine/Public/Physics/ImmediatePhysics/ImmediatePhysicsChaos/ImmediatePhysicsSimulation_Chaos.h
+++ b/Engine/Source/Runtime/Engine/Public/Physics/ImmediatePhysics/ImmediatePhysicsChaos/ImmediatePhysicsSimulation_Chaos.h
@@ -86,15 +86,6 @@ namespace ImmediatePhysics_Chaos
 			const int32 VelocityIts,
 			const int32 ProjectionIts);
 
-		/** Set iteration counts for the legacy solver. A negative value with leave that iteration count unchanged */
-		void SetLegacySolverSettings(
-			const int32 SolverIts,
-			const int32 JointIts,
-			const int32 CollisionIts,
-			const int32 SolverPushOutIts,
-			const int32 JointPushOutIts,
-			const int32 CollisionPushOutIts);
-
 		/** Explicit debug draw path if the use case needs it to happen at a point outside of the simulation **/
 		void DebugDraw();
 
diff --git a/Engine/Source/Runtime/Experimental/Chaos/Private/Chaos/ChaosDebugDraw.cpp b/Engine/Source/Runtime/Experimental/Chaos/Private/Chaos/ChaosDebugDraw.cpp
--- a/Engine/Source/Runtime/Experimental/Chaos/Private/Chaos/ChaosDebugDraw.cpp
+++ b/Engine/Source/Runtime/Experimental/Chaos/Private/Chaos/ChaosDebugDraw.cpp
@@ -1356,7 +1356,7 @@ namespace Chaos
 				{
 					FAABB3 IslandAABB = IslandBounds[IslandIndex];
 
-					const FColor IslandColor = GetIslandColor(IslandIndex, !Graph.GetSolverIsland(IslandIndex)->IsSleeping());
+					const FColor IslandColor = GetIslandColor(IslandIndex, !Graph.GetIsland(IslandIndex)->IsSleeping());
 					const FAABB3 Bounds = IslandAABB.TransformedAABB(SpaceTransform);
 					FDebugDrawQueue::GetInstance().DrawDebugBox(Bounds.Center(), 0.5f * Bounds.Extents(), SpaceTransform.GetRotation(), IslandColor, false, UE_KINDA_SMALL_NUMBER, Settings.DrawPriority, 3.0f * Settings.LineThickness);
 				}
diff --git a/Engine/Source/Runtime/Experimental/Chaos/Private/Chaos/Collision/PBDCollisionConstraint.cpp b/Engine/Source/Runtime/Experimental/Chaos/Private/Chaos/Collision/PBDCollisionConstraint.cpp
--- a/Engine/Source/Runtime/Experimental/Chaos/Private/Chaos/Collision/PBDCollisionConstraint.cpp
+++ b/Engine/Source/Runtime/Experimental/Chaos/Private/Chaos/Collision/PBDCollisionConstraint.cpp
@@ -109,7 +109,7 @@ namespace Chaos
 		return false;
 	}
 
-	void FPBDCollisionConstraint::MakeInline(
+	void FPBDCollisionConstraint::Make(
 		FGeometryParticleHandle* Particle0,
 		const FImplicitObject* Implicit0,
 		const FPerShapeData* Shape0,
@@ -137,28 +137,6 @@ namespace Chaos
 		OutConstraint.Setup(ECollisionCCDType::Disabled, ShapesType, ImplicitLocalTransform0, ImplicitLocalTransform1, InCullDistance, bInUseManifold);
 	}
 
-	TUniquePtr<FPBDCollisionConstraint> FPBDCollisionConstraint::Make(
-		FGeometryParticleHandle* Particle0,
-		const FImplicitObject* Implicit0,
-		const FPerShapeData* Shape0,
-		const FBVHParticles* Simplicial0,
-		const FRigidTransform3& ImplicitLocalTransform0,
-		FGeometryParticleHandle* Particle1,
-		const FImplicitObject* Implicit1,
-		const FPerShapeData* Shape1,
-		const FBVHParticles* Simplicial1,
-		const FRigidTransform3& ImplicitLocalTransform1,
-		const FReal InCullDistance,
-		const bool bInUseManifold,
-		const EContactShapesType ShapesType)
-	{
-		FPBDCollisionConstraint* Constraint = new FPBDCollisionConstraint(Particle0, Implicit0, Shape0, Simplicial0, Particle1, Implicit1, Shape1, Simplicial1);
-		
-		Constraint->Setup(ECollisionCCDType::Disabled, ShapesType, ImplicitLocalTransform0, ImplicitLocalTransform1, InCullDistance, bInUseManifold);
-
-		return TUniquePtr<FPBDCollisionConstraint>(Constraint);
-	}
-
 	FPBDCollisionConstraint FPBDCollisionConstraint::MakeTriangle(const FImplicitObject* Implicit0)
 	{
 		FPBDCollisionConstraint Constraint;
@@ -208,7 +186,6 @@ namespace Chaos
 		, LastShapeWorldPositionDelta()
 		, LastShapeWorldRotationDelta()
 		, SolverBodies{ nullptr, nullptr }
-		, SolverIndex(INDEX_NONE)
 		, GJKWarmStartData()
 		, SavedManifoldPoints()
 		, ManifoldPoints()
@@ -249,7 +226,6 @@ namespace Chaos
 		, LastShapeWorldPositionDelta()
 		, LastShapeWorldRotationDelta()
 		, SolverBodies{ nullptr, nullptr }
-		, SolverIndex(INDEX_NONE)
 		, GJKWarmStartData()
 		, SavedManifoldPoints()
 		, ManifoldPoints()
@@ -260,7 +236,7 @@ namespace Chaos
 	{
 #if !UE_BUILD_TEST && !UE_BUILD_SHIPPING
 		// Make sure we have been dropped by the graph
-		ensure(ConstraintGraphIndex() == INDEX_NONE);
+		ensure(GetConstraintGraphIndex() == INDEX_NONE);
 #endif
 	}
 
diff --git a/Engine/Source/Runtime/Experimental/Chaos/Private/Chaos/Collision/PBDCollisionSolver.h b/Engine/Source/Runtime/Experimental/Chaos/Private/Chaos/Collision/PBDCollisionSolver.h
--- a/Engine/Source/Runtime/Experimental/Chaos/Private/Chaos/Collision/PBDCollisionSolver.h
+++ b/Engine/Source/Runtime/Experimental/Chaos/Private/Chaos/Collision/PBDCollisionSolver.h
@@ -18,11 +18,6 @@ CHAOS_API DECLARE_LOG_CATEGORY_EXTERN(LogChaosCollision, Log, All);
 
 namespace Chaos
 {
-	namespace Collisions
-	{
-		struct FContactIterationParameters;
-		struct FContactParticleParameters;
-	}
 	class FManifoldPoint;
 	class FPBDCollisionSolver;
 	class FPBDCollisionConstraint;
diff --git a/Engine/Source/Runtime/Experimental/Chaos/Private/Chaos/Collision/ParticlePairMidPhase.cpp b/Engine/Source/Runtime/Experimental/Chaos/Private/Chaos/Collision/ParticlePairMidPhase.cpp
--- a/Engine/Source/Runtime/Experimental/Chaos/Private/Chaos/Collision/ParticlePairMidPhase.cpp
+++ b/Engine/Source/Runtime/Experimental/Chaos/Private/Chaos/Collision/ParticlePairMidPhase.cpp
@@ -61,14 +61,15 @@ namespace Chaos
 		return true;
 	}
 
-	TUniquePtr<FPBDCollisionConstraint> CreateShapePairConstraint(
+	FPBDCollisionConstraintPtr CreateShapePairConstraint(
 		FGeometryParticleHandle* Particle0,
 		const FPerShapeData* InShape0,
 		FGeometryParticleHandle* Particle1,
 		const FPerShapeData* InShape1,
 		const FReal CullDistance,
 		const EContactShapesType ShapePairType,
-		const bool bUseManifold)
+		const bool bUseManifold,
+		FCollisionConstraintAllocator& Allocator)
 	{
 		const FImplicitObject* Implicit0 = InShape0->GetLeafGeometry();
 		const FBVHParticles* BVHParticles0 = FConstGenericParticleHandle(Particle0)->CollisionParticles().Get();
@@ -77,10 +78,10 @@ namespace Chaos
 		const FBVHParticles* BVHParticles1 = FConstGenericParticleHandle(Particle1)->CollisionParticles().Get();
 		const FRigidTransform3& ShapeRelativeTransform1 = (FRigidTransform3)InShape1->GetLeafRelativeTransform();
 
-		return FPBDCollisionConstraint::Make(Particle0, Implicit0, InShape0, BVHParticles0, ShapeRelativeTransform0, Particle1, Implicit1, InShape1, BVHParticles1, ShapeRelativeTransform1, CullDistance, bUseManifold, ShapePairType);
+		return Allocator.CreateCollisionConstraint(Particle0, Implicit0, InShape0, BVHParticles0, ShapeRelativeTransform0, Particle1, Implicit1, InShape1, BVHParticles1, ShapeRelativeTransform1, CullDistance, bUseManifold, ShapePairType);
 	}
 
-	TUniquePtr<FPBDCollisionConstraint> CreateImplicitPairConstraint(
+	FPBDCollisionConstraintPtr CreateImplicitPairConstraint(
 		FGeometryParticleHandle* Particle0,
 		const FImplicitObject* Implicit0,
 		const FPerShapeData* Shape0,
@@ -93,9 +94,10 @@ namespace Chaos
 		const FRigidTransform3& ShapeRelativeTransform1,
 		const FReal CullDistance,
 		const EContactShapesType ShapePairType,
-		const bool bUseManifold)
+		const bool bUseManifold,
+		FCollisionConstraintAllocator& Allocator)
 	{
-		return FPBDCollisionConstraint::Make(Particle0, Implicit0, Shape0, BVHParticles0, ShapeRelativeTransform0, Particle1, Implicit1, Shape1, BVHParticles1, ShapeRelativeTransform1, CullDistance, bUseManifold, ShapePairType);
+		return Allocator.CreateCollisionConstraint(Particle0, Implicit0, Shape0, BVHParticles0, ShapeRelativeTransform0, Particle1, Implicit1, Shape1, BVHParticles1, ShapeRelativeTransform1, CullDistance, bUseManifold, ShapePairType);
 	}
 
 
@@ -258,9 +260,9 @@ namespace Chaos
 	void FSingleShapePairCollisionDetector::CreateConstraint(const FReal CullDistance, FCollisionContext& Context)
 	{
 		PHYSICS_CSV_SCOPED_EXPENSIVE(PhysicsVerbose, NarrowPhase_CreateConstraint);
-		check(Constraint == nullptr);
+		check(!Constraint.IsValid());
 
-		Constraint = CreateShapePairConstraint(Particle0, Shape0, Particle1, Shape1, CullDistance, ShapePairType, Context.GetSettings().bAllowManifolds);
+		Constraint = CreateShapePairConstraint(Particle0, Shape0, Particle1, Shape1, CullDistance, ShapePairType, Context.GetSettings().bAllowManifolds, MidPhase.GetCollisionAllocator());
 
 		Constraint->GetContainerCookie().MidPhase = &MidPhase;
 		Constraint->GetContainerCookie().bIsMultiShapePair = false;
@@ -278,7 +280,7 @@ namespace Chaos
 			return GenerateCollisionProbeImpl(CullDistance, Dt, Context);
 		}
 
-		if (Constraint == nullptr)
+		if (!Constraint.IsValid())
 		{
 			// Lazy creation of the constraint. If a shape pair never gets within CullDistance of each
 			// other, we never allocate a constraint for them. Once they overlap, we reuse the constraint
@@ -287,7 +289,7 @@ namespace Chaos
 			CreateConstraint(CullDistance, Context);
 		}
 
-		if (Constraint != nullptr)
+		if (Constraint.IsValid())
 		{
 			PHYSICS_CSV_SCOPED_EXPENSIVE(PhysicsVerbose, NarrowPhase_UpdateConstraint);
 
@@ -366,13 +368,13 @@ namespace Chaos
 			return GenerateCollisionProbeImpl(CullDistance, Dt, Context);
 		}
 
-		if (Constraint == nullptr)
+		if (!Constraint.IsValid())
 		{
 			// Lazy creation of the constraint. 
 			CreateConstraint(CullDistance, Context);
 		}
 
-		if (Constraint != nullptr)
+		if (Constraint.IsValid())
 		{
 			PHYSICS_CSV_SCOPED_EXPENSIVE(PhysicsVerbose, NarrowPhase_UpdateConstraintCCD);
 
@@ -430,12 +432,12 @@ namespace Chaos
 		// Same as regular constraint generation, but always defer narrow phase.
 		// Don't do any initial constraint computations.
 
-		if (Constraint == nullptr)
+		if (!Constraint.IsValid())
 		{
 			CreateConstraint(CullDistance, Context);
 		}
 
-		if (Constraint != nullptr)
+		if (Constraint.IsValid())
 		{
 			PHYSICS_CSV_SCOPED_EXPENSIVE(PhysicsVerbose, NarrowPhase_UpdateConstraintProbe);
 
@@ -449,7 +451,7 @@ namespace Chaos
 
 	void FSingleShapePairCollisionDetector::WakeCollision(const int32 SleepEpoch)
 	{
-		if ((Constraint != nullptr) && IsUsedSince(SleepEpoch))
+		if (Constraint.IsValid() && IsUsedSince(SleepEpoch))
 		{
 #if CHAOS_CONSTRAINTHANDLE_DEBUG_ENABLED
 			// If a particle changed to kinematic, all constraints without a dynamic should have been removed
@@ -473,9 +475,9 @@ namespace Chaos
 	{
 		const int32 CurrentEpoch = MidPhase.GetCollisionAllocator().GetCurrentEpoch();
 
-		if (Constraint == nullptr)
+		if (!Constraint.IsValid())
 		{
-			Constraint = MakeUnique<FPBDCollisionConstraint>();
+			Constraint = MidPhase.GetCollisionAllocator().CreateCollisionConstraint();
 			Constraint->GetContainerCookie().MidPhase = &MidPhase;
 			Constraint->GetContainerCookie().bIsMultiShapePair = false;
 			Constraint->GetContainerCookie().CreationEpoch = CurrentEpoch;
@@ -483,7 +485,7 @@ namespace Chaos
 
 		// Copy the constraint over the existing one, taking care to leave the cookie intact
 		const FPBDCollisionConstraintContainerCookie Cookie = Constraint->GetContainerCookie();
-		*(Constraint.Get()) = SourceConstraint;
+		*Constraint = SourceConstraint;
 		Constraint->GetContainerCookie() = Cookie;
 
 		// Add the constraint to the active list
@@ -662,10 +664,10 @@ namespace Chaos
 
 	FPBDCollisionConstraint* FMultiShapePairCollisionDetector::FindConstraint(const FCollisionParticlePairConstraintKey& Key)
 	{
-		TUniquePtr<FPBDCollisionConstraint>* Constraint = Constraints.Find(Key.GetKey());
-		if (Constraint != nullptr)
+		FPBDCollisionConstraintPtr* PConstraint = Constraints.Find(Key.GetKey());
+		if (PConstraint != nullptr)
 		{
-			return (*Constraint).Get();
+			return (*PConstraint).Get();
 		}
 		return nullptr;
 	}
@@ -687,7 +689,7 @@ namespace Chaos
 		const FCollisionParticlePairConstraintKey& Key)
 	{
 		PHYSICS_CSV_SCOPED_EXPENSIVE(PhysicsVerbose, NarrowPhase_CreateConstraint);
-		TUniquePtr<FPBDCollisionConstraint> Constraint = CreateImplicitPairConstraint(InParticle0, Implicit0, InShape0, BVHParticles0, ShapeRelativeTransform0, InParticle1, Implicit1, InShape1, BVHParticles1, ShapeRelativeTransform1, CullDistance, ShapePairType, bUseManifold);
+		FPBDCollisionConstraintPtr Constraint = CreateImplicitPairConstraint(InParticle0, Implicit0, InShape0, BVHParticles0, ShapeRelativeTransform0, InParticle1, Implicit1, InShape1, BVHParticles1, ShapeRelativeTransform1, CullDistance, ShapePairType, bUseManifold, MidPhase.GetCollisionAllocator());
 
 		Constraint->GetContainerCookie().MidPhase = &MidPhase;
 		Constraint->GetContainerCookie().bIsMultiShapePair = true;
@@ -701,7 +703,7 @@ namespace Chaos
 		const int32 CurrentEpoch = MidPhase.GetCollisionAllocator().GetCurrentEpoch();
 		for (auto& KVP : Constraints)
 		{
-			TUniquePtr<FPBDCollisionConstraint>& Constraint = KVP.Value;
+			FPBDCollisionConstraintPtr& Constraint = KVP.Value;
 			if (Constraint->GetContainerCookie().LastUsedEpoch >= SleepEpoch)
 			{
 				Constraint->GetContainerCookie().LastUsedEpoch = CurrentEpoch;
@@ -736,7 +738,7 @@ namespace Chaos
 		for (auto& KVP : Constraints)
 		{
 			const uint32 Key = KVP.Key;
-			TUniquePtr<FPBDCollisionConstraint>& Constraint = KVP.Value;
+			FPBDCollisionConstraintPtr& Constraint = KVP.Value;
 			if (Constraint->GetContainerCookie().LastUsedEpoch < CurrentEpoch)
 			{
 				Pruned.Add(Key);
@@ -1044,7 +1046,7 @@ namespace Chaos
 		bool bInGraph = false;
 		VisitConstCollisions([&bInGraph](const FPBDCollisionConstraint& Constraint)
 		{
-			if (Constraint.ConstraintGraphIndex() != INDEX_NONE)
+			if (Constraint.GetConstraintGraphIndex() != INDEX_NONE)
 			{
 				bInGraph = true;
 				return ECollisionVisitorResult::Stop;
