Raycast Broad Phase Grid (Bounding Volume) optimization :
Early exit when blocked hit is requested. No need to travel in further cell if hit has been found.

diff --git a/Engine/Source/Programs/HeadlessChaos/Private/HeadlessChaosApp.cpp b/Engine/Source/Programs/HeadlessChaos/Private/HeadlessChaosApp.cpp
--- a/Engine/Source/Programs/HeadlessChaos/Private/HeadlessChaosApp.cpp
+++ b/Engine/Source/Programs/HeadlessChaos/Private/HeadlessChaosApp.cpp
@@ -253,6 +253,7 @@ TEST(EPA, EPATests) {
 TEST(BP, BroadphaseTests) {
 	ChaosTest::AABBTreeDirtyGridFunctionsWithEdgeCase();
 	ChaosTest::GridBPTest();
+	ChaosTest::GridBPEarlyExitTest();
 	ChaosTest::GridBPTest2();
 	ChaosTest::AABBTreeTest();
 	ChaosTest::AABBTreeTestDynamic();
diff --git a/Engine/Source/Programs/HeadlessChaos/Private/HeadlessChaosTestBroadphase.cpp b/Engine/Source/Programs/HeadlessChaos/Private/HeadlessChaosTestBroadphase.cpp
--- a/Engine/Source/Programs/HeadlessChaos/Private/HeadlessChaosTestBroadphase.cpp
+++ b/Engine/Source/Programs/HeadlessChaos/Private/HeadlessChaosTestBroadphase.cpp
@@ -74,11 +74,6 @@ namespace ChaosTest
 				{
 					return false;
 				}
-				if (Instances.Num() >= BlockAfterN)
-				{
-					//blocking so adjust Length
-					CurData.SetLength(NewLength);
-				}
 			}
 			
 			return true;
@@ -115,6 +110,11 @@ namespace ChaosTest
 			return VisitSweep(Instance, CurData);
 		}
 
+		virtual bool HasBlockingHit() const override
+		{ 
+			return Instances.Num() >= BlockAfterN; 
+		}
+
 		TArray<int32> Instances;
 	};
 
@@ -442,6 +442,41 @@ namespace ChaosTest
 		SpatialTestHelper(Spatial, Boxes.Get(), Box);
 	}
 
+	void GridBPEarlyExitTest()
+	{
+		TUniquePtr<TBox<FReal, 3>> Box;
+		auto Boxes = BuildBoxes(Box);
+		TBoundingVolume<int32> Spatial(MakeParticleView(Boxes.Get()));
+		// SpatialTestHelper(Spatial, Boxes.Get(), Box);
+
+		//gather along ray
+		{
+			FVisitor Visitor(FVec3(10, 0, 0), FVec3(0, 1, 0), 0, *Boxes);
+			Spatial.Raycast(Visitor.Start, Visitor.Dir, 1000, Visitor);
+			EXPECT_EQ(Visitor.Instances.Num(), 10);
+			EXPECT_EQ(Visitor.Instances[0], 0);
+			EXPECT_EQ(Visitor.Instances[9], 90);
+		}
+
+		// Stop after first hits in the first cell
+		{
+			FVisitor Visitor(FVec3(10, 0, 0), FVec3(0, 1, 0), 0, *Boxes);
+			Visitor.BlockAfterN = 1;
+			Spatial.Raycast(Visitor.Start, Visitor.Dir, 1000, Visitor);
+			EXPECT_EQ(Visitor.Instances.Num(), 1);
+			EXPECT_EQ(Visitor.Instances[0], 0);
+		}
+
+		// Stop after first hits in the first cell, going backward
+		{
+			FVisitor Visitor(FVec3(10, 1000, 0), FVec3(0, -1, 0), 0, *Boxes);
+			Visitor.BlockAfterN = 1;
+			Spatial.Raycast(Visitor.Start, Visitor.Dir, 1000, Visitor);
+			EXPECT_EQ(Visitor.Instances.Num(), 1);
+			EXPECT_EQ(Visitor.Instances[0], 90);
+		}
+	}
+
 	void GridBPTest2()
 	{
 		TUniquePtr<TBox<FReal, 3>> Box = MakeUnique<TBox<FReal, 3>>(FVec3(0, 0, 0), FVec3(100, 100, 100));
diff --git a/Engine/Source/Programs/HeadlessChaos/Public/HeadlessChaosTestBroadphase.h b/Engine/Source/Programs/HeadlessChaos/Public/HeadlessChaosTestBroadphase.h
--- a/Engine/Source/Programs/HeadlessChaos/Public/HeadlessChaosTestBroadphase.h
+++ b/Engine/Source/Programs/HeadlessChaos/Public/HeadlessChaosTestBroadphase.h
@@ -9,6 +9,8 @@ namespace ChaosTest {
 
 	void GridBPTest();
 
+	void GridBPEarlyExitTest();
+
 	void GridBPTest2();
 
 	void AABBTreeTest();
diff --git a/Engine/Source/Runtime/Experimental/Chaos/Public/Chaos/AABBTree.h b/Engine/Source/Runtime/Experimental/Chaos/Public/Chaos/AABBTree.h
--- a/Engine/Source/Runtime/Experimental/Chaos/Public/Chaos/AABBTree.h
+++ b/Engine/Source/Runtime/Experimental/Chaos/Public/Chaos/AABBTree.h
@@ -3156,6 +3156,11 @@ private:
 				return nullptr; 
 			}
 
+			bool HasBlockingHit() const
+			{
+				return false;
+			}
+
 			bool ShouldIgnore(const TSpatialVisitorData<TPayloadType>& Instance) const { return false; }
 			TArray<TPayloadType>& CollectedResults;
 		};
@@ -3416,5 +3421,4 @@ FArchive& operator<<(FChaosArchive& Ar, TAABBTree<TPayloadType, TLeafType, bMuta
 	return Ar;
 }
 
-
 }
diff --git a/Engine/Source/Runtime/Experimental/Chaos/Public/Chaos/BoundingVolume.h b/Engine/Source/Runtime/Experimental/Chaos/Public/Chaos/BoundingVolume.h
--- a/Engine/Source/Runtime/Experimental/Chaos/Public/Chaos/BoundingVolume.h
+++ b/Engine/Source/Runtime/Experimental/Chaos/Public/Chaos/BoundingVolume.h
@@ -583,6 +583,12 @@ private:
 					}
 				}
 
+				// Early exit if the raycast has a hit and it has set as bloking hit, it is not necessary to navigate along the raycast
+				if (Visitor.HasBlockingHit())
+				{
+					return false;
+				}
+
 				CellsVisited.Add(CellIdx);
 
 
diff --git a/Engine/Source/Runtime/Experimental/Chaos/Public/Chaos/Collision/SpatialAccelerationBroadPhase.h b/Engine/Source/Runtime/Experimental/Chaos/Public/Chaos/Collision/SpatialAccelerationBroadPhase.h
--- a/Engine/Source/Runtime/Experimental/Chaos/Public/Chaos/Collision/SpatialAccelerationBroadPhase.h
+++ b/Engine/Source/Runtime/Experimental/Chaos/Public/Chaos/Collision/SpatialAccelerationBroadPhase.h
@@ -66,6 +66,11 @@ namespace Chaos
 			return &AccelerationHandle;
 		}
 
+		bool HasBlockingHit() const
+		{
+			return false;
+		}
+
 	private:
 		FCollisionFilterData SimFilterData;
 		FUniqueIdx ParticleUniqueIdx; // unique id of the particle visiting, used to skip self intersection as early as possible
diff --git a/Engine/Source/Runtime/Experimental/Chaos/Public/Chaos/ISpatialAcceleration.h b/Engine/Source/Runtime/Experimental/Chaos/Public/Chaos/ISpatialAcceleration.h
--- a/Engine/Source/Runtime/Experimental/Chaos/Public/Chaos/ISpatialAcceleration.h
+++ b/Engine/Source/Runtime/Experimental/Chaos/Public/Chaos/ISpatialAcceleration.h
@@ -148,6 +148,8 @@ public:
 
 	/** Return a pointer to the payload on which we are querying the acceleration structure */
 	virtual const void* GetQueryPayload() const { return nullptr; }
+
+	virtual bool HasBlockingHit() const { return false; }
 };
 
 /**
@@ -443,6 +445,11 @@ public:
 		return Visitor.GetQueryPayload();
 	}
 
+	FORCEINLINE bool HasBlockingHit() const
+	{
+		return Visitor.HasBlockingHit();
+	}
+
 private:
 	ISpatialVisitor<TPayloadType, T>& Visitor;
 };
diff --git a/Engine/Source/Runtime/PhysicsCore/Private/SQAccelerator.cpp b/Engine/Source/Runtime/PhysicsCore/Private/SQAccelerator.cpp
--- a/Engine/Source/Runtime/PhysicsCore/Private/SQAccelerator.cpp
+++ b/Engine/Source/Runtime/PhysicsCore/Private/SQAccelerator.cpp
@@ -178,6 +178,11 @@ struct TSQVisitor : public Chaos::ISpatialVisitor<TPayload, Chaos::FReal>
 		return nullptr;
 	}
 
+	virtual bool HasBlockingHit() const override
+	{ 
+		return HitBuffer.HasBlockingHit();
+	}
+
 private:
 
 	enum class ESQType
