TraceLandscape_Stress test failing due to Expected Single Line Trace At Hit Location: to be 1, but it was 0

diff --git a/Engine/Source/Programs/HeadlessChaos/Private/HeadlessChaosTestHeightfield.cpp b/Engine/Source/Programs/HeadlessChaos/Private/HeadlessChaosTestHeightfield.cpp
--- a/Engine/Source/Programs/HeadlessChaos/Private/HeadlessChaosTestHeightfield.cpp
+++ b/Engine/Source/Programs/HeadlessChaos/Private/HeadlessChaosTestHeightfield.cpp
@@ -378,60 +378,106 @@ namespace ChaosTest {
 
 	void RaycastVariousWalkOnHeightField()
 	{
-		constexpr int32 Rows = 64;
-		constexpr int32 Columns = Rows;
-		FVec3 Scale(1.0, 1.0, 1.0);
-		TArray<FReal> Heights;
-		Heights.AddZeroed(Rows * Columns);
-
-		// Add a mountain on the diagonal
-		for (int32 Index = 0; Index < Columns; ++Index)
 		{
-			Heights[Index * Columns + Index] = 20;
-		}
-		constexpr int32 BigMountainIndex = Columns / 2;
-		Heights[BigMountainIndex * Columns + BigMountainIndex] = 40;
+			constexpr int32 Rows = 64;
+			constexpr int32 Columns = Rows;
+			FVec3 Scale(1.0, 1.0, 1.0);
+			TArray<FReal> Heights;
+			Heights.AddZeroed(Rows * Columns);
 
-		TArray<FReal> HeightsCopy = Heights;
+			// Add a mountain on the diagonal
+			for (int32 Index = 0; Index < Columns; ++Index)
+			{
+				Heights[Index * Columns + Index] = 20;
+			}
+			constexpr int32 BigMountainIndex = Columns / 2;
+			Heights[BigMountainIndex * Columns + BigMountainIndex] = 40;
 
-		FHeightField Heightfield(MoveTemp(HeightsCopy), TArray<uint8>(), Rows, Columns, Scale);
-		const auto& Bounds = Heightfield.BoundingBox();	//Current API forces us to do this to cache the bounds
+			TArray<FReal> HeightsCopy = Heights;
 
-		FReal TOI;
-		FVec3 Position, Normal;
-		int32 FaceIdx = 0;
+			FHeightField Heightfield(MoveTemp(HeightsCopy), TArray<uint8>(), Rows, Columns, Scale);
+			const auto& Bounds = Heightfield.BoundingBox();	//Current API forces us to do this to cache the bounds
+
+			FReal TOI;
+			FVec3 Position, Normal;
+			int32 FaceIdx = 0;
+
+			// Hit the diagonal, used directly the three walks (LowRes / Fast / Slow)
+			{
+				const FVec3 Start(2.0, 0.0, 10.0);
+				FVec3 Dir(0.0, 1.0, 0.0);
+				Dir.Normalize();
+				const bool bResult = Heightfield.Raycast(Start, Dir, 100.0, 0, TOI, Position, Normal, FaceIdx);
+				EXPECT_TRUE(bResult);
+			}
+			// Miss the diagonal, used the three walks but miss them  (LowRes / Fast / Slow) (worse case)
+			{
+				const FVec3 Start(0.6, 0.0, 17.0);
+				FVec3 Dir(1.0, 1.0, 0.0);
+				Dir.Normalize();
+				const bool bResult = Heightfield.Raycast(Start, Dir, 100.0, 0, TOI, Position, Normal, FaceIdx);
+				EXPECT_FALSE(bResult);
+			}
+			// Hit the diagonal in the middle used Low Res for few steps then Fast walk for few step and Slow walks (optimal use case)
+			{
+				const FVec3 Start(64.0, 0.0, 10.0);
+				FVec3 Dir(-1.0, 1.0, 0.0);
+				Dir.Normalize();
+				const bool bResult = Heightfield.Raycast(Start, Dir, 100.0, 0, TOI, Position, Normal, FaceIdx);
+				EXPECT_TRUE(bResult);
+			}
+			// Miss the diagonal in the middle used Low Res for few steps then Fast walk and Slow walks, then miss, then use LowRes until the end
+			{
+				const FVec3 Start(63.0, 0.0, 36.0);
+				FVec3 Dir(-1.0, 1.0, 0.0);
+				Dir.Normalize();
+				const bool bResult = Heightfield.Raycast(Start, Dir, 100.0, 0, TOI, Position, Normal, FaceIdx);
+				EXPECT_FALSE(bResult);
+			}
+			// Hit at the HeightField end boundary testing the LowRes Heightfield not fully filled  
+			// (HeightField size) % LowResolution =>  64 % 6 = 4
+			{
+				const FVec3 Start(20.0, 63.0, 10.0);
+				FVec3 Dir(1.0, 0.0, 0.0);
+				Dir.Normalize();
+				const bool bResult = Heightfield.Raycast(Start, Dir, 100.0, 0, TOI, Position, Normal, FaceIdx);
+				EXPECT_TRUE(bResult);
+			}
 
-		// Hit the diagonal, used directly the three walks (LowRes / Fast / Slow)
-		{
-			const FVec3 Start(2.0, 0.0, 10.0);
-			FVec3 Dir(0.0, 1.0, 0.0);
-			Dir.Normalize();
-			const bool bResult = Heightfield.Raycast(Start, Dir, 100.0, 0, TOI, Position, Normal, FaceIdx);
-			EXPECT_TRUE(bResult);
-		}
-		// Miss the diagonal, used the three walks but miss them  (LowRes / Fast / Slow) (worse case)
-		{
-			const FVec3 Start(0.6, 0.0, 17.0);
-			FVec3 Dir(1.0, 1.0, 0.0);
-			Dir.Normalize();
-			const bool bResult = Heightfield.Raycast(Start, Dir, 100.0, 0, TOI, Position, Normal, FaceIdx);
-			EXPECT_FALSE(bResult);
-		}
-		// Hit the diagonal in the middle used Low Res for few steps then Fast walk for few step and Slow walks (optimal use case)
-		{
-			const FVec3 Start(64.0, 0.0, 10.0);
-			FVec3 Dir(-1.0, 1.0, 0.0);
-			Dir.Normalize();
-			const bool bResult = Heightfield.Raycast(Start, Dir, 100.0, 0, TOI, Position, Normal, FaceIdx);
-			EXPECT_TRUE(bResult);
 		}
-		// Miss the diagonal in the middle used Low Res for few steps then Fast walk and Slow walks, then miss, then use LowRes until the end
 		{
-			const FVec3 Start(63.0, 0.0, 36.0);
-			FVec3 Dir(-1.0, 1.0, 0.0);
-			Dir.Normalize();
-			const bool bResult = Heightfield.Raycast(Start, Dir, 100.0, 0, TOI, Position, Normal, FaceIdx);
-			EXPECT_FALSE(bResult);
+			constexpr int32 Rows = 64;
+			constexpr int32 Columns = Rows;
+			FVec3 Scale(1.0, 1.0, 1.0);
+			TArray<FReal> Heights;
+			Heights.AddZeroed(Rows * Columns);
+
+			// Add a mountain on the diagonal
+			for (int32 Index = 0; Index < Columns; ++Index)
+			{
+				Heights[Index * Columns + 62] = 20;
+			}
+
+			TArray<FReal> HeightsCopy = Heights;
+
+			FHeightField Heightfield(MoveTemp(HeightsCopy), TArray<uint8>(), Rows, Columns, Scale);
+			const auto& Bounds = Heightfield.BoundingBox();	//Current API forces us to do this to cache the bounds
+
+			FReal TOI;
+			FVec3 Position, Normal;
+			int32 FaceIdx = 0;
+
+			// Hit at the HeightField end boundary testing the LowRes Heightfield not fully filled  
+			// (HeightField size) % LowResolution =>  64 % 6 = 4
+			// Jira UE-162256
+			{
+				const FVec3 Start(58.0, 63.0, 10.0);
+				FVec3 Dir(1.0, 0.0, 0.0);
+				Dir.Normalize();
+				const bool bResult = Heightfield.Raycast(Start, Dir, 100.0, 0, TOI, Position, Normal, FaceIdx);
+				EXPECT_TRUE(bResult);
+			}
+
 		}
 	}
 
diff --git a/Engine/Source/Runtime/Experimental/Chaos/Private/Chaos/HeightField.cpp b/Engine/Source/Runtime/Experimental/Chaos/Private/Chaos/HeightField.cpp
--- a/Engine/Source/Runtime/Experimental/Chaos/Private/Chaos/HeightField.cpp
+++ b/Engine/Source/Runtime/Experimental/Chaos/Private/Chaos/HeightField.cpp
@@ -1257,10 +1257,9 @@ namespace Chaos
 			for (int Axis = 0; Axis < 2; ++Axis)
 			{
 				CellIdx[Axis] += (Times[Axis] <= BestTime) ? ((Dir[Axis] * ScaleSign[Axis]) > 0 ? LowResInc : -LowResInc) : 0;
-				if (CellIdx[Axis] < 0 || CellIdx[Axis] >= FlatGrid.Counts()[Axis])
-				{
-					return false;
-				}
+				// Make sure CellIdx doesn't go beyond the grid
+				CellIdx[Axis] = FMath::Min<int32>(CellIdx[Axis], FlatGrid.Counts()[Axis] - 1);
+				CellIdx[Axis] = FMath::Max<int32>(CellIdx[Axis], 0);
 			}
 
 			if (PrevIdx == CellIdx)
