Chaos - Sphere-Convex collision optimization

diff --git a/Engine/Source/Programs/HeadlessChaos/Private/HeadlessChaosTestDetectCollision.cpp b/Engine/Source/Programs/HeadlessChaos/Private/HeadlessChaosTestDetectCollision.cpp
--- a/Engine/Source/Programs/HeadlessChaos/Private/HeadlessChaosTestDetectCollision.cpp
+++ b/Engine/Source/Programs/HeadlessChaos/Private/HeadlessChaosTestDetectCollision.cpp
@@ -61,14 +61,6 @@ namespace ChaosTest
 				CheckContactPoint(ContactPoint, FVec3(0, 0, -SphereRadius), FVec3(0, 0, 0.5f * BoxSize.Z), FVec3(0,0,1), 1, FVec3(0,0,1), SpherePos.Z - SphereRadius);
 			}
 		}
-		{
-			FContactPoint ContactPoint = ConvexSphereContactPoint(Convex, ConvexTransform, Sphere, SphereTransform);
-			EXPECT_TRUE(ContactPoint.IsSet());
-			if (ContactPoint.IsSet())
-			{
-				CheckContactPoint(ContactPoint, FVec3(0, 0, 0.5f * BoxSize.Z), FVec3(0, 0, -SphereRadius), FVec3(0, 0, -1), 0, FVec3(0, 0, -1), SpherePos.Z - SphereRadius);
-			}
-		}
 	}
 
 	// Sphere-Convex(with margin) far separated when an edge is the nearest feature.
diff --git a/Engine/Source/Programs/HeadlessChaos/Private/HeadlessChaosTestGJKDistance.cpp b/Engine/Source/Programs/HeadlessChaos/Private/HeadlessChaosTestGJKDistance.cpp
--- a/Engine/Source/Programs/HeadlessChaos/Private/HeadlessChaosTestGJKDistance.cpp
+++ b/Engine/Source/Programs/HeadlessChaos/Private/HeadlessChaosTestGJKDistance.cpp
@@ -24,7 +24,12 @@ namespace ChaosTest
 		{
 			TSphere<FReal, 3> A(FVec3(12, 0, 0), 5);
 			TSphere<FReal, 3> B(FVec3(4, 0, 0), 2);
-			EGJKDistanceResult Result = GJKDistance<FReal>(A, B, FRigidTransform3(FVec3(2, 0, 0), FRotation3::FromIdentity()), Distance, NearestA, NearestB, Normal);
+			const FRigidTransform3 BToATm = FRigidTransform3(FVec3(2, 0, 0), FRotation3::FromIdentity());
+			EGJKDistanceResult Result = GJKDistance<FReal>(
+				TGJKCoreShape(A),
+				TGJKCoreShapeTransformed(B, BToATm),
+				GJKDistanceInitialV(A, B, BToATm),
+				Distance, NearestA, NearestB, Normal);
 			EXPECT_NE(Result, EGJKDistanceResult::Separated);
 		}
 
@@ -32,7 +37,12 @@ namespace ChaosTest
 		{
 			TSphere<FReal, 3> A(FVec3(12, 0, 0), 5);
 			TSphere<FReal, 3> B(FVec3(4, 0, 0), 2);
-			EGJKDistanceResult Result = GJKDistance<FReal>(A, B, FRigidTransform3::Identity, Distance, NearestA, NearestB, Normal);
+			const FRigidTransform3 BToATm = FRigidTransform3::Identity;
+			EGJKDistanceResult Result = GJKDistance<FReal>(
+				TGJKCoreShape(A), 
+				TGJKCoreShapeTransformed(B, BToATm),
+				GJKDistanceInitialV(A, B, BToATm),
+				Distance, NearestA, NearestB, Normal);
 			EXPECT_EQ(Result, EGJKDistanceResult::Separated);
 			EXPECT_NEAR(Distance, (FReal)1, Tolerance);
 			EXPECT_NEAR(NearestA.X, (FReal)7, Tolerance);
@@ -48,7 +58,12 @@ namespace ChaosTest
 			TSphere<FReal, 3> A(FVec3(0, 0, 0), 2);
 			TSphere<FReal, 3> B(FVec3(0, 0, 0), 2);
 			FVec3 BPos = FVec3(3, 3, 0);
-			EGJKDistanceResult Result = GJKDistance<FReal>(A, B, FRigidTransform3(BPos, FRotation3::FromIdentity()), Distance, NearestA, NearestB, Normal);
+			const FRigidTransform3 BToATm = FRigidTransform3(BPos, FRotation3::FromIdentity());
+			EGJKDistanceResult Result = GJKDistance<FReal>(
+				TGJKCoreShape(A),
+				TGJKCoreShapeTransformed(B, BToATm),
+				GJKDistanceInitialV(A, B, BToATm),
+				Distance, NearestA, NearestB, Normal);
 			EXPECT_EQ(Result, EGJKDistanceResult::Separated);
 			FVec3 CenterDelta = (B.GetCenter() + BPos) - A.GetCenter();
 			FVec3 CenterDir = CenterDelta.GetSafeNormal();
@@ -66,7 +81,12 @@ namespace ChaosTest
 			TSphere<FReal, 3> A(FVec3(12, 0, 0), 5);
 			TSphere<FReal, 3> B(FVec3(4, 0, 0), 2);
 			FVec3 BPos = FVec3(0.99, 0, 0);
-			EGJKDistanceResult Result = GJKDistance<FReal>(A, B, FRigidTransform3(BPos, FRotation3::FromIdentity()), Distance, NearestA, NearestB, Normal);
+			const FRigidTransform3 BToATm = FRigidTransform3(BPos, FRotation3::FromIdentity());
+			EGJKDistanceResult Result = GJKDistance<FReal>(
+				TGJKCoreShape(A),
+				TGJKCoreShapeTransformed(B, BToATm),
+				GJKDistanceInitialV(A, B, BToATm),
+				Distance, NearestA, NearestB, Normal);
 			EXPECT_EQ(Result, EGJKDistanceResult::Separated);
 			EXPECT_NEAR(Distance, (FReal)1 - BPos.X, Tolerance);
 			EXPECT_NEAR(NearestA.X, (FReal)7, Tolerance);
@@ -96,7 +116,12 @@ namespace ChaosTest
 		{
 			FAABB3 A(FVec3(5, -2, -2), FVec3(8, 2, 2));
 			TSphere<FReal, 3> B(FVec3(2, 0, 0), 2);
-			EGJKDistanceResult Result = GJKDistance<FReal>(A, B, FRigidTransform3(FVec3(2, 0, 0), FRotation3::FromIdentity()), Distance, NearestA, NearestB, Normal);
+			const FRigidTransform3 BToATm = FRigidTransform3(FVec3(2, 0, 0), FRotation3::FromIdentity());
+			EGJKDistanceResult Result = GJKDistance<FReal>(
+				TGJKShape(A),
+				TGJKCoreShapeTransformed(B, BToATm),
+				GJKDistanceInitialV(A, B, BToATm),
+				Distance, NearestA, NearestB, Normal);
 			EXPECT_NE(Result, EGJKDistanceResult::Separated);
 		}
 
@@ -104,7 +129,12 @@ namespace ChaosTest
 		{
 			FAABB3 A(FVec3(5, -2, -2), FVec3(8, 2, 2));
 			TSphere<FReal, 3> B(FVec3(2, 0, 0), 2);
-			EGJKDistanceResult Result = GJKDistance<FReal>(A, B, FRigidTransform3::Identity, Distance, NearestA, NearestB, Normal);
+			const FRigidTransform3 BToATm = FRigidTransform3::Identity;
+			EGJKDistanceResult Result = GJKDistance<FReal>(
+				TGJKShape(A),
+				TGJKCoreShapeTransformed(B, BToATm),
+				GJKDistanceInitialV(A, B, BToATm),
+				Distance, NearestA, NearestB, Normal);
 			EXPECT_EQ(Result, EGJKDistanceResult::Separated);
 			EXPECT_NEAR(Distance, (FReal)1, Tolerance);
 			EXPECT_NEAR(NearestA.X, (FReal)5, Tolerance);
@@ -118,7 +148,12 @@ namespace ChaosTest
 		{
 			FAABB3 A(FVec3(5, -2, -2), FVec3(8, 2, 2));
 			TSphere<FReal, 3> B(FVec3(2, 0, 0), 2);
-			EGJKDistanceResult Result = GJKDistance<FReal>(B, A, FRigidTransform3::Identity, Distance, NearestB, NearestA, Normal);
+			const FRigidTransform3 BToATm = FRigidTransform3::Identity;
+			EGJKDistanceResult Result = GJKDistance<FReal>(
+				TGJKShape(A),
+				TGJKCoreShapeTransformed(B, BToATm),
+				GJKDistanceInitialV(A, B, BToATm),
+				Distance, NearestA, NearestB, Normal);
 			EXPECT_EQ(Result, EGJKDistanceResult::Separated);
 			EXPECT_NEAR(Distance, (FReal)1, Tolerance);
 			EXPECT_NEAR(NearestA.X, (FReal)5, Tolerance);
@@ -133,7 +168,12 @@ namespace ChaosTest
 		{
 			FAABB3 A(FVec3(5, 2, 2), FVec3(8, 4, 4));
 			TSphere<FReal, 3> B(FVec3(2, 0, 0), 2);
-			EGJKDistanceResult Result = GJKDistance<FReal>(A, B, FRigidTransform3::Identity, Distance, NearestA, NearestB, Normal);
+			const FRigidTransform3 BToATm = FRigidTransform3::Identity;
+			EGJKDistanceResult Result = GJKDistance<FReal>(
+				TGJKShape(A),
+				TGJKCoreShapeTransformed(B, BToATm),
+				GJKDistanceInitialV(A, B, BToATm),
+				Distance, NearestA, NearestB, Normal);
 			FVec3 NearPointOnA = A.Min();
 			FVec3 SphereNearPointDir = (NearPointOnA - B.GetCenter()).GetSafeNormal();
 			FVec3 NearPointOnB = B.GetCenter() + SphereNearPointDir * B.GetRadius();
@@ -150,7 +190,12 @@ namespace ChaosTest
 		{
 			FAABB3 A(FVec3(5, 2, 2), FVec3(8, 4, 4));
 			TSphere<FReal, 3> B(FVec3(2, 0, 0), 2);
-			EGJKDistanceResult Result = GJKDistance<FReal>(B, A, FRigidTransform3::Identity, Distance, NearestB, NearestA, Normal);
+			const FRigidTransform3 BToATm = FRigidTransform3::Identity;
+			EGJKDistanceResult Result = GJKDistance<FReal>(
+				TGJKShape(A),
+				TGJKCoreShapeTransformed(B, BToATm),
+				GJKDistanceInitialV(A, B, BToATm),
+				Distance, NearestA, NearestB, Normal);
 			FVec3 NearPointOnA = A.Min();
 			FVec3 SphereNearPointDir = (NearPointOnA - B.GetCenter()).GetSafeNormal();
 			FVec3 NearPointOnB = B.GetCenter() + SphereNearPointDir * B.GetRadius();
@@ -169,7 +214,11 @@ namespace ChaosTest
 			FAABB3 A(FVec3(-2, -2, -2), FVec3(4, 4, 4));
 			TSphere<FReal, 3> B(FVec3(0, 0, 0), 2);
 			FRigidTransform3 BToATm = FRigidTransform3(FVec3(8, 0, 0), FRotation3::FromAxisAngle(FVec3(0, 1, 0), FMath::DegreesToRadians(45)));	// Rotation won't affect contact depth, but does affect local contact position
-			EGJKDistanceResult Result = GJKDistance<FReal>(A, B, BToATm, Distance, NearestA, NearestB, Normal);
+			EGJKDistanceResult Result = GJKDistance<FReal>(
+				TGJKShape(A),
+				TGJKCoreShapeTransformed(B, BToATm),
+				GJKDistanceInitialV(A, B, BToATm),
+				Distance, NearestA, NearestB, Normal);
 			FVec3 NearPointOnA = FVec3(4, 0, 0);
 			FVec3 BPos = BToATm.TransformPositionNoScale(B.GetCenter());
 			FVec3 NearPointDir = (NearPointOnA - BPos).GetSafeNormal();
@@ -189,7 +238,11 @@ namespace ChaosTest
 			FAABB3 A(FVec3(-2, -2, -2), FVec3(4, 4, 4));
 			TSphere<FReal, 3> B(FVec3(0, 0, 0), 2);
 			FRigidTransform3 BToATm = FRigidTransform3(FVec3(-8, 0, 0), FRotation3::FromAxisAngle(FVec3(0, 1, 0), FMath::DegreesToRadians(45)));
-			EGJKDistanceResult Result = GJKDistance<FReal>(B, A, BToATm, Distance, NearestB, NearestA, Normal);
+			EGJKDistanceResult Result = GJKDistance<FReal>(
+				TGJKCoreShape(B),
+				TGJKShapeTransformed(A, BToATm),
+				GJKDistanceInitialV(B, A, BToATm),
+				Distance, NearestB, NearestA, Normal);
 			FVec3 NearPointOnA = FVec3(4, 0, 4);
 			FVec3 BPos = BToATm.InverseTransformPositionNoScale(B.GetCenter());
 			FVec3 NearPointDir = (NearPointOnA - BPos).GetSafeNormal();
@@ -213,7 +266,12 @@ namespace ChaosTest
 			bool bOverlap = GJKIntersection<FReal>(A, B, FRigidTransform3::Identity);
 			EXPECT_FALSE(bOverlap);
 
-			EGJKDistanceResult Result = GJKDistance<FReal>(A, B, FRigidTransform3::Identity, Distance, NearestA, NearestB, Normal);
+			const FRigidTransform3 BToATm = FRigidTransform3::Identity;
+			EGJKDistanceResult Result = GJKDistance<FReal>(
+				TGJKShape(A),
+				TGJKCoreShapeTransformed(B, BToATm),
+				GJKDistanceInitialV(A, B, BToATm),
+				Distance, NearestA, NearestB, Normal);
 			FVec3 NearPointOnA = FVec3(5, 0, 2);
 			FVec3 NearPointDir = (NearPointOnA - B.GetCenter()).GetSafeNormal();
 			FVec3 NearPointOnB = B.GetCenter() + NearPointDir * B.GetRadius();
@@ -244,7 +302,12 @@ namespace ChaosTest
 		{
 			FAABB3 A(FVec3(5, -2, -2), FVec3(8, 2, 2));
 			FCapsule B(FVec3(2, -2, 0), FVec3(2, 2, 0), 2);
-			EGJKDistanceResult Result = GJKDistance<FReal>(A, B, FRigidTransform3(FVec3(2, 0, 0), FRotation3::FromIdentity()), Distance, NearestA, NearestB, Normal);
+			const FRigidTransform3 BToATm = FRigidTransform3(FVec3(2, 0, 0), FRotation3::FromIdentity());
+			EGJKDistanceResult Result = GJKDistance<FReal>(
+				TGJKShape(A),
+				TGJKCoreShapeTransformed(B, BToATm),
+				GJKDistanceInitialV(A, B, BToATm),
+				Distance, NearestA, NearestB, Normal);
 			EXPECT_NE(Result, EGJKDistanceResult::Separated);
 		}
 
@@ -252,8 +315,12 @@ namespace ChaosTest
 		{
 			FAABB3 A(FVec3(5, -2, -2), FVec3(8, 2, 2));
 			FCapsule B(FVec3(2, 0, -1), FVec3(2, 0, 2), 2);
-			EGJKDistanceResult Result = GJKDistance<FReal>(A, B, FRigidTransform3::Identity, Distance, NearestA, NearestB, Normal);
-
+			const FRigidTransform3 BToATm = FRigidTransform3::Identity;
+			EGJKDistanceResult Result = GJKDistance<FReal>(
+				TGJKShape(A),
+				TGJKCoreShapeTransformed(B, BToATm),
+				GJKDistanceInitialV(A, B, BToATm),
+				Distance, NearestA, NearestB, Normal);
 			const FReal Tolerance = (FReal)2e-3;
 			EXPECT_EQ(Result, EGJKDistanceResult::Separated);
 			EXPECT_NEAR(Distance, (FReal)1, Tolerance);
@@ -271,7 +338,12 @@ namespace ChaosTest
 		{
 			FAABB3 A(FVec3(5, -2, -2), FVec3(8, 2, 2));
 			FCapsule B(FVec3(-2, 0, 3), FVec3(2, 0, -3), 2);
-			EGJKDistanceResult Result = GJKDistance<FReal>(A, B, FRigidTransform3::Identity, Distance, NearestA, NearestB, Normal);
+			const FRigidTransform3 BToATm = FRigidTransform3::Identity;
+			EGJKDistanceResult Result = GJKDistance<FReal>(
+				TGJKShape(A),
+				TGJKCoreShapeTransformed(B, BToATm),
+				GJKDistanceInitialV(A, B, BToATm),
+				Distance, NearestA, NearestB, Normal);
 			FVec3 ExpectedNearestA = FVec3(5, 0, -2);
 			FVec3 ExpectedDir = (ExpectedNearestA - B.GetX2()).GetSafeNormal();
 			FVec3 ExpectedNearestB = B.GetX2() + ExpectedDir * B.GetRadius();
@@ -291,7 +363,12 @@ namespace ChaosTest
 		{
 			FCapsule A(FVec3(4, 0, -1), FVec3(4, 0, -7), 1);
 			FAABB3 B(FVec3(-2, -2, -2), FVec3(2, 2, 2));
-			EGJKDistanceResult Result = GJKDistance<FReal>(A, B, FRigidTransform3::Identity, Distance, NearestA, NearestB, Normal);
+			const FRigidTransform3 BToATm = FRigidTransform3::Identity;
+			EGJKDistanceResult Result = GJKDistance<FReal>(
+				TGJKCoreShape(A),
+				TGJKShapeTransformed(B, BToATm),
+				GJKDistanceInitialV(A, B, BToATm),
+				Distance, NearestA, NearestB, Normal);
 			FVec3 ExpectedNearestA = FVec3(3, 0, (FReal)-1.5);
 			FVec3 ExpectedNearestB = FVec3(2, 0, (FReal)-1.5);
 
@@ -312,8 +389,12 @@ namespace ChaosTest
 		{
 			FCapsule A(FVec3(0, 0, -3), FVec3(0, 0, 3), 1);
 			FAABB3 B(FVec3(-2, -2, -2), FVec3(2, 2, 2));
-			FRigidTransform3 BToA = FRigidTransform3(FVec3(-4, 0, 4), FRotation3::FromIdentity());
-			EGJKDistanceResult Result = GJKDistance<FReal>(A, B, BToA, Distance, NearestA, NearestB, Normal);
+			FRigidTransform3 BToATm = FRigidTransform3(FVec3(-4, 0, 4), FRotation3::FromIdentity());
+			EGJKDistanceResult Result = GJKDistance<FReal>(
+				TGJKCoreShape(A),
+				TGJKShapeTransformed(B, BToATm),
+				GJKDistanceInitialV(A, B, BToATm),
+				Distance, NearestA, NearestB, Normal);
 			FVec3 ExpectedNearestA = FVec3(-1, 0, (FReal)2);
 			FVec3 ExpectedNearestB = FVec3(2, 0, (FReal)-2);
 
@@ -326,7 +407,7 @@ namespace ChaosTest
 			EXPECT_GT(NearestA.Z, (FReal)ExpectedNearestA.Z - (FReal)0.5 - Tolerance);
 			EXPECT_NEAR(NearestB.X, (FReal)ExpectedNearestB.X, Tolerance);
 			EXPECT_NEAR(NearestB.Y, (FReal)ExpectedNearestB.Y, Tolerance);
-			EXPECT_NEAR(NearestB.Z + BToA.GetTranslation().Z, (FReal)NearestA.Z, Tolerance);
+			EXPECT_NEAR(NearestB.Z + BToATm.GetTranslation().Z, (FReal)NearestA.Z, Tolerance);
 		}
 	}
 
@@ -351,7 +432,12 @@ namespace ChaosTest
 			FCapsule B(FVec3(-2, 0, 3), FVec3(2, 0, -3), 2);
 			FReal Epsilon = (FReal)1e-6;
 			int32 MaxIts = 5;
-			EGJKDistanceResult Result = GJKDistance<FReal>(A, B, FRigidTransform3::Identity, Distance, NearestA, NearestB, Normal, Epsilon, MaxIts);
+			const FRigidTransform3 BToATm = FRigidTransform3::Identity;
+			EGJKDistanceResult Result = GJKDistance<FReal>(
+				TGJKShape(A),
+				TGJKCoreShapeTransformed(B, BToATm),
+				GJKDistanceInitialV(A, B, BToATm),
+				Distance, NearestA, NearestB, Normal);
 			FVec3 ExpectedNearestA = FVec3(5, 0, -2);
 			FVec3 ExpectedDir = (ExpectedNearestA - B.GetX2()).GetSafeNormal();
 			FVec3 ExpectedNearestB = B.GetX2() + ExpectedDir * B.GetRadius();
diff --git a/Engine/Source/Runtime/Experimental/Chaos/Private/Chaos/Collision/SphereConvexContactPoint.cpp b/Engine/Source/Runtime/Experimental/Chaos/Private/Chaos/Collision/SphereConvexContactPoint.cpp
--- a/Engine/Source/Runtime/Experimental/Chaos/Private/Chaos/Collision/SphereConvexContactPoint.cpp
+++ b/Engine/Source/Runtime/Experimental/Chaos/Private/Chaos/Collision/SphereConvexContactPoint.cpp
@@ -44,104 +44,97 @@ namespace Chaos
 
 	// Use GJK (point to convex) to calculate separation.
 	// Fall back to plane testing if penetrating by more than Radius.
-	template<typename T_CONVEX>
-	FContactPoint ConvexSphereContactPointImpl(const T_CONVEX& Convex, const FRigidTransform3& ConvexTransform, const FImplicitSphere3& Sphere, const FRigidTransform3& SphereTransform)
+	template<typename ConvexType>
+	FContactPoint SphereConvexContactPointImpl(const FImplicitSphere3& Sphere, const ConvexType& Convex, const FRigidTransform3& SphereToConvexTransform)
 	{
-		FContactPoint ContactPoint;
 		if (Convex.NumPlanes() > 0)
 		{
-			const FRigidTransform3 SphereToConvexTransform = SphereTransform.GetRelativeTransform(ConvexTransform);
 			FVec3 PosConvex, PosSphere, NormalConvex;
 			FReal Phi;
 
 			// Run GJK to find separating distance if available
 			// NOTE: Sphere is treated as a point (its core shape), Convex margin is ignored so we are using the outer non-shrunken hull.
-			const EGJKDistanceResult GjkResult = GJKDistance(MakeGJKShape(Convex), MakeGJKCoreShape(Sphere), SphereToConvexTransform, Phi, PosConvex, PosSphere, NormalConvex);
+			// Sphere is also created in the space of the convex to eliminate the transform in the Support call (per GJK iteration).
+			const FVec3 SpherePosConvex = SphereToConvexTransform.TransformPosition(Sphere.GetCenter());
+			const FReal SphereRadius = Sphere.GetRadius();
+
+			// A point in the Minkowski Sum (A-B) in a likely direction
+			// @todo(chaos): we could do better and return the correct direction separation vector most of the time by selecting the face 
+			// most opposing the vector we have here (assuming we usually collide with a face rather than an edge). If we do this, 
+			// it should probably only be done when the convex is larger than the sphere or the convex has few planes.
+			const FVec3 InitialV = Convex.GetCenterOfMass() - SpherePosConvex;
+
+			const EGJKDistanceResult GjkResult = GJKDistance(
+				MakeGJKShape(Convex), 
+				FGJKSphere(SpherePosConvex, 0), 
+				InitialV, 
+				Phi, PosConvex, PosSphere, NormalConvex);
+				
 			bool bHaveResult = (GjkResult != EGJKDistanceResult::DeepContact);
 
-			// If GJK failed, core shapes are penetrating so find the minimum penetration vector
+			// If the sphere center is inside the convex, find the depenetration vector
 			if (!bHaveResult)
 			{
-				const FVec3 SpherePosConvex = SphereToConvexTransform.TransformPosition(Sphere.GetCenter());
-
-				FReal PointPhi;
-				bHaveResult = ConvexPointPenetrationVector(Convex, SpherePosConvex, NormalConvex, PointPhi);
-				if (bHaveResult)
-				{
-					PosConvex = SpherePosConvex - PointPhi * NormalConvex;
-					PosSphere = SphereToConvexTransform.InverseTransformVector(-Sphere.GetRadius() * NormalConvex);
-					Phi = PointPhi - Sphere.GetRadius();
-				}
+				bHaveResult = ConvexPointPenetrationVector(Convex, SpherePosConvex, NormalConvex, Phi);
 			}
 
 			// Build the contact point
 			if (bHaveResult)
 			{
-				ContactPoint.ShapeContactPoints[0] = PosConvex;
-				ContactPoint.ShapeContactPoints[1] = PosSphere;
-				ContactPoint.ShapeContactNormal = -NormalConvex;
+				// Results so far are all in convex-space and for a point rather than a sphere. Convert them.
+				PosConvex = SpherePosConvex - Phi * NormalConvex;
+				PosSphere = Sphere.GetCenter() + SphereToConvexTransform.InverseTransformVector(-SphereRadius * NormalConvex);
+				Phi = Phi - SphereRadius;
+
+				FContactPoint ContactPoint;
+				ContactPoint.ShapeContactPoints[0] = PosSphere;
+				ContactPoint.ShapeContactPoints[1] = PosConvex;
+				ContactPoint.ShapeContactNormal = NormalConvex;
 				ContactPoint.Phi = Phi;
+				return ContactPoint;
 			}
 		}
 
-		return ContactPoint;
-	}
-
-	template<typename T_CONVEX>
-	FContactPoint SphereConvexContactPointImpl(const FImplicitSphere3& Sphere, const FRigidTransform3& SphereTransform, const T_CONVEX& Convex, const FRigidTransform3& ConvexTransform)
-	{
-		return ConvexSphereContactPointImpl(Convex, ConvexTransform, Sphere, SphereTransform).SwapShapes();
-	}
-
-	FContactPoint SphereConvexContactPoint(const FImplicitSphere3& Sphere, const FRigidTransform3& SphereTransform, const FImplicitConvex3& Convex, const FRigidTransform3& ConvexTransform)
-	{
-		return SphereConvexContactPointImpl(Sphere, SphereTransform, Convex, ConvexTransform);
-	}
-
-	FContactPoint SphereConvexContactPoint(const FImplicitSphere3& Sphere, const FRigidTransform3& SphereTransform, const TImplicitObjectInstanced<FImplicitConvex3>& Convex, const FRigidTransform3& ConvexTransform)
-	{
-		return SphereConvexContactPointImpl(Sphere, SphereTransform, Convex, ConvexTransform);
-	}
-
-	FContactPoint SphereConvexContactPoint(const FImplicitSphere3& Sphere, const FRigidTransform3& SphereTransform, const TImplicitObjectScaled<FImplicitConvex3>& Convex, const FRigidTransform3& ConvexTransform)
-	{
-		return SphereConvexContactPointImpl(Sphere, SphereTransform, Convex, ConvexTransform);
+		return FContactPoint();
 	}
 
-
-
-	FContactPoint ConvexSphereContactPoint(const FImplicitConvex3& Convex, const FRigidTransform3& ConvexTransform, const FImplicitSphere3& Sphere, const FRigidTransform3& SphereTransform)
+	FContactPoint SphereConvexContactPoint(const FImplicitSphere3& Sphere, const FImplicitConvex3& Convex, const FRigidTransform3& SphereToConvexTransform)
 	{
-		return ConvexSphereContactPointImpl(Convex, ConvexTransform, Sphere, SphereTransform);
+		return SphereConvexContactPointImpl(Sphere, Convex, SphereToConvexTransform);
 	}
 
-	FContactPoint ConvexSphereContactPoint(const TImplicitObjectInstanced<FImplicitConvex3>& Convex, const FRigidTransform3& ConvexTransform, const FImplicitSphere3& Sphere, const FRigidTransform3& SphereTransform)
+	FContactPoint SphereConvexContactPoint(const FImplicitSphere3& Sphere, const TImplicitObjectInstanced<FImplicitConvex3>& Convex, const FRigidTransform3& SphereToConvexTransform)
 	{
-		return ConvexSphereContactPointImpl(Convex, ConvexTransform, Sphere, SphereTransform);
+		return SphereConvexContactPointImpl(Sphere, Convex, SphereToConvexTransform);
 	}
 
-	FContactPoint ConvexSphereContactPoint(const TImplicitObjectScaled<FImplicitConvex3>& Convex, const FRigidTransform3& ConvexTransform, const FImplicitSphere3& Sphere, const FRigidTransform3& SphereTransform)
+	FContactPoint SphereConvexContactPoint(const FImplicitSphere3& Sphere, const TImplicitObjectScaled<FImplicitConvex3>& Convex, const FRigidTransform3& SphereToConvexTransform)
 	{
-		return ConvexSphereContactPointImpl(Convex, ConvexTransform, Sphere, SphereTransform);
+		return SphereConvexContactPointImpl(Sphere, Convex, SphereToConvexTransform);
 	}
 
 
-
-	FContactPoint SphereConvexContactPoint(const FImplicitSphere3& Sphere, const FRigidTransform3& SphereTransform, const FImplicitObject& Object, const FRigidTransform3& ConvexTransform)
+	FContactPoint SphereConvexContactPoint(const FImplicitSphere3& Sphere, const FImplicitObject& Object, const FRigidTransform3& SphereToConvexTransform)
 	{
 		if (const TImplicitObjectInstanced<FImplicitConvex3>* InstancedConvex = Object.template GetObject<const TImplicitObjectInstanced<FImplicitConvex3>>())
 		{
-			return SphereConvexContactPoint(Sphere, SphereTransform, *InstancedConvex, ConvexTransform);
+			return SphereConvexContactPoint(Sphere, *InstancedConvex, SphereToConvexTransform);
 		}
 		else if (const TImplicitObjectScaled<FImplicitConvex3>* ScaledConvex = Object.template GetObject<const TImplicitObjectScaled<FImplicitConvex3>>())
 		{
-			return SphereConvexContactPoint(Sphere, SphereTransform, *ScaledConvex, ConvexTransform);
+			return SphereConvexContactPoint(Sphere, *ScaledConvex, SphereToConvexTransform);
 		}
 		else if (const FImplicitConvex3* Convex = Object.template GetObject<const FImplicitConvex3>())
 		{
-			return SphereConvexContactPoint(Sphere, SphereTransform, *Convex, ConvexTransform);
+			return SphereConvexContactPoint(Sphere, *Convex, SphereToConvexTransform);
 		}
 		return FContactPoint();
 	}
 
+	FContactPoint SphereConvexContactPoint(const FImplicitSphere3& Sphere, const FRigidTransform3& SphereTransform, const FImplicitObject& Object, const FRigidTransform3& ConvexTransform)
+	{
+		const FRigidTransform3 SphereToConvexTransform = SphereTransform.GetRelativeTransform(ConvexTransform);
+		return SphereConvexContactPoint(Sphere, Object, SphereToConvexTransform);
+	}
+
 }
diff --git a/Engine/Source/Runtime/Experimental/Chaos/Private/Chaos/CollisionOneShotManifoldsMiscShapes.cpp b/Engine/Source/Runtime/Experimental/Chaos/Private/Chaos/CollisionOneShotManifoldsMiscShapes.cpp
--- a/Engine/Source/Runtime/Experimental/Chaos/Private/Chaos/CollisionOneShotManifoldsMiscShapes.cpp
+++ b/Engine/Source/Runtime/Experimental/Chaos/Private/Chaos/CollisionOneShotManifoldsMiscShapes.cpp
@@ -244,73 +244,20 @@ namespace Chaos
 		template<typename ConvexType>
 		void ConstructSphereConvexManifoldImpl(const FImplicitSphere3& Sphere, const ConvexType& Convex, const FRigidTransform3& SphereToConvexTransform, const FReal CullDistance, TCArray<FContactPoint, 4>& ContactPoints)
 		{
-			// Transform the sphere into convex space
-			const FVec3 SpherePos = SphereToConvexTransform.TransformPositionNoScale(Sphere.GetCenter());
-			const FReal SphereRadius = Sphere.GetRadius();
-
-			// No margins for the convex, but treat the sphere as a point with a margin
-			FGJKSphereShape GJKSphere(SpherePos, SphereRadius);
-			TGJKShape<ConvexType> GJKConvex(Convex);
-
-			// GJK and EPA tolerances. See comments in GJKContactPointMargin
-			const FReal GJKEpsilon = Chaos_Collision_GJKEpsilon;
-			const FReal EPAEpsilon = Chaos_Collision_EPAEpsilon;
-			const FVec3 InitialDir = FVec3(-1,0,0);
-			FReal ClosestPenetration;
-			int32 ClosestVertexIndexSphere = INDEX_NONE, ClosestVertexIndexConvex = INDEX_NONE;
-			FReal ClosestSupportMaxDelta = FReal(0);
-
-			// Primary contact
-			// NOTE: swapped contact point order to match desired output order (Sphere, Convex)
-			FContactPoint ClosestContactPoint;
-			if (bChaos_Collision_UseGJK2)
-			{
-				GJKPenetrationSameSpace2(
-					GJKConvex, 
-					GJKSphere, 
-					ClosestPenetration,
-					ClosestContactPoint.ShapeContactPoints[1], 
-					ClosestContactPoint.ShapeContactPoints[0], 
-					ClosestContactPoint.ShapeContactNormal, 
-					ClosestVertexIndexConvex,
-					ClosestVertexIndexSphere,
-					ClosestSupportMaxDelta,
-					InitialDir,
-					GJKEpsilon, 
-					EPAEpsilon);
-			}
-			else
-			{
-				GJKPenetrationSameSpace(
-					GJKConvex,
-					GJKSphere,
-					ClosestPenetration,
-					ClosestContactPoint.ShapeContactPoints[1],
-					ClosestContactPoint.ShapeContactPoints[0],
-					ClosestContactPoint.ShapeContactNormal,
-					ClosestVertexIndexConvex,
-					ClosestVertexIndexSphere,
-					ClosestSupportMaxDelta,
-					InitialDir,
-					GJKEpsilon,
-					EPAEpsilon);
-			}
+			FContactPoint ClosestContactPoint = SphereConvexContactPoint(Sphere, Convex, SphereToConvexTransform);
 
 			// Stop now if beyond cull distance
-			const FReal ClosestPhi = -ClosestPenetration;
-			if (ClosestPhi > CullDistance)
+			if (ClosestContactPoint.Phi > CullDistance)
 			{
 				return;
 			}
 
 			// We always use the primary contact so add it to the output now
-			ClosestContactPoint.ShapeContactPoints[0] = SphereToConvexTransform.InverseTransformPositionNoScale(ClosestContactPoint.ShapeContactPoints[0]);
-			ClosestContactPoint.Phi = ClosestPhi;
-			ClosestContactPoint.FaceIndex = INDEX_NONE;
-			ClosestContactPoint.ContactType = EContactPointType::Unknown;
 			ContactPoints.Add(ClosestContactPoint);
 
 			// If the sphere is "large" compared to the convex add more points
+			const FVec3 SpherePos = SphereToConvexTransform.TransformPositionNoScale(Sphere.GetCenter());
+			const FReal SphereRadius = Sphere.GetRadius();
 			const FReal SpheerConvexManifoldSizeThreshold = FReal(1);
 			const FReal ConvexSize = Convex.BoundingBox().Extents().GetAbsMax();
 			if (SphereRadius > SpheerConvexManifoldSizeThreshold * ConvexSize)
diff --git a/Engine/Source/Runtime/Experimental/Chaos/Public/Chaos/Collision/SphereConvexContactPoint.h b/Engine/Source/Runtime/Experimental/Chaos/Public/Chaos/Collision/SphereConvexContactPoint.h
--- a/Engine/Source/Runtime/Experimental/Chaos/Public/Chaos/Collision/SphereConvexContactPoint.h
+++ b/Engine/Source/Runtime/Experimental/Chaos/Public/Chaos/Collision/SphereConvexContactPoint.h
@@ -11,12 +11,5 @@ namespace Chaos
 	class FContactPoint;
 
 	extern CHAOS_API FContactPoint SphereConvexContactPoint(const FImplicitSphere3& Sphere, const FRigidTransform3& SphereTransform, const FImplicitObject& Convex, const FRigidTransform3& ConvexTransform);
-
-	extern CHAOS_API FContactPoint SphereConvexContactPoint(const FImplicitSphere3& Sphere, const FRigidTransform3& SphereTransform, const FImplicitConvex3& Convex, const FRigidTransform3& ConvexTransform);
-	extern CHAOS_API FContactPoint SphereConvexContactPoint(const FImplicitSphere3& Sphere, const FRigidTransform3& SphereTransform, const TImplicitObjectInstanced<FImplicitConvex3>& Convex, const FRigidTransform3& ConvexTransform);
-	extern CHAOS_API FContactPoint SphereConvexContactPoint(const FImplicitSphere3& Sphere, const FRigidTransform3& SphereTransform, const TImplicitObjectScaled<FImplicitConvex3>& Convex, const FRigidTransform3& ConvexTransform);
-
-	extern CHAOS_API FContactPoint ConvexSphereContactPoint(const FImplicitConvex3& Convex, const FRigidTransform3& ConvexTransform, const FImplicitSphere3& Sphere, const FRigidTransform3& SphereTransform);
-	extern CHAOS_API FContactPoint ConvexSphereContactPoint(const TImplicitObjectInstanced<FImplicitConvex3>& Convex, const FRigidTransform3& ConvexTransform, const FImplicitSphere3& Sphere, const FRigidTransform3& SphereTransform);
-	extern CHAOS_API FContactPoint ConvexSphereContactPoint(const TImplicitObjectScaled<FImplicitConvex3>& Convex, const FRigidTransform3& ConvexTransform, const FImplicitSphere3& Sphere, const FRigidTransform3& SphereTransform);
+	extern CHAOS_API FContactPoint SphereConvexContactPoint(const FImplicitSphere3& Sphere, const FImplicitObject& Convex, const FRigidTransform3& SphereToConvexTransform);
 }
diff --git a/Engine/Source/Runtime/Experimental/Chaos/Public/Chaos/GJK.h b/Engine/Source/Runtime/Experimental/Chaos/Public/Chaos/GJK.h
--- a/Engine/Source/Runtime/Experimental/Chaos/Public/Chaos/GJK.h
+++ b/Engine/Source/Runtime/Experimental/Chaos/Public/Chaos/GJK.h
@@ -1828,13 +1828,19 @@ namespace Chaos
 
 
 	/**
-	 * Used by GJKDistance. It must return a vector in the Minkowski sum A - B. In principle this can be the vector of any point
-	 * in A to any point in B, but some choices will cause GJK to minimize faster (e.g., for two spheres, we can easily calculate
-	 * the actual separating vector and GJK will converge immediately).
+	 * Can be used to generate the initial support direction for use with GjkDistance. Returns a point on the Minkowski Sum
+	 * surface opposite to the direction of the supplied transform. This is usually a good guess for the initial direction
+	 * but it calls SupportCore on both shapes, and there are often faster alternatives if you know the type of shapes
+	 * you are dealing with (e.g., return the vectors between the centers of the two convex shapes).
+	 * 
+	 * If you do roll your own function, make sure that the vector returned is in or on the Minkowski sum (and don't just use a unit
+	 * vector along some direction for example) or GJKDistance may early-exit with an inaccurate result.
 	 */
 	template <typename T, typename TGeometryA, typename TGeometryB>
-	TVector<T, 3> GJKDistanceInitialV(const TGeometryA& A, T MarginA, const TGeometryB& B, T MarginB, const TRigidTransform<T, 3>& BToATM)
+	TVector<T, 3> GJKDistanceInitialV(const TGeometryA& A, const TGeometryB& B, const TRigidTransform<T, 3>& BToATM)
 	{
+		const T MarginA = A.GetMargin();
+		const T MarginB = B.GetMargin();
 		int32 VertexIndexA = INDEX_NONE, VertexIndexB = INDEX_NONE;
 		const TVec3<T> V = -BToATM.GetTranslation();
 		const TVector<T, 3> SupportA = A.SupportCore(-V, MarginA, nullptr, VertexIndexA);
@@ -1844,16 +1850,6 @@ namespace Chaos
 		return SupportA - SupportB;
 	}
 
-	/**
-	 * Used by GJKDistance. Specialization for sphere-sphere gives correct result immediately.
-	 */
-	template <typename T>
-	TVector<T, 3> GJKDistanceInitialV(const TSphere<T, 3>& A, const TSphere<T, 3>& B, const TRigidTransform<T, 3>& BToATM)
-	{
-		TVector<T, 3> Delta = A.GetCenter() - (B.GetCenter() + BToATM.GetTranslation());
-		return Delta;
-	}
-
 	// Status of a call to GJKDistance
 	enum class EGJKDistanceResult
 	{
@@ -1863,20 +1859,29 @@ namespace Chaos
 		// The shapes are overlapping by less than the net margin and all outputs have valid values (with a negative separation)
 		Contact,
 
-		// The shapes are overlapping by more than the net margin and all outputs are invalid
+		// The shapes are overlapping by more than the net margin and all outputs are invalid (unchanged from their input values)
 		DeepContact,
 	};
 
 	/**
-	 * Find the distance and nearest points on two convex geometries A and B.
-	 * All calculations are performed in the local-space of object A, and the transform from B-space to A-space must be provided.
+	 * Find the distance and nearest points on two convex geometries A and B, if they do not overlap.
+	 *
+	 * This is intended to be used with TGJKShape, TGJKCoreShape, TGJKShapeTransformed and TGJKCoreShapeTransformed
+	 * and generally you will have a TGJKShape or TGJKCoreShape as A, and a Transformed version for B. The transform in 
+	 * B should transform a vector in B space to a vector in A space as all calculations are performed in the local-
+	 * space of object A.
+	 *
+	 * You can also use this function when you have 2 shapes already in the same space to avoid transforms in the support
+	 * function(s). In this case you would use the non-transform versions of TGJKShape and would need to transform the 
+	 * results from the shared space to the space you desire.
+	 *
 	 * For algorithm see "A Fast and Robust GJK Implementation for Collision Detection of Convex Objects", Gino Van Deb Bergen, 1999.
 	 * @note This algorithm aborts if objects are overlapping and it does not initialize the out parameters.
 	 *
-	 * @param A The first object.
-	 * @param B The second object.
-	 * @param BToATM A transform taking vectors in B-space to A-space
+	 * @param A The first object (usually TGJKShape or TGJKCoreShape)
+	 * @param B The second object (usually TGJKShapeTransformed or TGJKCoreShapeTransformed)
 	 * @param B The second object.
+	 * @param InitialV  Starting support direction that must be in the Minkowski Sum. Use GJKDistanceInitialV() if unsure.
 	 * @param OutDistance if returns true, the minimum distance between A and B, otherwise not modified.
 	 * @param OutNearestA if returns true, the near point on A in local-space, otherwise not modified.
 	 * @param OutNearestB if returns true, the near point on B in local-space, otherwise not modified.
@@ -1884,22 +1889,20 @@ namespace Chaos
 	 * @param MaxIts A limit on the number of iterations. Results may be approximate if this is too low.
 	 * @return EGJKDistanceResult - see comments on the enum
 	 */
-	template <typename T, typename TGeometryA, typename TGeometryB>
-	EGJKDistanceResult GJKDistance(const TGeometryA& A, const TGeometryB& B, const TRigidTransform<T, 3>& BToATM, T& OutDistance, TVector<T, 3>& OutNearestA, TVector<T, 3>& OutNearestB, TVector<T, 3>& OutNormalA, const T Epsilon = (T)1e-3, const int32 MaxIts = 16)
+	template <typename T, typename GJKShapeTypeA, typename GJKShapeTypeB>
+	EGJKDistanceResult GJKDistance(const GJKShapeTypeA& A, const GJKShapeTypeB& B, const TVec3<T>& InitialV, T& OutDistance, TVec3<T>& OutNearestA, TVec3<T>& OutNearestB, TVec3<T>& OutNormalA, const T Epsilon = (T)1e-3, const int32 MaxIts = 16)
 	{
-		check(A.IsConvex() && B.IsConvex());
-
 		FSimplex SimplexIDs;
-		TVector<T, 3> Simplex[4], SimplexA[4], SimplexB[4];
+		TVec3<T> Simplex[4], SimplexA[4], SimplexB[4];
 		T Barycentric[4] = { -1, -1, -1, -1 };
 
-		const TRotation<T, 3> AToBRotation = BToATM.GetRotation().Inverse();
 		const T AMargin = A.GetMargin();
 		const T BMargin = B.GetMargin();
 		T Mu = 0;
 
-		// Select an initial vector in Minkowski(A - B)
-		TVector<T, 3> V = GJKDistanceInitialV(A, AMargin, B, BMargin, BToATM);
+		// Initial vector in Minkowski(A - B)
+		// NOTE: If InitialV is not in the Minkowski Sum we may incorrectly early-out in the bCloseEnough chaeck below
+		TVec3<T> V = InitialV;
 		T VLen = V.Size();
 		int32 VertexIndexA = INDEX_NONE, VertexIndexB = INDEX_NONE;
 
@@ -1909,13 +1912,11 @@ namespace Chaos
 			// Find a new point in A-B that is closer to the origin
 			// NOTE: we do not use support thickness here. Thickness is used when separating objects
 			// so that GJK can find a solution, but that can be added in a later step.
-			const TVector<T, 3> SupportA = A.SupportCore(-V, AMargin, nullptr, VertexIndexA);
-			const TVector<T, 3> VInB = AToBRotation * V;
-			const TVector<T, 3> SupportBLocal = B.SupportCore(VInB, BMargin, nullptr, VertexIndexB);
-			const TVector<T, 3> SupportB = BToATM.TransformPositionNoScale(SupportBLocal);
-			const TVector<T, 3> W = SupportA - SupportB;
+			const TVec3<T> SupportA = A.SupportCore(-V, AMargin, nullptr, VertexIndexA);
+			const TVec3<T> SupportB = B.SupportCore(V, BMargin, nullptr, VertexIndexB);
+			const TVec3<T> W = SupportA - SupportB;
 
-			T D = TVector<T, 3>::DotProduct(V, W) / VLen;
+			T D = TVec3<T>::DotProduct(V, W) / VLen;
 			Mu = FMath::Max(Mu, D);
 
 			// See if we are still making progress toward the origin
@@ -1930,13 +1931,13 @@ namespace Chaos
 				{
 					// Our initial guess of V was already the minimum separating vector
 					OutNearestA = SupportA;
-					OutNearestB = SupportBLocal;
+					OutNearestB = SupportB;
 				}
 				else
 				{
 					// The simplex vertices are the nearest point/line/face
-					OutNearestA = TVector<T, 3>(0, 0, 0);
-					OutNearestB = TVector<T, 3>(0, 0, 0);
+					OutNearestA = TVec3<T>(0, 0, 0);
+					OutNearestB = TVec3<T>(0, 0, 0);
 					for (int32 VertIndex = 0; VertIndex < SimplexIDs.NumVerts; ++VertIndex)
 					{
 						int32 WIndex = SimplexIDs[VertIndex];
@@ -1945,13 +1946,15 @@ namespace Chaos
 						OutNearestB += Barycentric[WIndex] * SimplexB[WIndex];
 					}
 				}
-				const TVector<T, 3> NormalA = -V / VLen;
-				const TVector<T, 3> NormalB = VInB / VLen;
+				const TVec3<T> NormalA = -V / VLen;
 				OutDistance = VLen - (AMargin + BMargin);
 				OutNearestA += AMargin * NormalA;
-				OutNearestB += BMargin * NormalB;
+				OutNearestB -= BMargin * NormalA;
 				OutNormalA = NormalA;
 
+				// NearestB should be in B-space but is currently in A-space
+				OutNearestB = B.InverseTransformPositionNoScale(OutNearestB);
+
 				return (OutDistance >= 0.0f) ? EGJKDistanceResult::Separated : EGJKDistanceResult::Contact;
 			}
 
@@ -1959,7 +1962,7 @@ namespace Chaos
 			SimplexIDs[SimplexIDs.NumVerts] = SimplexIDs.NumVerts;
 			Simplex[SimplexIDs.NumVerts] = W;
 			SimplexA[SimplexIDs.NumVerts] = SupportA;
-			SimplexB[SimplexIDs.NumVerts] = SupportBLocal;
+			SimplexB[SimplexIDs.NumVerts] = SupportB;
 			++SimplexIDs.NumVerts;
 
 			// Find the closest point to the origin on the simplex, and update the simplex to eliminate unnecessary vertices
@@ -1971,67 +1974,4 @@ namespace Chaos
 		return EGJKDistanceResult::DeepContact;
 	}
 
-
-	// Assumes objects are already intersecting, computes a minimum translation
-	// distance, deepest penetration positions on each body, and approximates
-	// a penetration normal and minimum translation distance.
-	//
-	// TODO: We want to re-visit how these functions work. Probably should be
-	// embedded in GJKOverlap and GJKRaycast so that secondary queries are unnecessary.
-	template <typename T, typename TGeometryA, typename TGeometryB>
-	bool GJKPenetrationTemp(const TGeometryA& A, const TGeometryB& B, const TRigidTransform<T, 3>& BToATM, TVector<T, 3>& OutPositionA, TVector<T, 3>& OutPositionB, TVector<T, 3>& OutNormal, T& OutDistance, const T ThicknessA = 0, const TVector<T, 3>& InitialDir = TVector<T, 3>(1, 0, 0), const T ThicknessB = 0, const T Epsilon = (T)1e-6, const int32 MaxIts = 16)
-	{
-		//
-		// TODO: General case for MTD determination.
-		//
-		ensure(false);
-		OutPositionA = TVector<T, 3>(0.f);
-		OutPositionB = TVector<T, 3>(0.f);
-		OutNormal = TVector<T, 3>(0.f, 0.f, 1.f);
-		OutDistance = 0.f;
-		return GJKIntersection(A, B, BToATM, ThicknessA, InitialDir, ThicknessB);
-	}
-
-	// Specialization for when getting MTD against a capsule.
-	template <typename T, typename TGeometryA>
-	bool GJKPenetrationTemp(const TGeometryA& A, const FCapsule& B, const TRigidTransform<T, 3>& BToATM, TVector<T, 3>& OutPositionA, TVector<T, 3>& OutPositionB, TVector<T, 3>& OutNormal, T& OutDistance, const T ThicknessA = 0, const TVector<T, 3>& InitialDir = TVector<T, 3>(1, 0, 0), const T ThicknessB = 0, const T Epsilon = (T)1e-6, const int32 MaxIts = 16)
-	{
-		T SegmentDistance;
-		const TSegment<T>& Segment = B.GetSegment();
-		const T MarginB = B.GetRadius();
-		TVector<T, 3> PositionBInB;
-		if (GJKDistance(A, Segment, BToATM, SegmentDistance, OutPositionA, PositionBInB, Epsilon, MaxIts))
-		{
-			OutPositionB = BToATM.TransformPosition(PositionBInB);
-			OutNormal
-				= ensure(SegmentDistance > TNumericLimits<T>::Min())
-				? (OutPositionB - OutPositionA) / SegmentDistance
-				: TVector<T, 3>(0.f, 0.f, 1.f);
-			OutPositionB -= OutNormal * MarginB;
-			OutDistance = SegmentDistance - MarginB;
-
-			if (OutDistance > 0.f)
-			{
-				// In this case, our distance calculation says we're not penetrating.
-				//
-				// TODO: check(false)! This shouldn't happen.
-				// It probably won't happen anymore if we warm-start GJKDistance
-				// with a polytope.
-				//
-				OutDistance = 0.f;
-				return false;
-			}
-
-			return true;
-		}
-		else
-		{
-			// TODO: Deep penetration - do EPA
-			ensure(false);
-			return true;
-		}
-
-		return false;
-	}
-
 }
\ No newline at end of file
diff --git a/Engine/Source/Runtime/Experimental/Chaos/Public/Chaos/GJKShape.h b/Engine/Source/Runtime/Experimental/Chaos/Public/Chaos/GJKShape.h
--- a/Engine/Source/Runtime/Experimental/Chaos/Public/Chaos/GJKShape.h
+++ b/Engine/Source/Runtime/Experimental/Chaos/Public/Chaos/GJKShape.h
@@ -1,6 +1,8 @@
 // Copyright Epic Games, Inc. All Rights Reserved.
 #pragma once
 
+#include "Chaos/Core.h"
+
 namespace Chaos
 {
 	/**
@@ -9,21 +11,29 @@ namespace Chaos
 	* and margins are used or not depending on the shape pair type involved.
 	*/
 
-	// Wraps an FImplicitObject  and provides the API required for GJK, 
-	// treating the shape as if it has zero margin. This means spheres will 
-	// be spheres, convexes will be the outer hull, etc.
+	// Wraps an shape object (could be FImplicitObject or some other type with the required API)
+	// and provides the API required for GJK, treating the shape as if it has zero margin. 
+	// This means spheres will be spheres (not points with padding), convexes will be the outer hull, etc.
 	//
-	// See also TGJKCoreShape
+	// See also TGJKShapeTransformed, TGJKCoreShape
 	//
-	// E.g., to use GJK of a sphere as a point against a marginless convex:
-	//		GJKDistance(TGJKCoreShape(MySphere), TGJKShape(MyConvex), ...);
+	// E.g., to use GJK between two Convex implicit objects
+	//		GJKDistance(TGJKShape(ConvexA), TGJKShapeTransformed(ConvexB, BToATransform), ...);
+	//
+	// E.g., to use GJK between a convex and a sphere treated as a point with padding
+	//		GJKDistance(TGJKShape(ConvexA), TGJKCoreShapeTransformed(SphereB, BToATransform), ...);
 	//
 	template<typename T_SHAPE>
 	struct TGJKShape
 	{
-		using FImplicitObjectType = T_SHAPE;
+		using FShapeType = T_SHAPE;
+
+		TGJKShape(const FShapeType& InShape) : Shape(InShape) {}
 
-		TGJKShape(const FImplicitObjectType& InShape) : Shape(InShape) {}
+		FVec3 InverseTransformPositionNoScale(const FVec3& V) const
+		{
+			return V;
+		}
 
 		FReal GetMargin() const
 		{
@@ -45,12 +55,64 @@ namespace Chaos
 			return FString::Printf(TEXT("TGJKShape: %s"), *Shape.ToString());
 		}
 
-		const FImplicitObjectType& Shape;
+		const FShapeType& Shape;
 	};
 
+	// Like TGJKShape but for the second shape if it has a transform relative to the first.
+	// 
+	// @param InTransform The shape transform (relative to the first shape when used with GJK)
+	// 
+	// See also TGJKShape, TGJKCoreShape
+	//
+	template<typename T_SHAPE>
+	struct TGJKShapeTransformed
+	{
+		using FShapeType = T_SHAPE;
+
+		TGJKShapeTransformed(const FShapeType& InShape, const FRigidTransform3& InTransform) 
+			: Transform(InTransform)
+			, Shape(InShape)
+		{}
+
+		const FRigidTransform3& GetTransform() const
+		{
+			return Transform;
+		}
+
+		FVec3 InverseTransformPositionNoScale(const FVec3& V) const
+		{
+			return Transform.InverseTransformPositionNoScale(V);
+		}
+
+		FReal GetMargin() const
+		{
+			return 0.0f;
+		}
 
-	// Wraps an FImplicitObject  and provides the API required for GJK, 
-	// treating the shape as if it has a reduced "core" shape with a 
+		FVec3 SupportCore(const FVec3 Dir, const FReal InMargin, FReal* OutSupportDelta, int32& VertexIndex) const
+		{
+			const FVec3 LocalDir = Transform.InverseTransformVectorNoScale(Dir);
+			const FVec3 LocalSupport = Shape.Support(LocalDir, InMargin, VertexIndex);
+			return Transform.TransformPositionNoScale(LocalSupport);
+		}
+
+		bool IsConvex() const
+		{
+			return Shape.IsConvex();
+		}
+
+		FString ToString() const
+		{
+			return FString::Printf(TEXT("TGJKShapeTransformed: %s"), *Shape.ToString());
+		}
+
+		FRigidTransform3 Transform;
+		const FShapeType& Shape;
+	};
+
+
+
+	// Like TGJKShape, but treats the shape as if it has a reduced "core" shape with a 
 	// margin suitable for collision detection where significant overlaps are likely.
 	// This means spheres will be points, convexes will be rounded shrunken hulls, etc.
 	//
@@ -62,18 +124,23 @@ namespace Chaos
 	template<typename T_SHAPE>
 	struct TGJKCoreShape
 	{
-		using FImplicitObjectType = T_SHAPE;
+		using FShapeType = T_SHAPE;
 
-		TGJKCoreShape(const FImplicitObjectType& InShape) 
+		TGJKCoreShape(const FShapeType& InShape)
 			: Shape(InShape)
 			, Margin(InShape.GetMargin())
 		{}
 
-		TGJKCoreShape(const FImplicitObjectType& InShape, const FReal InMargin)
+		TGJKCoreShape(const FShapeType& InShape, const FReal InMargin)
 			: Shape(InShape)
 			, Margin(InMargin)
 		{}
 
+		FVec3 InverseTransformPositionNoScale(const FVec3& V) const
+		{
+			return V;
+		}
+
 		FReal GetMargin() const
 		{
 			return Margin;
@@ -94,26 +161,101 @@ namespace Chaos
 			return FString::Printf(TEXT("TGJKCoreShape: %s, Margin: %f"), *Shape.ToString(), GetMargin());
 		}
 
-		const FImplicitObjectType& Shape;
+		const FShapeType& Shape;
+		const FReal Margin;
+	};
+
+	// Like TGJKCoreShape but for the second shape if it has a transform relative to the first.
+	// 
+	// @param InTransform The shape transform (relative to the first shape when used with GJK)
+	// 
+	// See also TGJKShapeTransformed, TGJKCoreShape
+	//
+	template<typename T_SHAPE>
+	struct TGJKCoreShapeTransformed
+	{
+		using FShapeType = T_SHAPE;
+
+		TGJKCoreShapeTransformed(const FShapeType& InShape, const FRigidTransform3& InTransform)
+			: Transform(InTransform)
+			, Shape(InShape)
+			, Margin(InShape.GetMargin())
+		{}
+
+		TGJKCoreShapeTransformed(const FShapeType& InShape, const FReal InMargin, const FRigidTransform3& InTransform)
+			: Transform(InTransform)
+			, Shape(InShape)
+			, Margin(InMargin)
+		{}
+
+		const FRigidTransform3& GetTransform() const
+		{
+			return Transform;
+		}
+
+		FVec3 InverseTransformPositionNoScale(const FVec3& V) const
+		{
+			return Transform.InverseTransformPositionNoScale(V);
+		}
+
+		FReal GetMargin() const
+		{
+			return Margin;
+		}
+
+		FVec3 SupportCore(const FVec3 Dir, const FReal InMargin, FReal* OutSupportDelta, int32& VertexIndex) const
+		{
+			const FVec3 LocalDir = Transform.InverseTransformVectorNoScale(Dir);
+			const FVec3 LocalSupport = Shape.SupportCore(LocalDir, InMargin, OutSupportDelta, VertexIndex);
+			return Transform.TransformPositionNoScale(LocalSupport);
+		}
+
+		bool IsConvex() const
+		{
+			return Shape.IsConvex();
+		}
+
+		FString ToString() const
+		{
+			return FString::Printf(TEXT("TGJKCoreShapeTransformed: %s"), *Shape.ToString());
+		}
+
+		FRigidTransform3 Transform;
+		const FShapeType& Shape;
 		const FReal Margin;
 	};
 
 	/**
-	 * @brief A sphere with minimal API for use in GJK/EPA
-	 * Equivalent to TGJKCoreShape<FImplicitSphere3> without any indirection.
+	 * A sphere with minimal API for use in GJK/EPA.
+	 * Equivalent to TGJKCoreShape<FImplicitSphere3> but avoids calls to the implicit
 	*/
-	class FGJKSphereShape
+	class FGJKSphere
 	{
 	public:
-		FGJKSphereShape(const FVec3 InPos, const FReal InRadius) 
-			: Pos(InPos), Radius(InRadius) 
+		FGJKSphere(const FVec3 InCenter, const FReal InRadius)
+			: Center(InCenter), Radius(InRadius)
+		{
+		}
+
+		inline const FVec3& GetCenter() const
+		{
+			return Center;
+		}
+
+		inline FReal GetRadius() const
+		{
+			return Radius;
+		}
+
+		inline FVec3 InverseTransformPositionNoScale(const FVec3& V) const
 		{
+			return V;
 		}
 
 		inline const FVec3& SupportCore(const FVec3& Direction, const FReal InMargin, FReal* MaxMarginDelta, int32& VertexIndex) const
 		{
 			VertexIndex = 0;
-			return Pos;
+			return Center;
 		}
 
 		inline FReal GetMargin() const
@@ -123,11 +265,11 @@ namespace Chaos
 
 		FString ToString() const
 		{
-			return FString::Printf(TEXT("FGJKSphereShape: Pos: [%f, %f, %f], Radius: %f"), Pos.X, Pos.Y, Pos.Z, Radius);
+			return FString::Printf(TEXT("FGJKSphere: Center: [%f, %f, %f], Radius: %f"), Center.X, Center.Y, Center.Z, Radius);
 		}
 
 	private:
-		FVec3 Pos;
+		FVec3 Center;
 		FReal Radius;
 	};
 
