Chaos terminology pass

diff --git a/Engine/Plugins/Experimental/ChaosCloth/Source/Chaos/Private/ChaosCloth/ChaosClothingSimulationSolver.cpp b/Engine/Plugins/Experimental/ChaosCloth/Source/Chaos/Private/ChaosCloth/ChaosClothingSimulationSolver.cpp
--- a/Engine/Plugins/Experimental/ChaosCloth/Source/Chaos/Private/ChaosCloth/ChaosClothingSimulationSolver.cpp
+++ b/Engine/Plugins/Experimental/ChaosCloth/Source/Chaos/Private/ChaosCloth/ChaosClothingSimulationSolver.cpp
@@ -528,8 +528,7 @@ void FClothingSimulationSolver::SetParticleMassUniform(int32 Offset, FRealSingle
 		Particles.M(Index) = Vertices.Contains(Index) ? (Softs::FSolverReal)UniformMass : (Softs::FSolverReal)0.;
 	}
 
-	// Clamp and enslave
-	ParticleMassClampAndEnslave(Offset, Size, (Softs::FSolverReal)MinPerParticleMass, KinematicPredicate);
+	ParticleMassClampAndKinematicStateUpdate(Offset, Size, (Softs::FSolverReal)MinPerParticleMass, KinematicPredicate);
 }
 
 void FClothingSimulationSolver::SetParticleMassFromTotalMass(int32 Offset, FRealSingle TotalMass, FRealSingle MinPerParticleMass, const FTriangleMesh& Mesh, const TFunctionRef<bool(int32)>& KinematicPredicate)
@@ -546,8 +545,7 @@ void FClothingSimulationSolver::SetParticleMassFromTotalMass(int32 Offset, FReal
 	// Update mass from mesh and density
 	ParticleMassUpdateDensity(Mesh, Density);
 
-	// Clamp and enslave
-	ParticleMassClampAndEnslave(Offset, Size, (Softs::FSolverReal)MinPerParticleMass, KinematicPredicate);
+	ParticleMassClampAndKinematicStateUpdate(Offset, Size, (Softs::FSolverReal)MinPerParticleMass, KinematicPredicate);
 }
 
 void FClothingSimulationSolver::SetParticleMassFromDensity(int32 Offset, FRealSingle Density, FRealSingle MinPerParticleMass, const FTriangleMesh& Mesh, const TFunctionRef<bool(int32)>& KinematicPredicate)
@@ -564,8 +562,7 @@ void FClothingSimulationSolver::SetParticleMassFromDensity(int32 Offset, FRealSi
 	// Update mass from mesh and density
 	ParticleMassUpdateDensity(Mesh, DensityScaled);
 
-	// Clamp and enslave
-	ParticleMassClampAndEnslave(Offset, Size, (Softs::FSolverReal)MinPerParticleMass, KinematicPredicate);
+	ParticleMassClampAndKinematicStateUpdate(Offset, Size, (Softs::FSolverReal)MinPerParticleMass, KinematicPredicate);
 }
 
 void FClothingSimulationSolver::SetReferenceVelocityScale(
@@ -651,7 +648,7 @@ void FClothingSimulationSolver::ParticleMassUpdateDensity(const FTriangleMesh& M
 	UE_LOG(LogChaosCloth, Verbose, TEXT("Total mass: %f, "), TotalMass);
 }
 
-void FClothingSimulationSolver::ParticleMassClampAndEnslave(int32 Offset, int32 Size, Softs::FSolverReal MinPerParticleMass, const TFunctionRef<bool(int32)>& KinematicPredicate)
+void FClothingSimulationSolver::ParticleMassClampAndKinematicStateUpdate(int32 Offset, int32 Size, Softs::FSolverReal MinPerParticleMass, const TFunctionRef<bool(int32)>& KinematicPredicate)
 {
 	Softs::FSolverParticles& Particles = Evolution->Particles();
 	for (int32 Index = Offset; Index < Offset + Size; ++Index)
diff --git a/Engine/Plugins/Experimental/ChaosCloth/Source/Chaos/Public/ChaosCloth/ChaosClothingSimulationSolver.h b/Engine/Plugins/Experimental/ChaosCloth/Source/Chaos/Public/ChaosCloth/ChaosClothingSimulationSolver.h
--- a/Engine/Plugins/Experimental/ChaosCloth/Source/Chaos/Public/ChaosCloth/ChaosClothingSimulationSolver.h
+++ b/Engine/Plugins/Experimental/ChaosCloth/Source/Chaos/Public/ChaosCloth/ChaosClothingSimulationSolver.h
@@ -199,7 +199,7 @@ namespace Chaos
 		void ApplyPreSimulationTransforms();
 		Softs::FSolverReal SetParticleMassPerArea(int32 Offset, int32 Size, const FTriangleMesh& Mesh);
 		void ParticleMassUpdateDensity(const FTriangleMesh& Mesh, Softs::FSolverReal Density);
-		void ParticleMassClampAndEnslave(int32 Offset, int32 Size, Softs::FSolverReal MinPerParticleMass, const TFunctionRef<bool(int32)>& KinematicPredicate);
+		void ParticleMassClampAndKinematicStateUpdate(int32 Offset, int32 Size, Softs::FSolverReal MinPerParticleMass, const TFunctionRef<bool(int32)>& KinematicPredicate);
 
 		// Update the solver field forces/velocities at the particles location
 		void UpdateSolverField();
diff --git a/Engine/Source/Programs/HeadlessChaos/Private/HeadlessChaosTestRewind.cpp b/Engine/Source/Programs/HeadlessChaos/Private/HeadlessChaosTestRewind.cpp
--- a/Engine/Source/Programs/HeadlessChaos/Private/HeadlessChaosTestRewind.cpp
+++ b/Engine/Source/Programs/HeadlessChaos/Private/HeadlessChaosTestRewind.cpp
@@ -2245,7 +2245,7 @@ namespace ChaosTest {
 		});
 	}
 
-	GTEST_TEST(AllTraits, RewindTest_ResimFallingObjectWithTeleportAsSlave)
+	GTEST_TEST(AllTraits, RewindTest_ResimFallingObjectWithTeleportAsFollower)
 	{
 		TRewindHelper::TestDynamicSphere([](auto* Solver, FReal SimDt, int32 Optimization, auto Proxy, auto Sphere)
 		{
@@ -2256,7 +2256,7 @@ namespace ChaosTest {
 				Solver->GetEvolution()->GetGravityForces().SetAcceleration(FVec3(0, 0, -1));
 				Particle.SetGravityEnabled(true);
 				Particle.SetX(FVec3(0, 0, 100));
-				Particle.SetResimType(EResimType::ResimAsSlave);
+				Particle.SetResimType(EResimType::ResimAsFollower);
 
 				for (int Step = 0; Step <= LastGameStep; ++Step)
 				{
@@ -2294,7 +2294,7 @@ namespace ChaosTest {
 				else
 				{
 #if REWIND_DESYNC
-					//we'll see the teleport automatically because ResimAsSlave
+					//we'll see the teleport automatically because ResimAsFollower
 					//but it's done by solver so before tick teleport is not known
 					EXPECT_NEAR(Particle.X()[2], ExpectedXZ, 1e-4);
 					EXPECT_NEAR(Particle.V()[2], ExpectedVZ, 1e-4);
@@ -2921,7 +2921,7 @@ namespace ChaosTest {
 		}
 	}
 
-	GTEST_TEST(AllTraits, RewindTest_ResimAsSlave)
+	GTEST_TEST(AllTraits, RewindTest_ResimAsFollower)
 	{
 #if REWIND_DESYNC
 		for (int Optimization = 0; Optimization < 2; ++Optimization)
@@ -2959,7 +2959,7 @@ namespace ChaosTest {
 				Dynamic.SetGravityEnabled(false);
 				Dynamic.SetV(FVec3(0, 0, -1));
 				Dynamic.SetObjectState(EObjectStateType::Dynamic);
-				Dynamic.SetResimType(EResimType::ResimAsSlave);
+				Dynamic.SetResimType(EResimType::ResimAsFollower);
 
 				Kinematic.SetX(FVec3(0, 0, 0));
 
@@ -2990,14 +2990,14 @@ namespace ChaosTest {
 
 			for (int Step = RewindStep; Step <= LastStep; ++Step)
 			{
-				//Resim but dynamic will take old path since it's marked as ResimAsSlave
+				//Resim but dynamic will take old path since it's marked as ResimAsFollower
 				TickSolverHelper(Solver);
 
 				EXPECT_VECTOR_FLOAT_EQ(Dynamic.X(), Xs[Step]);
 			}
 
 #if REWIND_DESYNC
-			//slave so dynamic in sync, kinematic desync
+			// follower - so dynamic in sync, kinematic desync
 			const TArray<FDesyncedParticleInfo> DesyncedParticles = RewindData->ComputeDesyncInfo();
 			EXPECT_EQ(DesyncedParticles.Num(), 1);
 			EXPECT_EQ(DesyncedParticles[0].MostDesynced, ESyncState::HardDesync);
@@ -3107,7 +3107,7 @@ namespace ChaosTest {
 		}
 	}
 
-	GTEST_TEST(AllTraits, DISABLED_RewindTest_ResimAsSlaveFallIgnoreCollision)
+	GTEST_TEST(AllTraits, DISABLED_RewindTest_ResimAsFollowerFallIgnoreCollision)
 	{
 		for (int Optimization = 0; Optimization < 2; ++Optimization)
 		{
@@ -3143,7 +3143,7 @@ namespace ChaosTest {
 				Dynamic.SetGravityEnabled(false);
 				Dynamic.SetV(FVec3(0, 0, -1));
 				Dynamic.SetObjectState(EObjectStateType::Dynamic);
-				Dynamic.SetResimType(EResimType::ResimAsSlave);
+				Dynamic.SetResimType(EResimType::ResimAsFollower);
 
 				Kinematic.SetX(FVec3(0, 0, -1000));
 
@@ -3175,7 +3175,7 @@ namespace ChaosTest {
 
 			for (int Step = RewindStep; Step <= LastStep; ++Step)
 			{
-				//Resim ignores collision since it's ResimAsSlave
+				//Resim ignores collision since it's ResimAsFollower
 				TickSolverHelper(Solver);
 
 				EXPECT_VECTOR_FLOAT_EQ(Dynamic.X(), Xs[Step]);
@@ -3186,7 +3186,7 @@ namespace ChaosTest {
 			EXPECT_LE(Dynamic.X()[2], 6);
 
 #if REWIND_DESYNC
-			//dynamic slave so only kinematic desyncs
+			//dynamic follower so only kinematic desyncs
 			const TArray<FDesyncedParticleInfo> DesyncedParticles = RewindData->ComputeDesyncInfo();
 			EXPECT_EQ(DesyncedParticles.Num(), 1);
 			EXPECT_EQ(DesyncedParticles[0].MostDesynced, ESyncState::HardDesync);
@@ -3197,7 +3197,7 @@ namespace ChaosTest {
 		}
 	}
 
-	GTEST_TEST(AllTraits, RewindTest_ResimAsSlaveWithForces)
+	GTEST_TEST(AllTraits, RewindTest_ResimAsFollowerWithForces)
 	{
 #if REWIND_DESYNC
 		for (int Optimization = 0; Optimization < 2; ++Optimization)
@@ -3214,37 +3214,37 @@ namespace ChaosTest {
 
 			// Make particles
 			auto FullSimProxy = FSingleParticlePhysicsProxy::Create(Chaos::FPBDRigidParticle::CreateParticle());
-			auto SlaveSimProxy = FSingleParticlePhysicsProxy::Create(Chaos::FPBDRigidParticle::CreateParticle());
+			auto FollowerSimProxy = FSingleParticlePhysicsProxy::Create(Chaos::FPBDRigidParticle::CreateParticle());
 			const int32 LastStep = 11;
 			TArray<FVec3> Xs;
 
 			{
 				auto& FullSim = FullSimProxy->GetGameThreadAPI();
-				auto& SlaveSim = SlaveSimProxy->GetGameThreadAPI();
+				auto& FollowerSim = FollowerSimProxy->GetGameThreadAPI();
 
 				FullSim.SetGeometry(Box);
 				FullSim.SetGravityEnabled(false);
 				Solver->RegisterObject(FullSimProxy);
 
-				SlaveSim.SetGeometry(Box);
-				SlaveSim.SetGravityEnabled(false);
-				Solver->RegisterObject(SlaveSimProxy);
+				FollowerSim.SetGeometry(Box);
+				FollowerSim.SetGravityEnabled(false);
+				Solver->RegisterObject(FollowerSimProxy);
 
 				FullSim.SetX(FVec3(0, 0, 20));
 				FullSim.SetObjectState(EObjectStateType::Dynamic);
 				FullSim.SetM(1);
 				FullSim.SetInvM(1);
 
-				SlaveSim.SetX(FVec3(0, 0, 0));
-				SlaveSim.SetResimType(EResimType::ResimAsSlave);
-				SlaveSim.SetM(1);
-				SlaveSim.SetInvM(1);
+				FollowerSim.SetX(FVec3(0, 0, 0));
+				FollowerSim.SetResimType(EResimType::ResimAsFollower);
+				FollowerSim.SetM(1);
+				FollowerSim.SetInvM(1);
 
-				ChaosTest::SetParticleSimDataToCollide({ FullSimProxy->GetParticle_LowLevel(),SlaveSimProxy->GetParticle_LowLevel() });
+				ChaosTest::SetParticleSimDataToCollide({ FullSimProxy->GetParticle_LowLevel(),FollowerSimProxy->GetParticle_LowLevel() });
 
 				for (int Step = 0; Step <= LastStep; ++Step)
 				{
-					SlaveSim.SetLinearImpulse(FVec3(0, 0, 0.5));
+					FollowerSim.SetLinearImpulse(FVec3(0, 0, 0.5));
 					TickSolverHelper(Solver);
 					Xs.Add(FullSim.X());
 				}
@@ -3253,21 +3253,21 @@ namespace ChaosTest {
 			FPhysicsThreadContextScope Scope(true);
 			const int RewindStep = 5;
 			auto& FullSim = *FullSimProxy->GetPhysicsThreadAPI();
-			auto& SlaveSim = *SlaveSimProxy->GetPhysicsThreadAPI();
+			auto& FollowerSim = *FollowerSimProxy->GetPhysicsThreadAPI();
 
 			FRewindData* RewindData = Solver->GetRewindData();
 			EXPECT_TRUE(RewindData->RewindToFrame(RewindStep));
 
 			for (int Step = RewindStep; Step <= LastStep; ++Step)
 			{
-				//resim - slave sim should have its impulses automatically added thus moving FullSim in the exact same way
+				//resim - follower sim should have its impulses automatically added thus moving FullSim in the exact same way
 				TickSolverHelper(Solver);
 
 				EXPECT_VECTOR_FLOAT_EQ(FullSim.X(), Xs[Step]);
 			}
 
 #if REWIND_DESYNC
-			//slave so no desync
+			//follower so no desync
 			const TArray<FDesyncedParticleInfo> DesyncedParticles = RewindData->ComputeDesyncInfo();
 			EXPECT_EQ(DesyncedParticles.Num(), 0);
 #endif
@@ -3277,7 +3277,7 @@ namespace ChaosTest {
 #endif
 	}
 
-	GTEST_TEST(AllTraits, RewindTest_ResimAsSlaveWokenUp)
+	GTEST_TEST(AllTraits, RewindTest_ResimAsFollowerWokenUp)
 	{
 #if REWIND_DESYNC
 		for (int Optimization = 0; Optimization < 2; ++Optimization)
@@ -3315,13 +3315,13 @@ namespace ChaosTest {
 				ImpulsedObj.SetX(FVec3(0, 0, 20));
 				ImpulsedObj.SetM(1);
 				ImpulsedObj.SetInvM(1);
-				ImpulsedObj.SetResimType(EResimType::ResimAsSlave);
+				ImpulsedObj.SetResimType(EResimType::ResimAsFollower);
 				ImpulsedObj.SetObjectState(EObjectStateType::Sleeping);
 
 				HitObj.SetX(FVec3(0, 0, 0));
 				HitObj.SetM(1);
 				HitObj.SetInvM(1);
-				HitObj.SetResimType(EResimType::ResimAsSlave);
+				HitObj.SetResimType(EResimType::ResimAsFollower);
 				HitObj.SetObjectState(EObjectStateType::Sleeping);
 
 
@@ -3357,7 +3357,7 @@ namespace ChaosTest {
 			}
 
 #if REWIND_DESYNC
-			//slave so no desync
+			//follower so no desync
 			const TArray<FDesyncedParticleInfo> DesyncedParticles = RewindData->ComputeDesyncInfo();
 			EXPECT_EQ(DesyncedParticles.Num(), 0);
 #endif
@@ -3367,7 +3367,7 @@ namespace ChaosTest {
 #endif
 	}
 
-	GTEST_TEST(AllTraits, RewindTest_ResimAsSlaveWokenUpNoHistory)
+	GTEST_TEST(AllTraits, RewindTest_ResimAsFollowerWokenUpNoHistory)
 	{
 #if REWIND_DESYNC
 		for (int Optimization = 0; Optimization < 2; ++Optimization)
@@ -3410,7 +3410,7 @@ namespace ChaosTest {
 				HitObj.SetX(FVec3(0, 0, 0));
 				HitObj.SetM(1);
 				HitObj.SetInvM(1);
-				HitObj.SetResimType(EResimType::ResimAsSlave);
+				HitObj.SetResimType(EResimType::ResimAsFollower);
 				HitObj.SetObjectState(EObjectStateType::Sleeping);
 
 
@@ -3441,12 +3441,12 @@ namespace ChaosTest {
 
 				TickSolverHelper(Solver);
 
-				//even though there's now a different collision in the sim, the final result of slave is the same as before
+				//even though there's now a different collision in the sim, the final result of follower is the same as before
 				EXPECT_VECTOR_FLOAT_EQ(HitObj.X(), Xs[Step]);
 			}
 
 #if REWIND_DESYNC
-			//only desync non-slave
+			//only desync non-follower
 			const TArray<FDesyncedParticleInfo> DesyncedParticles = RewindData->ComputeDesyncInfo();
 			EXPECT_EQ(DesyncedParticles.Num(), 1);
 			EXPECT_EQ(DesyncedParticles[0].MostDesynced, ESyncState::HardDesync);
diff --git a/Engine/Source/Runtime/ClothingSystemRuntimeInterface/Public/ClothingSimulationInterface.h b/Engine/Source/Runtime/ClothingSystemRuntimeInterface/Public/ClothingSimulationInterface.h
--- a/Engine/Source/Runtime/ClothingSystemRuntimeInterface/Public/ClothingSimulationInterface.h
+++ b/Engine/Source/Runtime/ClothingSystemRuntimeInterface/Public/ClothingSimulationInterface.h
@@ -109,7 +109,7 @@ protected:
 	 * Fill FClothSimulData map for the clothing simulation. Should fill a map pair per-actor 
 	 * @param OutData - The simulation data to write to
 	 * @param InOwnerComponent - the component that owns the simulation
-	 * @param InOverrideComponent - An override component if bound to a master pose component
+	 * @param InOverrideComponent - An override component if bound to a leader pose component
 	 */
 	virtual void GetSimulationData(TMap<int32, FClothSimulData>& OutData, USkeletalMeshComponent* InOwnerComponent, USkinnedMeshComponent* InOverrideComponent) const = 0;
 
diff --git a/Engine/Source/Runtime/Engine/Private/PhysicsEngine/ImmediatePhysics/ImmediatePhysicsChaos/ImmediatePhysicsSimulation_Chaos.cpp b/Engine/Source/Runtime/Engine/Private/PhysicsEngine/ImmediatePhysics/ImmediatePhysicsChaos/ImmediatePhysicsSimulation_Chaos.cpp
--- a/Engine/Source/Runtime/Engine/Private/PhysicsEngine/ImmediatePhysics/ImmediatePhysicsChaos/ImmediatePhysicsSimulation_Chaos.cpp
+++ b/Engine/Source/Runtime/Engine/Private/PhysicsEngine/ImmediatePhysics/ImmediatePhysicsChaos/ImmediatePhysicsSimulation_Chaos.cpp
@@ -729,12 +729,12 @@ namespace ImmediatePhysics_Chaos
 		Implementation->SimulationSpace.AngularVelocity = AngularVel;
 	}
 
-	void FSimulation::SetSimulationSpaceSettings(const FReal MasterAlpha, const FVector& ExternalLinearEtherDrag)
+	void FSimulation::SetSimulationSpaceSettings(const FReal Alpha, const FVector& ExternalLinearEtherDrag)
 	{
 		using namespace Chaos;
 
 		FSimulationSpaceSettings& SimSpaceSettings = Implementation->Evolution.GetSimulationSpaceSettings();
-		SimSpaceSettings.MasterAlpha = MasterAlpha;
+		SimSpaceSettings.Alpha = Alpha;
 		SimSpaceSettings.ExternalLinearEtherDrag = ExternalLinearEtherDrag;
 		SimSpaceSettings.CentrifugalAlpha = ChaosImmediate_Evolution_SimSpaceCentrifugalAlpha;
 		SimSpaceSettings.CoriolisAlpha = ChaosImmediate_Evolution_SimSpaceCoriolisAlpha;
diff --git a/Engine/Source/Runtime/Engine/Public/Physics/ImmediatePhysics/ImmediatePhysicsChaos/ImmediatePhysicsSimulation_Chaos.h b/Engine/Source/Runtime/Engine/Public/Physics/ImmediatePhysics/ImmediatePhysicsChaos/ImmediatePhysicsSimulation_Chaos.h
--- a/Engine/Source/Runtime/Engine/Public/Physics/ImmediatePhysics/ImmediatePhysicsChaos/ImmediatePhysicsSimulation_Chaos.h
+++ b/Engine/Source/Runtime/Engine/Public/Physics/ImmediatePhysics/ImmediatePhysicsChaos/ImmediatePhysicsSimulation_Chaos.h
@@ -73,7 +73,7 @@ namespace ImmediatePhysics_Chaos
 			const FVector& AngularAcc);
 
 		void SetSimulationSpaceSettings(
-			const FReal MasterAlpha, 
+			const FReal Alpha, 
 			const FVector& ExternalLinearEtherDrag);
 
 		/** Set settings. Invalid (negative) values with leave that value unchanged from defaults */
diff --git a/Engine/Source/Runtime/Experimental/Chaos/Private/Chaos/Evolution/PBDMinEvolution.cpp b/Engine/Source/Runtime/Experimental/Chaos/Private/Chaos/Evolution/PBDMinEvolution.cpp
--- a/Engine/Source/Runtime/Experimental/Chaos/Private/Chaos/Evolution/PBDMinEvolution.cpp
+++ b/Engine/Source/Runtime/Experimental/Chaos/Private/Chaos/Evolution/PBDMinEvolution.cpp
@@ -226,7 +226,7 @@ namespace Chaos
 		FVec3 SpaceW = FVec3(0);	// Angular Velocity
 		FVec3 SpaceA = FVec3(0);	// Acceleration
 		FVec3 SpaceB = FVec3(0);	// Angular Acceleration
-		if (SimulationSpaceSettings.MasterAlpha > 0.0f)
+		if (SimulationSpaceSettings.Alpha > 0.0f)
 		{
 			SpaceV = SimulationSpace.Transform.InverseTransformVector(SimulationSpace.LinearVelocity);
 			SpaceW = SimulationSpace.Transform.InverseTransformVector(SimulationSpace.AngularVelocity);
@@ -259,7 +259,7 @@ namespace Chaos
 
 				// Moving and accelerating simulation frame
 				// https://en.wikipedia.org/wiki/Rotating_reference_frame
-				if (SimulationSpaceSettings.MasterAlpha > 0.0f)
+				if (SimulationSpaceSettings.Alpha > 0.0f)
 				{
 					const FVec3 CoriolisAcc = SimulationSpaceSettings.CoriolisAlpha * 2.0f * FVec3::CrossProduct(SpaceW, Particle.V());
 					const FVec3 CentrifugalAcc = SimulationSpaceSettings.CentrifugalAlpha * FVec3::CrossProduct(SpaceW, FVec3::CrossProduct(SpaceW, XCoM));
@@ -267,10 +267,10 @@ namespace Chaos
 					const FVec3 LinearAcc = SimulationSpaceSettings.LinearAccelerationAlpha * SpaceA;
 					const FVec3 AngularAcc = SimulationSpaceSettings.AngularAccelerationAlpha * SpaceB;
 					const FVec3 LinearDragAcc = SimulationSpaceSettings.ExternalLinearEtherDrag * SpaceV;
-					DV -= SimulationSpaceSettings.MasterAlpha * (LinearAcc + LinearDragAcc + CoriolisAcc + CentrifugalAcc + EulerAcc) * Dt;
-					DW -= SimulationSpaceSettings.MasterAlpha * AngularAcc * Dt;
-					TargetV = -SimulationSpaceSettings.MasterAlpha * SimulationSpaceSettings.LinearVelocityAlpha * SpaceV;
-					TargetW = -SimulationSpaceSettings.MasterAlpha * SimulationSpaceSettings.AngularVelocityAlpha * SpaceW;
+					DV -= SimulationSpaceSettings.Alpha * (LinearAcc + LinearDragAcc + CoriolisAcc + CentrifugalAcc + EulerAcc) * Dt;
+					DW -= SimulationSpaceSettings.Alpha * AngularAcc * Dt;
+					TargetV = -SimulationSpaceSettings.Alpha * SimulationSpaceSettings.LinearVelocityAlpha * SpaceV;
+					TargetW = -SimulationSpaceSettings.Alpha * SimulationSpaceSettings.AngularVelocityAlpha * SpaceW;
 				}
 
 				// New velocity
diff --git a/Engine/Source/Runtime/Experimental/Chaos/Private/Chaos/Island/IslandGraph.cpp b/Engine/Source/Runtime/Experimental/Chaos/Private/Chaos/Island/IslandGraph.cpp
--- a/Engine/Source/Runtime/Experimental/Chaos/Private/Chaos/Island/IslandGraph.cpp
+++ b/Engine/Source/Runtime/Experimental/Chaos/Private/Chaos/Island/IslandGraph.cpp
@@ -76,20 +76,20 @@ void FIslandGraph<NodeType, EdgeType, IslandType, OwnerType>::UpdateNode(const i
 		if (bValidNode && !bWasValidNode)
 		{
 			// @todo(chaos): we could just use the GraphNode.NodeIslands here if it were maintained internally
-			int32 MasterIsland = INDEX_NONE;
+			int32 PrimaryIsland = INDEX_NONE;
 			for (int32& EdgeIndex : GraphNode.NodeEdges)
 			{
 				FGraphEdge& GraphEdge = GraphEdges[EdgeIndex];
 
 				const bool bIsEdgeMoving = IsEdgeMoving(EdgeIndex);
-				ParentIslands(MasterIsland, GraphEdge.IslandIndex, bIsEdgeMoving);
+				ParentIslands(PrimaryIsland, GraphEdge.IslandIndex, bIsEdgeMoving);
 			
-				MasterIsland = GraphEdge.IslandIndex;
+				PrimaryIsland = GraphEdge.IslandIndex;
 			}
 
 			// Put the valid node into one of the islands - they will be merged anyway so doesn't matter which
 			// If we did not have an island (no edges), one will be assigned later
-			GraphNode.IslandIndex = MasterIsland;
+			GraphNode.IslandIndex = PrimaryIsland;
 		}
 
 		// If we are changing to invalid (Dynamic/Sleeping->Kinematic) wake the island
diff --git a/Engine/Source/Runtime/Experimental/Chaos/Private/Chaos/PhysicalMaterials.cpp b/Engine/Source/Runtime/Experimental/Chaos/Private/Chaos/PhysicalMaterials.cpp
--- a/Engine/Source/Runtime/Experimental/Chaos/Private/Chaos/PhysicalMaterials.cpp
+++ b/Engine/Source/Runtime/Experimental/Chaos/Private/Chaos/PhysicalMaterials.cpp
@@ -90,10 +90,20 @@ namespace Chaos
 
 	const Chaos::THandleArray<FChaosPhysicsMaterial>& FPhysicalMaterialManager::GetMasterMaterials_External() const
 	{
-		return Materials;
+		return GetPrimaryMaterials_External();
 	}
 
 	const Chaos::THandleArray<FChaosPhysicsMaterialMask>& FPhysicalMaterialManager::GetMasterMaterialMasks_External() const
+	{
+		return GetPrimaryMaterialMasks_External();
+	}
+	
+	const Chaos::THandleArray<FChaosPhysicsMaterial>& FPhysicalMaterialManager::GetPrimaryMaterials_External() const
+	{
+		return Materials;
+	}
+
+	const Chaos::THandleArray<FChaosPhysicsMaterialMask>& FPhysicalMaterialManager::GetPrimaryMaterialMasks_External() const
 	{
 		return MaterialMasks;
 	}
diff --git a/Engine/Source/Runtime/Experimental/Chaos/Private/ChaosSolversModule.cpp b/Engine/Source/Runtime/Experimental/Chaos/Private/ChaosSolversModule.cpp
--- a/Engine/Source/Runtime/Experimental/Chaos/Private/ChaosSolversModule.cpp
+++ b/Engine/Source/Runtime/Experimental/Chaos/Private/ChaosSolversModule.cpp
@@ -205,14 +205,14 @@ Chaos::FPBDRigidsSolver* FChaosSolversModule::CreateSolver(UObject* InOwner, Cha
 	NewSolver->SetDebugName(NewDebugName);
 #endif
 
-	// Set up the material lists on the new solver, copying from the current master list
+	// Set up the material lists on the new solver, copying from the current primary list
 	{
 		FPhysicalMaterialManager& Manager =	Chaos::FPhysicalMaterialManager::Get();
 		FPhysicsSceneGuardScopedWrite ScopedWrite(NewSolver->GetExternalDataLock_External());
-		NewSolver->QueryMaterials_External = Manager.GetMasterMaterials_External();
-		NewSolver->QueryMaterialMasks_External = Manager.GetMasterMaterialMasks_External();
-		NewSolver->SimMaterials = Manager.GetMasterMaterials_External();
-		NewSolver->SimMaterialMasks = Manager.GetMasterMaterialMasks_External();
+		NewSolver->QueryMaterials_External = Manager.GetPrimaryMaterials_External();
+		NewSolver->QueryMaterialMasks_External = Manager.GetPrimaryMaterialMasks_External();
+		NewSolver->SimMaterials = Manager.GetPrimaryMaterials_External();
+		NewSolver->SimMaterialMasks = Manager.GetPrimaryMaterialMasks_External();
 	}
 
 	return NewSolver;
diff --git a/Engine/Source/Runtime/Experimental/Chaos/Private/RewindData.cpp b/Engine/Source/Runtime/Experimental/Chaos/Private/RewindData.cpp
--- a/Engine/Source/Runtime/Experimental/Chaos/Private/RewindData.cpp
+++ b/Engine/Source/Runtime/Experimental/Chaos/Private/RewindData.cpp
@@ -161,26 +161,26 @@ bool FRewindData::RewindToFrame(int32 Frame)
 	FFrameAndPhase RewindFrameAndPhase{ Frame, FFrameAndPhase::PostPushData };
 	FFrameAndPhase CurFrameAndPhase{ CurFrame, FFrameAndPhase::PrePushData };
 
-	auto RewindHelper = [RewindFrameAndPhase, CurFrameAndPhase, this](auto Obj, bool bResimAsSlave, auto& Property, const auto& RewindFunc) -> bool
+	auto RewindHelper = [RewindFrameAndPhase, CurFrameAndPhase, this](auto Obj, bool bResimAsFollower, auto& Property, const auto& RewindFunc) -> bool
 	{
-		if (bResimAsSlave)
-		{
-	//If we're rewinding a particle that doesn't need to save head (resim as slave never checks for desync so we don't care about head)
-		if (auto Val = Property.Read(RewindFrameAndPhase, PropertiesPool))
+		if (bResimAsFollower)
 		{
+			//If we're rewinding a particle that doesn't need to save head (resim as follower never checks for desync so we don't care about head)
+			if (auto Val = Property.Read(RewindFrameAndPhase, PropertiesPool))
+			{
 				RewindFunc(Obj, *Val);
 			}
 		}
 		else
-	{
-			//If we're rewinding an object that needs to save head (during resim when we get back to latest frame and phase we need to check for desync)
-		if (!Property.IsClean(RewindFrameAndPhase))
 		{
+			//If we're rewinding an object that needs to save head (during resim when we get back to latest frame and phase we need to check for desync)
+			if (!Property.IsClean(RewindFrameAndPhase))
+			{
 				CopyDataFromObject(Property.WriteAccessMonotonic(CurFrameAndPhase, PropertiesPool), *Obj);
 				RewindFunc(Obj, *Property.Read(RewindFrameAndPhase, PropertiesPool));
 
-			return true;
-		}
+				return true;
+			}
 		}
 
 		return false;
@@ -195,16 +195,16 @@ bool FRewindData::RewindToFrame(int32 Frame)
 		
 		FGeometryParticleStateBase& History = DirtyParticleInfo.AddFrame(CurFrame);	//non-const in case we need to record what's at head for a rewind (CurFrame has already been increased to the next frame)
 
-		const bool bResimAsSlave = DirtyParticleInfo.bResimAsSlave;
-		bool bAnyChange = RewindHelper(PTParticle, bResimAsSlave, History.ParticlePositionRotation, [](auto Particle, const auto& Data) {Particle->SetXR(Data); });
-		bAnyChange |= RewindHelper(PTParticle, bResimAsSlave, History.NonFrequentData, [](auto Particle, const auto& Data) {Particle->SetNonFrequentData(Data); });
-		bAnyChange |= RewindHelper(PTParticle->CastToKinematicParticle(), bResimAsSlave, History.Velocities, [](auto Particle, const auto& Data) {Particle->SetVelocities(Data); });
-		bAnyChange |= RewindHelper(PTParticle->CastToKinematicParticle(), bResimAsSlave, History.KinematicTarget, [](auto Particle, const auto& Data) {Particle->SetKinematicTarget(Data); });
-		bAnyChange |= RewindHelper(PTParticle->CastToRigidParticle(),bResimAsSlave,  History.Dynamics, [](auto Particle, const auto& Data) {Particle->SetDynamics(Data); });
-		bAnyChange |= RewindHelper(PTParticle->CastToRigidParticle(),bResimAsSlave,  History.DynamicsMisc, [Evolution = Solver->GetEvolution()](auto Particle, const auto& Data) {Particle->SetDynamicMisc(Data, *Evolution); });
-		bAnyChange |= RewindHelper(PTParticle->CastToRigidParticle(),bResimAsSlave,  History.MassProps, [](auto Particle, const auto& Data) {Particle->SetMassProps(Data); });
+		const bool bResimAsFollower = DirtyParticleInfo.bResimAsFollower;
+		bool bAnyChange = RewindHelper(PTParticle, bResimAsFollower, History.ParticlePositionRotation, [](auto Particle, const auto& Data) {Particle->SetXR(Data); });
+		bAnyChange |= RewindHelper(PTParticle, bResimAsFollower, History.NonFrequentData, [](auto Particle, const auto& Data) {Particle->SetNonFrequentData(Data); });
+		bAnyChange |= RewindHelper(PTParticle->CastToKinematicParticle(), bResimAsFollower, History.Velocities, [](auto Particle, const auto& Data) {Particle->SetVelocities(Data); });
+		bAnyChange |= RewindHelper(PTParticle->CastToKinematicParticle(), bResimAsFollower, History.KinematicTarget, [](auto Particle, const auto& Data) {Particle->SetKinematicTarget(Data); });
+		bAnyChange |= RewindHelper(PTParticle->CastToRigidParticle(),bResimAsFollower,  History.Dynamics, [](auto Particle, const auto& Data) {Particle->SetDynamics(Data); });
+		bAnyChange |= RewindHelper(PTParticle->CastToRigidParticle(),bResimAsFollower,  History.DynamicsMisc, [Evolution = Solver->GetEvolution()](auto Particle, const auto& Data) {Particle->SetDynamicMisc(Data, *Evolution); });
+		bAnyChange |= RewindHelper(PTParticle->CastToRigidParticle(),bResimAsFollower,  History.MassProps, [](auto Particle, const auto& Data) {Particle->SetMassProps(Data); });
 
-		if (!bResimAsSlave)
+		if (!bResimAsFollower)
 		{
 			if (bAnyChange)
 			{
@@ -275,14 +275,15 @@ void FRewindData::FinishFrame()
 		auto FinishHelper = [this, FutureFrame](auto& DirtyObjs)
 		{
 			for (auto& Info : DirtyObjs)
-		{
-			if (Info.bResimAsSlave)
 			{
-				//resim as slave means always in sync and no cleanup needed
-				continue;
-			}
+				if (Info.bResimAsFollower)
+				{
+					//resim as follower means always in sync and no cleanup needed
+					continue;
+				}
+
+				auto& Handle = *Info.GetObjectPtr();
 
-			auto& Handle = *Info.GetObjectPtr();
 				if (Handle.ResimType() == EResimType::FullResim)
 				{
 					if (IsFinalResim())
@@ -370,9 +371,9 @@ void FRewindData::AdvanceFrameImp(IResimCacheBase* ResimCache)
 		{
 
 			auto Handle = Info.GetObjectPtr();
-			Info.bResimAsSlave = Handle->ResimType() == EResimType::ResimAsSlave;
+			Info.bResimAsFollower = Handle->ResimType() == EResimType::ResimAsFollower;
 
-			if (IsResim() && !Info.bResimAsSlave)
+			if (IsResim() && !Info.bResimAsFollower)
 			{
 					DesyncIfNecessary</*bSkipDynamics=*/false>(Info, FrameAndPhase);
 			}
diff --git a/Engine/Source/Runtime/Experimental/Chaos/Public/Chaos/Evolution/SimulationSpace.h b/Engine/Source/Runtime/Experimental/Chaos/Public/Chaos/Evolution/SimulationSpace.h
--- a/Engine/Source/Runtime/Experimental/Chaos/Public/Chaos/Evolution/SimulationSpace.h
+++ b/Engine/Source/Runtime/Experimental/Chaos/Public/Chaos/Evolution/SimulationSpace.h
@@ -35,7 +35,7 @@ namespace Chaos
 	 * should get passed into the simulation.
 	 * If all the Alpha values are set to 1, the simulation
 	 * should be the same as a world-space simulation.
-	 * Usually you would only change MasterAlpha, but it
+	 * Usually you would only change Alpha, but it
 	 * is also possible to manipulate the elements of the
 	 * phantom forces (like centrifugal force) if you know
 	 * what you are doing.
@@ -47,6 +47,7 @@ namespace Chaos
 	public:
 		FSimulationSpaceSettings()
 			: MasterAlpha(0)
+			, Alpha(0)
 			, LinearAccelerationAlpha(1)
 			, CoriolisAlpha(1)
 			, CentrifugalAlpha(1)
@@ -58,9 +59,12 @@ namespace Chaos
 		{
 		}
 
-		// Global multipler on the effects of simulation space movement
+		UE_DEPRECATED(5.1, TEXT("MasterAlpha is deprecated - please use Alpha"))
 		FReal MasterAlpha;
 
+		// Global multipler on the effects of simulation space movement
+		FReal Alpha;
+
 		// How much of the simulation frame's linear acceleration to pass onto the particles
 		FReal LinearAccelerationAlpha;
 
diff --git a/Engine/Source/Runtime/Experimental/Chaos/Public/Chaos/GeometryParticles.h b/Engine/Source/Runtime/Experimental/Chaos/Public/Chaos/GeometryParticles.h
--- a/Engine/Source/Runtime/Experimental/Chaos/Public/Chaos/GeometryParticles.h
+++ b/Engine/Source/Runtime/Experimental/Chaos/Public/Chaos/GeometryParticles.h
@@ -417,9 +417,10 @@ namespace Chaos
 
 	enum class EResimType : uint8
 	{
-		FullResim,	//fully re-run simulation and keep results (any forces must be applied again)
+		FullResim = 0,	//fully re-run simulation and keep results (any forces must be applied again)
 		//ResimWithPrevForces, //use previous forces and keep results (UNIMPLEMENTED)
-		ResimAsSlave //use previous forces and snap to previous results regardless of variation - used to push other objects away
+		ResimAsSlave UE_DEPRECATED(5.1, "EResimType::ResimAsSlave is deprecated, please use EResimType::ResimAsFollower") = 1,
+		ResimAsFollower = 1 //use previous forces and snap to previous results regardless of variation - used to push other objects away
 		//ResimAsKinematic //treat as kinematic (UNIMPLEMENTED)
 	};
 	
diff --git a/Engine/Source/Runtime/Experimental/Chaos/Public/Chaos/PhysicalMaterials.h b/Engine/Source/Runtime/Experimental/Chaos/Public/Chaos/PhysicalMaterials.h
--- a/Engine/Source/Runtime/Experimental/Chaos/Public/Chaos/PhysicalMaterials.h
+++ b/Engine/Source/Runtime/Experimental/Chaos/Public/Chaos/PhysicalMaterials.h
@@ -134,10 +134,15 @@ namespace Chaos
 		void UpdateMaterial(FMaterialHandle InHandle);
 		void UpdateMaterialMask(FMaterialMaskHandle InHandle);
 
-		/** Gets the internal list of master materials representing the current user state of the material data */
+		/** Gets the internal list of primary materials representing the current user state of the material data */
+		UE_DEPRECATED(5.1, "GetMasterMaterials_External is deprecated, please use GetPrimaryMaterials_External instead")
 		const THandleArray<FChaosPhysicsMaterial>& GetMasterMaterials_External() const;
+		UE_DEPRECATED(5.1, "GetMasterMaterialMasks_External is deprecated, please use GetPrimaryMaterialMasks_External instead")
 		const THandleArray<FChaosPhysicsMaterialMask>& GetMasterMaterialMasks_External() const;
 
+		const THandleArray<FChaosPhysicsMaterial>& GetPrimaryMaterials_External() const;
+		const THandleArray<FChaosPhysicsMaterialMask>& GetPrimaryMaterialMasks_External() const;
+
 		/** Events */
 		FOnMaterialUpdated OnMaterialUpdated;
 		FOnMaterialCreated OnMaterialCreated;
diff --git a/Engine/Source/Runtime/Experimental/Chaos/Public/ChaosSolversModule.h b/Engine/Source/Runtime/Experimental/Chaos/Public/ChaosSolversModule.h
--- a/Engine/Source/Runtime/Experimental/Chaos/Public/ChaosSolversModule.h
+++ b/Engine/Source/Runtime/Experimental/Chaos/Public/ChaosSolversModule.h
@@ -155,7 +155,7 @@ public:
 	 */
 	void DestroySolver(Chaos::FPhysicsSolverBase* InState);
 
-	/** Retrieve the master list of solvers. This contains all owned, unowned and standalone solvers */
+	/** Retrieve the list of all extant solvers. This contains all owned, unowned and standalone solvers */
 	const TArray<Chaos::FPhysicsSolverBase*>& GetAllSolvers() const;
 
 	/**
diff --git a/Engine/Source/Runtime/Experimental/Chaos/Public/GeometryCollection/ManagedArrayCollection.h b/Engine/Source/Runtime/Experimental/Chaos/Public/GeometryCollection/ManagedArrayCollection.h
--- a/Engine/Source/Runtime/Experimental/Chaos/Public/GeometryCollection/ManagedArrayCollection.h
+++ b/Engine/Source/Runtime/Experimental/Chaos/Public/GeometryCollection/ManagedArrayCollection.h
@@ -409,7 +409,7 @@ public:
 	void RemoveDependencyFor(FName Group);
 
 	/**
-	* Copy an attribute. Will perform an implicit group sync. Attribute must exist in both MasterCollection and this collection
+	* Copy an attribute. Will perform an implicit group sync. Attribute must exist in both InCollection and this collection
 	* @param Name - The name of the attribute
 	* @param Group - The group that manages the attribute
 	*/
diff --git a/Engine/Source/Runtime/Experimental/Chaos/Public/RewindData.h b/Engine/Source/Runtime/Experimental/Chaos/Public/RewindData.h
--- a/Engine/Source/Runtime/Experimental/Chaos/Public/RewindData.h
+++ b/Engine/Source/Runtime/Experimental/Chaos/Public/RewindData.h
@@ -377,13 +377,6 @@ struct FPerShapeDataStateBase
 	//helper functions for shape API
 	template <typename TParticle>
 	static const FCollisionFilterData& GetQueryData(const FPerShapeDataStateBase* State, const TParticle& Particle, int32 ShapeIdx) { return State && State->CollisionData.IsSet() ? State->CollisionData.Read().QueryData : Particle.ShapesArray()[ShapeIdx]->GetQueryData(); }
-	/*const FCollisionFilterData& GetSimData() const { return CollisionData.Read().SimData; }
-	
-	TSerializablePtr<FImplicitObject> GetGeometry() const { return Geometry; }
-	const TArray<FMaterialHandle>& GetMaterials() const { return Materials.Read().Materials; }
-	const TArray<FMaterialMaskHandle>& GetMaterialMasks() const { return Materials.Read().MaterialMasks; }
-	const TArray<uint32>& GetMaterialMaskMaps() const { return Materials.Read().MaterialMaskMaps; }
-	const TArray<FMaterialHandle>& GetMaterialMaskMapMaterials() const { return Materials.Read().MaterialMaskMapMaterials; }*/
 };
 
 class FPerShapeDataState
@@ -1068,7 +1061,9 @@ private:
 		int32 DirtyDynamics = INDEX_NONE;	//Only used by particles, indicates the dirty properties was written to.
 		int32 LastDirtyFrame;	//Track how recently this was made dirty
 		int32 InitializedOnStep = INDEX_NONE;	//if not INDEX_NONE, it indicates we saw initialization during rewind history window
-		bool bResimAsSlave = true;	//Indicates the particle will always resim in the exact same way from game thread data
+		UE_DEPRECATED(5.1, TEXT("bResimAsSlave is deprecated - please use bResimAsFollower"))
+		bool bResimAsSlave = true;
+		bool bResimAsFollower = true;	//Indicates the particle will always resim in the exact same way from game thread data
 
 		TDirtyObjectInfo(FDirtyPropertiesPool& InPropertiesPool, TObj& InObj, const int32 CurFrame, const int32 NumFrames)
 			: History(NumFrames)
@@ -1084,7 +1079,7 @@ private:
 			, PropertiesPool(Other.PropertiesPool)
 			, LastDirtyFrame(Other.LastDirtyFrame)
 			, InitializedOnStep(Other.InitializedOnStep)
-			, bResimAsSlave(Other.bResimAsSlave)
+			, bResimAsFollower(Other.bResimAsFollower)
 		{
 			Other.PropertiesPool = nullptr;
 		}
diff --git a/Engine/Source/Runtime/PhysicsCore/Private/ChaosEngineInterface.cpp b/Engine/Source/Runtime/PhysicsCore/Private/ChaosEngineInterface.cpp
--- a/Engine/Source/Runtime/PhysicsCore/Private/ChaosEngineInterface.cpp
+++ b/Engine/Source/Runtime/PhysicsCore/Private/ChaosEngineInterface.cpp
@@ -1714,7 +1714,7 @@ void FChaosEngineInterface::CreateActor(const FActorCreationParams& InParams,FPh
 	if(InParams.bStatic)
 	{
 		Particle = FGeometryParticle::CreateParticle();
-		Particle->SetResimType(EResimType::ResimAsSlave);
+		Particle->SetResimType(EResimType::ResimAsFollower);
 	}
 	else
 	{
@@ -1734,7 +1734,7 @@ void FChaosEngineInterface::CreateActor(const FActorCreationParams& InParams,FPh
 		} else
 		{
 			Rigid->SetObjectState(EObjectStateType::Kinematic);
-			Rigid->SetResimType(EResimType::ResimAsSlave);	//for now kinematics are never changed during resim
+			Rigid->SetResimType(EResimType::ResimAsFollower);	//for now kinematics are never changed during resim
 		}
 		//Particle.Reset(Rigid.Release());
 		Particle = MoveTemp(Rigid);
