Chaos : Deprecate remove on fracture feature
- update headless chaos unit test

diff --git a/Engine/Source/Programs/HeadlessChaos/Private/GeometryCollection/GeometryCollectionTestClustering.cpp b/Engine/Source/Programs/HeadlessChaos/Private/GeometryCollection/GeometryCollectionTestClustering.cpp
--- a/Engine/Source/Programs/HeadlessChaos/Private/GeometryCollection/GeometryCollectionTestClustering.cpp
+++ b/Engine/Source/Programs/HeadlessChaos/Private/GeometryCollection/GeometryCollectionTestClustering.cpp
@@ -2110,86 +2110,6 @@ namespace GeometryCollectionTest
 		}
 	}
 
-	
-	GTEST_TEST(AllTraits, DISABLED_GeometryCollection_RigidBodies_ClusterTest_RemoveOnFracture)
-	{
-		// Disabled as remove on fracture currently unimplemented for geometry collections. Potentially this should be deleted entirely.
-
-		FFramework UnitTest;
-
-		TSharedPtr<FGeometryCollection> RestCollection = GeometryCollection::MakeCubeElement(FTransform(FQuat::MakeFromEuler(FVector(0, 0, 0.)), FVector(0, -10, 10)), FVector(1.0));
-		RestCollection->AppendGeometry(*GeometryCollection::MakeCubeElement(FTransform(FQuat::MakeFromEuler(FVector(0, 0, 0.)), FVector(0, 10, 10)), FVector(1.0)));
-		EXPECT_EQ(RestCollection->Transform.Num(), 2);
-
-		// this transform should have a zero scale after the simulation has run to the point of fracture
-		RestCollection->SetFlags(1, FGeometryCollection::FS_RemoveOnFracture);
-
-		FGeometryCollectionClusteringUtility::ClusterAllBonesUnderNewRoot(RestCollection.Get());
-		EXPECT_EQ(RestCollection->Transform.Num(), 3);
-		RestCollection->Transform[2] = FTransform(FQuat::MakeFromEuler(FVector(90.f, 0, 0.)), FVector(0, 0, 40));
-
-		CreationParameters Params;
-		Params.RestCollection = RestCollection;
-		Params.DynamicState = EObjectStateTypeEnum::Chaos_Object_Dynamic;
-		Params.ImplicitType = EImplicitTypeEnum::Chaos_Implicit_Box;
-		Params.CollisionType = ECollisionTypeEnum::Chaos_Surface_Volumetric;
-		Params.Simulating = true;
-		Params.EnableClustering = true;
-		Params.DamageThreshold = { 0.1f };		
-		Params.RemoveOnFractureEnabled = true;
-		FGeometryCollectionWrapper* Collection = TNewSimulationObject<GeometryType::GeometryCollectionWithSuppliedRestCollection>::Init(Params)->template As<FGeometryCollectionWrapper>();
-
-		FRadialFalloff * FalloffField = new FRadialFalloff();
-		FalloffField->Magnitude = 10.5;
-		FalloffField->Radius = 100.0;
-		FalloffField->Position = FVector(0.0, 0.0, 0.0);
-		FalloffField->Falloff = EFieldFalloffType::Field_FallOff_None;
-
-		TSharedPtr<FGeometryDynamicCollection> DynamicCollection = Collection->DynamicCollection;
-		DynamicCollection->ModifyAttribute<int32>("DynamicState", FGeometryCollection::TransformGroup)[1] = (uint8)EObjectStateTypeEnum::Chaos_Object_Kinematic;
-
-		UnitTest.AddSimulationObject(Collection);
-		UnitTest.Initialize();
-
-		auto& Clustering = UnitTest.Solver->GetEvolution()->GetRigidClustering();
-		const auto& ClusterMap = Clustering.GetChildrenMap();
-
-		TManagedArray<FTransform>& Transform = Collection->DynamicCollection->Transform;
-		FReal StartingRigidDistance = (Transform[1].GetTranslation() - Transform[0].GetTranslation()).Size(), CurrentRigidDistance = 0.f;
-		
-		// #todo: is this even used?
-		/*
-		Chaos::TArrayCollectionArray<FReal>& InternalStrain = Clustering.GetStrainArray();
-		*/
-
-		FName TargetName = GetFieldPhysicsName(EFieldPhysicsType::Field_ExternalClusterStrain);
-		FFieldSystemCommand Command(TargetName, FalloffField->NewCopy());
-		FFieldSystemMetaDataProcessingResolution* ResolutionData = new FFieldSystemMetaDataProcessingResolution(EFieldResolutionType::Field_Resolution_Maximum);
-		Command.MetaData.Add(FFieldSystemMetaData::EMetaType::ECommandData_ProcessingResolution, TUniquePtr< FFieldSystemMetaDataProcessingResolution >(ResolutionData));
-		UnitTest.Solver->GetPerSolverField().AddTransientCommand(Command);
-
-		FVector Scale = Transform[1].GetScale3D();
-
-		EXPECT_NEAR(Scale.X, 1.0f, SMALL_NUMBER);
-		EXPECT_NEAR(Scale.Y, 1.0f, SMALL_NUMBER);
-		EXPECT_NEAR(Scale.Z, 1.0f, SMALL_NUMBER);
-		
-		UnitTest.Advance();		
-
-		UnitTest.Solver->GetPerSolverField().AddTransientCommand({ TargetName, FalloffField->NewCopy() });
-
-		UnitTest.Advance();		
-
-		FVector Scale2 = Transform[1].GetScale3D();
-		// geometry hidden by 0 scaling on transform
-		EXPECT_NEAR(Scale2.X, 0.0f, SMALL_NUMBER);
-		EXPECT_NEAR(Scale2.Y, 0.0f, SMALL_NUMBER);
-		EXPECT_NEAR(Scale2.Z, 0.0f, SMALL_NUMBER);
-		
-		delete FalloffField;
-	
-	}
-
 	GTEST_TEST(AllTraits, DISABLED_GeometryCollection_RigidBodiess_ClusterTest_MaxClusterLevel_1)
 	{
 		FFramework UnitTest;
diff --git a/Engine/Source/Programs/HeadlessChaos/Private/GeometryCollection/GeometryCollectionTestFramework.cpp b/Engine/Source/Programs/HeadlessChaos/Private/GeometryCollection/GeometryCollectionTestFramework.cpp
--- a/Engine/Source/Programs/HeadlessChaos/Private/GeometryCollection/GeometryCollectionTestFramework.cpp
+++ b/Engine/Source/Programs/HeadlessChaos/Private/GeometryCollection/GeometryCollectionTestFramework.cpp
@@ -170,7 +170,6 @@ namespace GeometryCollectionTest
 			SimulationParams.DamageThreshold = Params.DamageThreshold;
 			SimulationParams.MaxClusterLevel = Params.MaxClusterLevel;
 			SimulationParams.ClusterConnectionMethod = Params.ClusterConnectionMethod;
-			SimulationParams.RemoveOnFractureEnabled = Params.RemoveOnFractureEnabled;
 			SimulationParams.CollisionGroup = Params.CollisionGroup;
 			SimulationParams.ClusterGroupIndex = Params.ClusterGroupIndex;
 			SimulationParams.LinearDamping = 0;
diff --git a/Engine/Source/Programs/HeadlessChaos/Private/GeometryCollection/GeometryCollectionTestSimulationField.cpp b/Engine/Source/Programs/HeadlessChaos/Private/GeometryCollection/GeometryCollectionTestSimulationField.cpp
--- a/Engine/Source/Programs/HeadlessChaos/Private/GeometryCollection/GeometryCollectionTestSimulationField.cpp
+++ b/Engine/Source/Programs/HeadlessChaos/Private/GeometryCollection/GeometryCollectionTestSimulationField.cpp
@@ -1117,7 +1117,6 @@ namespace GeometryCollectionTest
 		Params.Simulating = true;
 		Params.EnableClustering = true;
 		Params.DamageThreshold = { 91.0f, 92.0f, 93.0f, 94.0f };
-		Params.RemoveOnFractureEnabled = true;
 		FGeometryCollectionWrapper* Collection = TNewSimulationObject<GeometryType::GeometryCollectionWithSuppliedRestCollection>::Init(Params)->template As<FGeometryCollectionWrapper>();
 		UnitTest.AddSimulationObject(Collection);
 
@@ -1187,7 +1186,6 @@ namespace GeometryCollectionTest
 		Params.Simulating = true;
 		Params.EnableClustering = true;
 		Params.DamageThreshold = { 91.0f, 92.0f, 93.0f, 94.0f };
-		Params.RemoveOnFractureEnabled = true;
 		FGeometryCollectionWrapper* Collection = TNewSimulationObject<GeometryType::GeometryCollectionWithSuppliedRestCollection>::Init(Params)->template As<FGeometryCollectionWrapper>();
 		UnitTest.AddSimulationObject(Collection);
 
diff --git a/Engine/Source/Programs/HeadlessChaos/Public/GeometryCollection/GeometryCollectionTestFramework.h b/Engine/Source/Programs/HeadlessChaos/Public/GeometryCollection/GeometryCollectionTestFramework.h
--- a/Engine/Source/Programs/HeadlessChaos/Public/GeometryCollection/GeometryCollectionTestFramework.h
+++ b/Engine/Source/Programs/HeadlessChaos/Public/GeometryCollection/GeometryCollectionTestFramework.h
@@ -93,7 +93,6 @@ namespace GeometryCollectionTest
 		int32 MaxClusterLevel = 100;
 		TArray<float> DamageThreshold = { 1000.0 };
 		Chaos::FClusterCreationParameters::EConnectionMethod ClusterConnectionMethod = Chaos::FClusterCreationParameters::EConnectionMethod::PointImplicit;
-		bool RemoveOnFractureEnabled = false;
 		int32 CollisionGroup = 0;
 		int32 MinLevelSetResolution = 5;
 		int32 MaxLevelSetResolution = 10;
diff --git a/Engine/Source/Runtime/Experimental/Chaos/Private/Chaos/PBDRigidClustering.cpp b/Engine/Source/Runtime/Experimental/Chaos/Private/Chaos/PBDRigidClustering.cpp
--- a/Engine/Source/Runtime/Experimental/Chaos/Private/Chaos/PBDRigidClustering.cpp
+++ b/Engine/Source/Runtime/Experimental/Chaos/Private/Chaos/PBDRigidClustering.cpp
@@ -663,14 +663,7 @@ namespace Chaos
 				// 1 entry at a time.
 				Children.RemoveAtSwap(ChildIdx, 1, false);
 
-				if (Child->ToBeRemovedOnFracture())
-				{
-					MActiveRemovalIndices.Add(Child);
-				}
-				else
-				{
-					SendBreakingEvent(Child);
-				}
+				SendBreakingEvent(Child);
 			}
 			if (bUseDamagePropagation)
 			{
@@ -836,21 +829,6 @@ namespace Chaos
 		FDurationTimer Timer(Time);
 		Timer.Start();
 
-		{
-			const FReal Threshold = (FReal)1.f;
-			TSet<Chaos::FPBDRigidParticleHandle*> RemovalIndicesCopy = MActiveRemovalIndices;
-			for (Chaos::FPBDRigidParticleHandle* Particle : RemovalIndicesCopy)
-			{
-				//if (MParticles.ToBeRemovedOnFracture(ParticleIdx) && MParticles.V(ParticleIdx).SizeSquared() > Threshold && MParticles.PreV(ParticleIdx).SizeSquared() > Threshold)
-				if (Particle->ToBeRemovedOnFracture() && 
-					Particle->V().SizeSquared() > Threshold && 
-					Particle->PreV().SizeSquared() > Threshold)
-				{
-					DisableParticleWithBreakEvent(Particle->CastToClustered());
-				}
-			}
-		}
-
 		if(MChildren.Num())
 		{
 			//
@@ -928,17 +906,6 @@ namespace Chaos
 			{
 				ReleaseClusterParticles(ClusteredParticle);
 			}
-			else
-			{
-				// there's no children to break but we need to process whether this single piece is to be removed when damaged
-				if (ClusteredParticle->ToBeRemovedOnFracture())
-				{
-					if (ClusteredParticle->CollisionImpulses() >= ClusteredParticle->Strains())
-					{
-						DisableParticleWithBreakEvent(ClusteredParticle);
-					}
-				}
-			}
 		}
 	}
 
@@ -1268,7 +1235,6 @@ namespace Chaos
 		GetChildrenMap().Remove(ClusteredParticle);
 		ClusteredParticle->ClusterIds() = ClusterId();
 		ClusteredParticle->ClusterGroupIndex() = 0;
-		MActiveRemovalIndices.Remove(ClusteredParticle);
 	}
 
 	void FRigidClustering::DisableParticleWithBreakEvent(FPBDRigidClusteredParticleHandle* ClusteredParticle)
@@ -1311,7 +1277,6 @@ namespace Chaos
 
 	// reset the structures
 		TopLevelClusterParents.Remove(ClusteredParticle);
-		MActiveRemovalIndices.Remove(ClusteredParticle);
 
 		// disconnect from the parents
 		if (ClusteredParticle->ClusterIds().Id)
diff --git a/Engine/Source/Runtime/Experimental/Chaos/Private/PhysicsProxy/GeometryCollectionPhysicsProxy.cpp b/Engine/Source/Runtime/Experimental/Chaos/Private/PhysicsProxy/GeometryCollectionPhysicsProxy.cpp
--- a/Engine/Source/Runtime/Experimental/Chaos/Private/PhysicsProxy/GeometryCollectionPhysicsProxy.cpp
+++ b/Engine/Source/Runtime/Experimental/Chaos/Private/PhysicsProxy/GeometryCollectionPhysicsProxy.cpp
@@ -972,10 +972,6 @@ void FGeometryCollectionPhysicsProxy::InitializeBodiesPT(Chaos::FPBDRigidsSolver
 			}
 		}
 
-#if TODO_REIMPLEMENT_FRACTURE
-		InitializeRemoveOnFracture(Particles, DynamicState);
-#endif // TODO_REIMPLEMENT_FRACTURE
-
 		// #BG Temporary - don't cluster when playing back. Needs to be changed when kinematics are per-proxy to support
 		// kinematic to dynamic transition for clusters.
 		if (Parameters.EnableClustering)// && Parameters.CacheType != EGeometryCollectionCacheType::Play)
@@ -1875,30 +1871,6 @@ int32 FGeometryCollectionPhysicsProxy::CalculateAndSetLevel(int32 TransformGroup
 	return Levels[TransformGroupIndex];
 }
 
-void FGeometryCollectionPhysicsProxy::InitializeRemoveOnFracture(FParticlesType& Particles, const TManagedArray<int32>& DynamicState)
-{
-	/*
-	@todo break everything
-	if (Parameters.DynamicCollection && Parameters.RemoveOnFractureEnabled)
-	{
-	//	TManagedArray<FGeometryCollectionBoneNode>& Hierarchy = Parameters.DynamicCollection->BoneHierarchy;
-
-		for (int TransformGroupIndex = 0; TransformGroupIndex < RigidBodyID.Num(); TransformGroupIndex++)
-		{
-			if (RigidBodyID[TransformGroupIndex] != INDEX_NONE)
-			{
-				int32 RigidBodyIndex = RigidBodyID[TransformGroupIndex];
-
-				if (Parameters.DynamicCollection->StatusFlags[TransformGroupIndex] & FGeometryCollection::FS_RemoveOnFracture)
-				{
-					Particles.ToBeRemovedOnFracture(RigidBodyIndex) = true;
-				}
-			}
-		}
-	}
-	*/
-}
-
 void FGeometryCollectionPhysicsProxy::OnRemoveFromSolver(Chaos::FPBDRigidsSolver *RBDSolver)
 {
 	Chaos::FPBDRigidsEvolutionGBF* Evolution = RBDSolver->GetEvolution();
@@ -3569,51 +3541,6 @@ void FGeometryCollectionPhysicsProxy::InitializeSharedCollisionStructures(
 			// Set level of TransformGroupIndex in Levels attribute.
 			CalculateAndSetLevel(TransformGroupIndex, Parent, Levels);
 		}
-
-		InitRemoveOnFracture(RestCollection, SharedParams);
-	}
-}
-
-void FGeometryCollectionPhysicsProxy::InitRemoveOnFracture(FGeometryCollection& RestCollection, const FSharedSimulationParameters& SharedParams)
-{
-	if (SharedParams.RemoveOnFractureIndices.Num() == 0)
-	{
-		return;
-	}
-
-	// Markup Node Hierarchy Status with FS_RemoveOnFracture flags where geometry is ALL glass
-	const int32 NumGeometries = RestCollection.NumElements(FGeometryCollection::GeometryGroup);
-	for (int32 Idx = 0; Idx < NumGeometries; Idx++)
-	{
-		const int32 TransformIndex = RestCollection.TransformIndex[Idx];
-		const int32 Start = RestCollection.FaceStart[Idx];
-		const int32 End = RestCollection.FaceCount[Idx];
-		bool IsToBeRemoved = true;
-		for (int32 Face = Start; Face < Start + End; Face++)
-		{
-			bool FoundMatch = false;
-			for (int32 MaterialIndex : SharedParams.RemoveOnFractureIndices)
-			{
-				if (RestCollection.MaterialID[Face] == MaterialIndex)
-				{
-					FoundMatch = true;
-					break;
-				}
-			}
-			if (!FoundMatch)
-			{
-				IsToBeRemoved = false;
-				break;
-			}
-		}
-		if (IsToBeRemoved)
-		{
-			RestCollection.SetFlags(TransformIndex, FGeometryCollection::FS_RemoveOnFracture);
-		}
-		else
-		{
-			RestCollection.ClearFlags(TransformIndex, FGeometryCollection::FS_RemoveOnFracture);
-		}
 	}
 }
 
diff --git a/Engine/Source/Runtime/Experimental/Chaos/Public/Chaos/PBDRigidClustering.h b/Engine/Source/Runtime/Experimental/Chaos/Public/Chaos/PBDRigidClustering.h
--- a/Engine/Source/Runtime/Experimental/Chaos/Public/Chaos/PBDRigidClustering.h
+++ b/Engine/Source/Runtime/Experimental/Chaos/Public/Chaos/PBDRigidClustering.h
@@ -376,8 +376,6 @@ private:
 	FRigidEvolution& MEvolution;
 	FPBDRigidClusteredParticles& MParticles;
 	TSet<Chaos::FPBDRigidClusteredParticleHandle*> TopLevelClusterParents;
-	TSet<Chaos::FPBDRigidParticleHandle*> MActiveRemovalIndices;
-
 
 	// Cluster data
 	FClusterMap MChildren;
diff --git a/Engine/Source/Runtime/Experimental/Chaos/Public/Chaos/ParticleHandle.h b/Engine/Source/Runtime/Experimental/Chaos/Public/Chaos/ParticleHandle.h
--- a/Engine/Source/Runtime/Experimental/Chaos/Public/Chaos/ParticleHandle.h
+++ b/Engine/Source/Runtime/Experimental/Chaos/Public/Chaos/ParticleHandle.h
@@ -867,7 +867,6 @@ protected:
 		SetAngularAcceleration(TVector<T, d>(0));
 		SetObjectStateLowLevel(Params.bStartSleeping ? EObjectStateType::Sleeping : EObjectStateType::Dynamic);
 		SetIslandIndex(INDEX_NONE);
-		SetToBeRemovedOnFracture(false);
 		SetSleepType(ESleepType::MaterialSleep);
 		SetInvIConditioning(TVec3<FRealSingle>(1));
 	}
@@ -1045,10 +1044,6 @@ public:
 	int32& IslandIndex() { return PBDRigidParticles->IslandIndex(ParticleIdx); }
 	void SetIslandIndex(const int32 InIslandIndex) { PBDRigidParticles->IslandIndex(ParticleIdx) = InIslandIndex; }
 	
-	bool ToBeRemovedOnFracture() const { return PBDRigidParticles->ToBeRemovedOnFracture(ParticleIdx); }
-	bool& ToBeRemovedOnFracture() { return PBDRigidParticles->ToBeRemovedOnFracture(ParticleIdx); }
-	void SetToBeRemovedOnFracture(const bool bToBeRemovedOnFracture) { PBDRigidParticles->ToBeRemovedOnFracture(ParticleIdx) = bToBeRemovedOnFracture; }
-
 	EObjectStateType ObjectState() const { return PBDRigidParticles->ObjectState(ParticleIdx); }
 	EObjectStateType PreObjectState() const { return PBDRigidParticles->PreObjectState(ParticleIdx); }
 
@@ -1859,16 +1854,6 @@ public:
 		return MHandle->IsInConstraintGraph();
 	}
 
-	bool ToBeRemovedOnFracture() const 
-	{
-		if (MHandle->CastToRigidParticle() && MHandle->ObjectState() == EObjectStateType::Dynamic)
-		{
-			return MHandle->CastToRigidParticle()->ToBeRemovedOnFracture();
-		}
-
-		return false;
-	}
-
 	const FShapesArray& ShapesArray() const { return MHandle->ShapesArray(); }
 
 	static constexpr EParticleType StaticType()
@@ -2570,7 +2555,6 @@ protected:
 	{
 		Type = EParticleType::Rigid;
 		MIsland = INDEX_NONE;
-		MToBeRemovedOnFracture = false;
 		PBDRigidParticleDefaultConstruct<T, d>(*this, DynamicParams);
 		ClearForces();
 		ClearTorques();
@@ -2591,7 +2575,10 @@ public:
 		Ar << MDynamics;
 
 		Ar << MIsland;
-		Ar << MToBeRemovedOnFracture;
+
+		// remove on fracture is deprecated and has been removed, so to avoid versioning let's just use a local variable instead
+		bool MToBeRemovedOnFracture_deprecated = false;
+		Ar << MToBeRemovedOnFracture_deprecated;
 	}
 
 	//const TUniquePtr<TBVHParticles<T, d>>& CollisionParticles() const { return MCollisionParticles; }
@@ -2805,13 +2792,6 @@ public:
 		this->MIsland = InIsland;
 	}
 
-	bool ToBeRemovedOnFracture() const { return MToBeRemovedOnFracture; }
-	// TODO(stett): Make the setter private. It is public right now to provide access to proxies.
-	void SetToBeRemovedOnFracture(const bool bToBeRemovedOnFracture)
-	{
-		this->MToBeRemovedOnFracture = bToBeRemovedOnFracture;
-	}
-
 	EObjectStateType ObjectState() const { return MMiscData.Read().ObjectState(); }
 	void SetObjectState(const EObjectStateType InState, bool bAllowEvents=false, bool bInvalidate=true)
 	{
@@ -2880,7 +2860,6 @@ private:
 	TChaosProperty<FParticleDynamicMisc,EChaosProperty::DynamicMisc> MMiscData;
 
 	int32 MIsland;
-	bool MToBeRemovedOnFracture;
 	bool MInitialized;
 	EWakeEventEntry MWakeEvent;
 
diff --git a/Engine/Source/Runtime/Experimental/Chaos/Public/Chaos/RigidParticles.h b/Engine/Source/Runtime/Experimental/Chaos/Public/Chaos/RigidParticles.h
--- a/Engine/Source/Runtime/Experimental/Chaos/Public/Chaos/RigidParticles.h
+++ b/Engine/Source/Runtime/Experimental/Chaos/Public/Chaos/RigidParticles.h
@@ -95,7 +95,6 @@ public:
 		TArrayCollection::AddArray(&MObjectState);
 		TArrayCollection::AddArray(&MPreObjectState);
 		TArrayCollection::AddArray(&MIslandIndex);
-		TArrayCollection::AddArray(&MToBeRemovedOnFracture);
 		TArrayCollection::AddArray(&MSleepType);
 	}
 	TRigidParticles(const TRigidParticles<T, d>& Other) = delete;
@@ -124,7 +123,6 @@ public:
 		, MControlFlags(MoveTemp(Other.MControlFlags))
 		, MTransientFlags(MoveTemp(Other.MTransientFlags))
 		, MIslandIndex(MoveTemp(Other.MIslandIndex))
-		, MToBeRemovedOnFracture(MoveTemp(Other.MToBeRemovedOnFracture))
 		, MObjectState(MoveTemp(Other.MObjectState))
 		, MPreObjectState(MoveTemp(Other.MPreObjectState))
 		, MSleepType(MoveTemp(Other.MSleepType))
@@ -156,7 +154,6 @@ public:
 		TArrayCollection::AddArray(&MObjectState);
 		TArrayCollection::AddArray(&MPreObjectState);
 		TArrayCollection::AddArray(&MIslandIndex);
-		TArrayCollection::AddArray(&MToBeRemovedOnFracture);
 		TArrayCollection::AddArray(&MSleepType);
 	}
 
@@ -239,9 +236,6 @@ public:
     // Using this will break stuff. This is for solver's use only, and possibly some particle construction/copy code.
 	FORCEINLINE void SetDisabledLowLevel(const int32 Index, bool disabled) { MDisabled[Index] = disabled; }
 
-	FORCEINLINE const bool ToBeRemovedOnFracture(const int32 Index) const { return MToBeRemovedOnFracture[Index]; }
-	FORCEINLINE bool& ToBeRemovedOnFracture(const int32 Index) { return MToBeRemovedOnFracture[Index]; }
-
 	FORCEINLINE const FRigidParticleControlFlags& ControlFlags(const int32 Index) const { return MControlFlags[Index]; }
 	FORCEINLINE FRigidParticleControlFlags& ControlFlags(const int32 Index) { return MControlFlags[Index]; }
 
@@ -370,7 +364,6 @@ private:
 	TArrayCollectionArray<FRigidParticleControlFlags> MControlFlags;
 	TArrayCollectionArray<FRigidParticleTransientFlags> MTransientFlags;
 	TArrayCollectionArray<int32> MIslandIndex;
-	TArrayCollectionArray<bool> MToBeRemovedOnFracture;
 	TArrayCollectionArray<EObjectStateType> MObjectState;
 	TArrayCollectionArray<EObjectStateType> MPreObjectState;
 	TArrayCollectionArray<ESleepType> MSleepType;
diff --git a/Engine/Source/Runtime/Experimental/Chaos/Public/GeometryCollection/GeometryCollectionSimulationCoreTypes.h b/Engine/Source/Runtime/Experimental/Chaos/Public/GeometryCollection/GeometryCollectionSimulationCoreTypes.h
--- a/Engine/Source/Runtime/Experimental/Chaos/Public/GeometryCollection/GeometryCollectionSimulationCoreTypes.h
+++ b/Engine/Source/Runtime/Experimental/Chaos/Public/GeometryCollection/GeometryCollectionSimulationCoreTypes.h
@@ -161,7 +161,6 @@ struct FSharedSimulationParameters
 	float MaximumVolumeClamp() const { return MaximumBoundingExtentClamp * MaximumBoundingExtentClamp * MaximumBoundingExtentClamp; }
 
 	TArray<FSharedSimulationSizeSpecificData> SizeSpecificData;
-	TArray<int32> RemoveOnFractureIndices;
 	int32 MaximumCollisionParticleCount;
 	bool bUseImportedCollisionImplicits;
 };
@@ -200,7 +199,6 @@ struct FSimulationParameters
 		, bGenerateRemovalsData(false)
 		, bGenerateCrumblingData(false)
 		, bGenerateCrumblingChildrenData(false)
-		, RemoveOnFractureEnabled(false)
 		, EnableGravity(true)
 		, UseInertiaConditioning(true)
 		, UseCCD(false)
@@ -249,7 +247,6 @@ struct FSimulationParameters
 		, bGenerateCrumblingData(Other.bGenerateCrumblingData)
 		, bGenerateCrumblingChildrenData(Other.bGenerateCrumblingChildrenData)
 		, Shared(Other.Shared)
-		, RemoveOnFractureEnabled(Other.RemoveOnFractureEnabled)
 		, EnableGravity(Other.EnableGravity)
 		, UseInertiaConditioning(Other.UseInertiaConditioning)
 		, UseCCD(Other.UseCCD)
@@ -320,8 +317,6 @@ struct FSimulationParameters
 
 	FSharedSimulationParameters Shared;
 
-	bool RemoveOnFractureEnabled;
-	
 	bool EnableGravity;
 	bool UseInertiaConditioning;
 	bool UseCCD;
diff --git a/Engine/Source/Runtime/Experimental/Chaos/Public/PhysicsProxy/GeometryCollectionPhysicsProxy.h b/Engine/Source/Runtime/Experimental/Chaos/Public/PhysicsProxy/GeometryCollectionPhysicsProxy.h
--- a/Engine/Source/Runtime/Experimental/Chaos/Public/PhysicsProxy/GeometryCollectionPhysicsProxy.h
+++ b/Engine/Source/Runtime/Experimental/Chaos/Public/PhysicsProxy/GeometryCollectionPhysicsProxy.h
@@ -245,7 +245,6 @@ public:
 	{ Commands.Add(Command); }
 
 	static void InitializeSharedCollisionStructures(Chaos::FErrorReporter& ErrorReporter, FGeometryCollection& RestCollection, const FSharedSimulationParameters& SharedParams);
-	static void InitRemoveOnFracture(FGeometryCollection& RestCollection, const FSharedSimulationParameters& SharedParams);
 
 	void FieldForcesUpdateCallback(Chaos::FPBDRigidsSolver* RigidSolver);
 
@@ -443,8 +442,6 @@ protected:
 	 */
 	static int32 CalculateHierarchyLevel(const FGeometryDynamicCollection& DynamicCollection, int32 TransformIndex);
 
-	void InitializeRemoveOnFracture(FParticlesType& Particles, const TManagedArray<int32>& DynamicState);
-
 	void SetClusteredParticleKinematicTarget_Internal(Chaos::FPBDRigidClusteredParticleHandle* Handle, const FTransform& WorldTransform);
 
 	void PrepareBufferData(Chaos::FDirtyGeometryCollectionData& BufferData, const FGeometryDynamicCollection& ThreadCollection,  Chaos::FReal SolverLastDt = 0.0);
diff --git a/Engine/Source/Runtime/Experimental/Chaos/Public/PhysicsProxy/SingleParticlePhysicsProxy.h b/Engine/Source/Runtime/Experimental/Chaos/Public/PhysicsProxy/SingleParticlePhysicsProxy.h
--- a/Engine/Source/Runtime/Experimental/Chaos/Public/PhysicsProxy/SingleParticlePhysicsProxy.h
+++ b/Engine/Source/Runtime/Experimental/Chaos/Public/PhysicsProxy/SingleParticlePhysicsProxy.h
@@ -1061,26 +1061,6 @@ public:
 		}
 	}
 
-	bool ToBeRemovedOnFracture() const
-	{
-		VerifyContext();
-		if (const TPBDRigidParticle<FReal, 3>* Rigid = GetParticle_LowLevel()->CastToRigidParticle())
-		{
-			return Rigid->ToBeRemovedOnFracture();
-		}
-
-		return false;
-	}
-	// TODO(stett): Make the setter private. It is public right now to provide access to proxies.
-	void SetToBeRemovedOnFracture(const bool InToBeRemovedOnFracture)
-	{
-		VerifyContext();
-		if (TPBDRigidParticle<FReal, 3>* Rigid = GetParticle_LowLevel()->CastToRigidParticle())
-		{
-			Rigid->SetToBeRemovedOnFracture(InToBeRemovedOnFracture);
-		}
-	}
-
 	void ClearEvents()
 	{
 		VerifyContext();
diff --git a/Engine/Source/Runtime/Experimental/GeometryCollectionEngine/Private/GeometryCollection/GeometryCollectionComponent.cpp b/Engine/Source/Runtime/Experimental/GeometryCollectionEngine/Private/GeometryCollection/GeometryCollectionComponent.cpp
--- a/Engine/Source/Runtime/Experimental/GeometryCollectionEngine/Private/GeometryCollection/GeometryCollectionComponent.cpp
+++ b/Engine/Source/Runtime/Experimental/GeometryCollectionEngine/Private/GeometryCollection/GeometryCollectionComponent.cpp
@@ -2514,7 +2514,6 @@ void UGeometryCollectionComponent::RegisterAndInitializePhysicsProxy()
 		SimulationParameters.bGenerateRemovalsData = bNotifyRemovals;
 		SimulationParameters.bGenerateCrumblingData = bNotifyCrumblings;
 		SimulationParameters.bGenerateCrumblingChildrenData = bCrumblingEventIncludesChildren;
-		SimulationParameters.RemoveOnFractureEnabled = SimulationParameters.Shared.RemoveOnFractureIndices.Num() > 0;
 		SimulationParameters.EnableGravity = BodyInstance.bEnableGravity;
 		SimulationParameters.UseInertiaConditioning = BodyInstance.IsInertiaConditioningEnabled();
 		SimulationParameters.UseCCD = BodyInstance.bUseCCD;
@@ -2702,7 +2701,7 @@ void UGeometryCollectionComponent::SetRestCollection(const UGeometryCollection*
 		}
 
 		// initialize the component per level damage threshold from the asset defaults 
-		DamageThreshold = RestCollection->DamageThreshold;
+		//DamageThreshold = RestCollection->DamageThreshold;
 
 		// initialize the component damage progataion data from the asset defaults 
 		DamagePropagationData = RestCollection->DamagePropagationData;
diff --git a/Engine/Source/Runtime/Experimental/GeometryCollectionEngine/Private/GeometryCollection/GeometryCollectionObject.cpp b/Engine/Source/Runtime/Experimental/GeometryCollectionEngine/Private/GeometryCollection/GeometryCollectionObject.cpp
--- a/Engine/Source/Runtime/Experimental/GeometryCollectionEngine/Private/GeometryCollection/GeometryCollectionObject.cpp
+++ b/Engine/Source/Runtime/Experimental/GeometryCollectionEngine/Private/GeometryCollection/GeometryCollectionObject.cpp
@@ -97,7 +97,7 @@ UGeometryCollection::UGeometryCollection(const FObjectInitializer& ObjectInitial
 	, RemovalDuration(2.5, 5.0)
 	, bSlowMovingAsSleeping(true)
 	, SlowMovingVelocityThreshold(1)
-	, EnableRemovePiecesOnFracture(false)
+	, EnableRemovePiecesOnFracture_DEPRECATED(false)
 	, GeometryCollection(new FGeometryCollection())
 {
 	PersistentGuid = FGuid::NewGuid();
@@ -346,37 +346,11 @@ void UGeometryCollection::GetSharedSimulationParams(FSharedSimulationParameters&
 		FillSharedSimulationSizeSpecificData(OutParams.SizeSpecificData[Idx+1], SizeSpecificData[Idx]);
 	}
 
-	if (EnableRemovePiecesOnFracture)
-	{
-		FixupRemoveOnFractureMaterials(OutParams);
-	}
-
 	OutParams.bUseImportedCollisionImplicits = bImportCollisionFromSource;
 	
 	OutParams.SizeSpecificData.Sort();	//can we do this at editor time on post edit change?
 }
 
-void UGeometryCollection::FixupRemoveOnFractureMaterials(FSharedSimulationParameters& SharedParms) const
-{
-	// Match RemoveOnFracture materials with materials in model and record the material indices
-
-	int32 NumMaterials = Materials.Num();
-	for (int32 MaterialIndex = 0; MaterialIndex < NumMaterials; MaterialIndex++)
-	{
-		UMaterialInterface* MaterialInfo = Materials[MaterialIndex];
-
-		for (int32 ROFMaterialIndex = 0; ROFMaterialIndex < RemoveOnFractureMaterials.Num(); ROFMaterialIndex++)
-		{
-			if (MaterialInfo == RemoveOnFractureMaterials[ROFMaterialIndex])
-			{
-				SharedParms.RemoveOnFractureIndices.Add(MaterialIndex);
-				break;
-			}
-		}
-
-	}
-}
-
 void UGeometryCollection::Reset()
 {
 	if (GeometryCollection.IsValid())
diff --git a/Engine/Source/Runtime/Experimental/GeometryCollectionEngine/Public/GeometryCollection/GeometryCollectionObject.h b/Engine/Source/Runtime/Experimental/GeometryCollectionEngine/Public/GeometryCollection/GeometryCollectionObject.h
--- a/Engine/Source/Runtime/Experimental/GeometryCollectionEngine/Public/GeometryCollection/GeometryCollectionObject.h
+++ b/Engine/Source/Runtime/Experimental/GeometryCollectionEngine/Public/GeometryCollection/GeometryCollectionObject.h
@@ -435,7 +435,6 @@ public:
 
 	/** Fills params struct with parameters used for precomputing content. */
 	void GetSharedSimulationParams(FSharedSimulationParameters& OutParams) const;
-	void FixupRemoveOnFractureMaterials(FSharedSimulationParameters& SharedParms) const;
 
 	/** Accessors for the two guids used to identify this collection */
 	FGuid GetIdGuid() const;
@@ -462,7 +461,7 @@ public:
 	/** compatibility check, when true, only cluster compute damage from parameters and propagate to direct children
 	 *  when false, each child will compute it's damage threshold allowing for more precise and intuitive destruction behavior
 	 */
-	UPROPERTY(EditAnywhere, Category = "Damage")
+	UPROPERTY(EditAnywhere, Category = "Compatibility")
 	bool PerClusterOnlyDamageThreshold;
 
 	/** Data about how damage propagation shoudl behave. */
@@ -622,14 +621,14 @@ public:
 	/**
 	* Enable remove pieces on fracture
 	*/
-	UPROPERTY(EditAnywhere, BlueprintReadOnly, Category = "Fracture")
-	bool EnableRemovePiecesOnFracture;
+	UPROPERTY(meta = (DeprecatedProperty, DeprecationMessage = "Use remove on break feature instead ( Fracture editor tools )."))
+	bool EnableRemovePiecesOnFracture_DEPRECATED;
 
 	/**
 	* Materials relating to remove on fracture
 	*/
-	UPROPERTY(EditAnywhere, Category = "Fracture")
-	TArray<TObjectPtr<UMaterialInterface>> RemoveOnFractureMaterials;
+	UPROPERTY(meta = (DeprecatedProperty, DeprecationMessage = "Use remove on break feature instead ( Fracture editor tools )."))
+	TArray<TObjectPtr<UMaterialInterface>> RemoveOnFractureMaterials_DEPRECATED;
 
 	FORCEINLINE const int32 GetBoneSelectedMaterialIndex() const { return BoneSelectedMaterialIndex; }
 
