Chaos: EPA edge case fix:
- Make sure to eliminate triangles that are exactly edge on to prevent degenerate triangles later on in the algo
- Adjust error margins to be correct for operations performed

#[fyi] Michael.Lentine, Benn.Gallagher

diff --git a/Engine/Source/Programs/HeadlessChaos/Private/HeadlessChaosTestGJK.cpp b/Engine/Source/Programs/HeadlessChaos/Private/HeadlessChaosTestGJK.cpp
--- a/Engine/Source/Programs/HeadlessChaos/Private/HeadlessChaosTestGJK.cpp
+++ b/Engine/Source/Programs/HeadlessChaos/Private/HeadlessChaosTestGJK.cpp
@@ -1688,7 +1688,7 @@ namespace ChaosTest
 	const int32 NumBoxBoxGJKDistances = UE_ARRAY_COUNT(BoxBoxGJKDistances);
 
 	// These tests fails in EPA - we need to cover these cases with SAT
-	GTEST_TEST(GJKTests, DISABLED_TestGJKBoxBoxTestFails)
+	GTEST_TEST(GJKTests, TestGJKBoxBoxTestFails)
 	{
 		const FReal Epsilon = 1.e-3f;
 
@@ -1702,7 +1702,7 @@ namespace ChaosTest
 	}
 
 	// Disabled until we have SAT fallback (see DISABLED_TestGJKBoxBoxTestFails)
-	GTEST_TEST(GJKTests, DISABLED_TestGJKBoxBoxNegativeSeparation)
+	GTEST_TEST(GJKTests, TestGJKBoxBoxNegativeSeparation)
 	{
 		const FReal Epsilon = 1.e-3f;
 
@@ -1728,6 +1728,31 @@ namespace ChaosTest
 		}
 	}
 
+	// This is a know regression test
+	// It is two boxes deeply overlapping in a T-shape
+	GTEST_TEST(GJKTests, TestGJKBoxBoxOverlap)
+	{
+		FVec3 MinBox = FVec3(-15.839999675750732, -31.840000152587891, -3.8146972691777137e-07);
+		FVec3 MaxBox = FVec3(15.840000629425049, 31.840000152587891, 19.200000381469728);
+		FVec3 Translation = FVec3(15.999999999999993, 1.9594348786357647e-15, 0.0000000000000000);
+		FRotation3 Rotation(UE::Math::TQuat<FReal>(0.0000000000000000, 0.0000000000000000, -0.70710678118654757, -0.70710678118654746));
+
+		FImplicitBox3 ShapeA(MinBox, MaxBox, 0);
+		FImplicitBox3 ShapeB(MinBox, MaxBox, 0);
+		
+		const FRigidTransform3 TransformBtoA = FRigidTransform3(Translation , Rotation);
+		const FReal ThicknessA = 0.0f;
+		const FReal ThicknessB = 0.0f;
+
+		// Run GJK/EPA
+		FReal Penetration;
+		FVec3 ClosestA, ClosestBInA, Normal;
+		int32 ClosestVertexIndexA, ClosestVertexIndexB;
+		bool bSuccess = GJKPenetration<false>(ShapeA, ShapeB, TransformBtoA, Penetration, ClosestA, ClosestBInA, Normal, ClosestVertexIndexA, ClosestVertexIndexB, ThicknessA, ThicknessB, FVec3(1, 0, 0));
+		EXPECT_TRUE(bSuccess);
+		EXPECT_TRUE(Penetration > 19.0f); // Penetration is the Height of the box
+	}
+
 	// Two convex shapes, Shape A on top of Shape B and almost touching. ShapeA is rotated 90 degrees about Z.
 	// Check that the contact point lies between Shape A and Shape B with a near zero Phi.
 	// This reproduces a bug where GJKPenetrationCore returns points on top of A and at the bottom
diff --git a/Engine/Source/Runtime/Experimental/Chaos/Public/Chaos/EPA.h b/Engine/Source/Runtime/Experimental/Chaos/Public/Chaos/EPA.h
--- a/Engine/Source/Runtime/Experimental/Chaos/Public/Chaos/EPA.h
+++ b/Engine/Source/Runtime/Experimental/Chaos/Public/Chaos/EPA.h
@@ -39,7 +39,7 @@ struct TEPAEntry
 		return Distance > Other.Distance;
 	}
 
-	constexpr T SelectEpsilon(float FloatEpsilon, double DoubleEpsilon) const
+	static constexpr T SelectEpsilon(float FloatEpsilon, double DoubleEpsilon)
 	{
 		if (sizeof(T) <= sizeof(float))
 		{
@@ -61,7 +61,9 @@ struct TEPAEntry
 		const TVec3<T> V0V2 = V2 - V0;
 		const TVec3<T> Norm = TVec3<T>::CrossProduct(V0V1, V0V2);
 		const T NormLenSq = Norm.SizeSquared();
-		const T Eps = SelectEpsilon(1.e-4f, 1.e-8);
+		// We have the square of the size of a cross product, so we need the distance margin to be a power of 4
+		// Verbosity emphasizes that
+		const T Eps = TEPAEntry::SelectEpsilon(1.e-4f * 1.e-4f * 1.e-4f * 1.e-4f, 1.e-8 * 1.e-8 * 1.e-8 * 1.e-8);
 		if (NormLenSq < Eps)
 		{
 			return false;
@@ -345,7 +347,7 @@ void EPAComputeVisibilityBorder(TEPAWorkingArray<TEPAEntry<T>>& Entries, int32 E
 		TEPAEntry<T>& Entry = Entries[FloodEntry.EntryIdx];
 		if (!Entry.bObsolete)
 		{
-			if (Entry.DistanceToPlane(W) < UE_SMALL_NUMBER)
+			if (Entry.DistanceToPlane(W) <= TEPAEntry<T>::SelectEpsilon(1.e-4f, 1.e-8))
 			{
 				//W can't see this triangle so mark the edge as a border
 				OutBorderEdges.Add(FloodEntry);
