HeadlessChaos : GeometryCollection type inspection test.

diff --git a/Engine/Source/Programs/HeadlessChaos/Private/GeometryCollection/GeometryCollectionTestCreation.cpp b/Engine/Source/Programs/HeadlessChaos/Private/GeometryCollection/GeometryCollectionTestCreation.cpp
--- a/Engine/Source/Programs/HeadlessChaos/Private/GeometryCollection/GeometryCollectionTestCreation.cpp
+++ b/Engine/Source/Programs/HeadlessChaos/Private/GeometryCollection/GeometryCollectionTestCreation.cpp
@@ -14,6 +14,31 @@ DEFINE_LOG_CATEGORY_STATIC(GCTCR_Log, Verbose, All);
 
 namespace GeometryCollectionTest
 {
+	void  CheckClassTypes()
+	{
+		FManagedArrayCollection ArrayCollection;
+		EXPECT_TRUE(ArrayCollection.IsA<FManagedArrayCollection>());
+		FManagedArrayCollection* Result = ArrayCollection.Cast<FManagedArrayCollection>();
+		EXPECT_TRUE(ArrayCollection.Cast<FManagedArrayCollection>()!=nullptr);
+		EXPECT_TRUE(ArrayCollection.Cast<FTransformCollection>()==nullptr);
+		EXPECT_TRUE(ArrayCollection.Cast<FGeometryCollection>()==nullptr);
+
+		FTransformCollection TransformCollection;
+		EXPECT_TRUE(TransformCollection.IsA<FManagedArrayCollection>());
+		EXPECT_TRUE(TransformCollection.Cast<FManagedArrayCollection>()!=nullptr);
+		EXPECT_TRUE(TransformCollection.IsA<FTransformCollection>());
+		EXPECT_TRUE(TransformCollection.Cast<FTransformCollection>()!=nullptr);
+		EXPECT_TRUE(TransformCollection.Cast<FGeometryCollection>()==nullptr);
+
+		FGeometryCollection GeometryCollection;
+		EXPECT_TRUE(GeometryCollection.IsA<FManagedArrayCollection>());
+		EXPECT_TRUE(GeometryCollection.Cast<FManagedArrayCollection>() != nullptr);
+		EXPECT_TRUE(GeometryCollection.IsA<FTransformCollection>());
+		EXPECT_TRUE(GeometryCollection.Cast<FTransformCollection>() != nullptr);
+		EXPECT_TRUE(GeometryCollection.IsA<FGeometryCollection>());
+		EXPECT_TRUE(GeometryCollection.Cast<FGeometryCollection>() != nullptr);
+	}
+
 
 	void CheckIncrementMask()
 	{
diff --git a/Engine/Source/Programs/HeadlessChaos/Private/HeadlessChaosApp.cpp b/Engine/Source/Programs/HeadlessChaos/Private/HeadlessChaosApp.cpp
--- a/Engine/Source/Programs/HeadlessChaos/Private/HeadlessChaosApp.cpp
+++ b/Engine/Source/Programs/HeadlessChaos/Private/HeadlessChaosApp.cpp
@@ -318,7 +318,8 @@ TEST(GeometryCollection_MatricesTest,TransformMatrixElement) { GeometryCollectio
 TEST(GeometryCollection_MatricesTest,ReparentingMatrices) { GeometryCollectionTest::ReparentingMatrices(); SUCCEED(); }
 
 // Creation Tests
-TEST(GeometryCollection_CreationTest,CheckIncrementMask) { GeometryCollectionTest::CheckIncrementMask(); SUCCEED(); }
+TEST(GeometryCollection_CreationTest, CheckClassTypes) { GeometryCollectionTest::CheckClassTypes(); SUCCEED(); }
+TEST(GeometryCollection_CreationTest, CheckIncrementMask) { GeometryCollectionTest::CheckIncrementMask(); SUCCEED(); }
 TEST(GeometryCollection_CreationTest,Creation) { GeometryCollectionTest::Creation(); SUCCEED(); }
 TEST(GeometryCollection_CreationTest,Empty) { GeometryCollectionTest::Empty(); SUCCEED(); }
 TEST(GeometryCollection_CreationTest,AppendTransformHierarchy) { GeometryCollectionTest::AppendTransformHierarchy(); SUCCEED(); }
diff --git a/Engine/Source/Programs/HeadlessChaos/Public/GeometryCollection/GeometryCollectionTestCreation.h b/Engine/Source/Programs/HeadlessChaos/Public/GeometryCollection/GeometryCollectionTestCreation.h
--- a/Engine/Source/Programs/HeadlessChaos/Public/GeometryCollection/GeometryCollectionTestCreation.h
+++ b/Engine/Source/Programs/HeadlessChaos/Public/GeometryCollection/GeometryCollectionTestCreation.h
@@ -5,6 +5,8 @@
 
 namespace GeometryCollectionTest
 {	
+	void  CheckClassTypes();
+
 	void  CheckIncrementMask();
 
 	void  Creation();
