GeometryCollection : Append Collection and Dependency Checking
- Added ability to append collections and preserve dependencies.
- Fixed dependency checking to perform directed graph search.
- updated unit tests.

diff --git a/Engine/Source/Programs/HeadlessChaos/Private/GeometryCollection/GeometryCollectionTestCreation.cpp b/Engine/Source/Programs/HeadlessChaos/Private/GeometryCollection/GeometryCollectionTestCreation.cpp
--- a/Engine/Source/Programs/HeadlessChaos/Private/GeometryCollection/GeometryCollectionTestCreation.cpp
+++ b/Engine/Source/Programs/HeadlessChaos/Private/GeometryCollection/GeometryCollectionTestCreation.cpp
@@ -1462,4 +1462,258 @@ namespace GeometryCollectionTest
 
 		delete Collection;
 	}
+
+
+	void AppendManagedArrayCollectionTest()
+	{
+		//--gtest_filter=GeometryCollection_CreationTest.AppendManagedArrayCollectionTest
+
+		//
+		// Build two ManagedArryaCollections then append CollectionA to CollectionB. 
+		//
+		// At the end of the test CollectionA should be moved and CollectionB
+		// should be at the start of the arrays. 
+		//
+		
+		auto Build = [](FManagedArrayCollection& InCollection, int Offset, FName GroupName)
+		{
+			FManagedArrayCollection::FConstructionParameters DependencyOnShared("Shared");
+
+			// single array with 10 entries from 0 to 1.0
+			InCollection.AddGroup("Shared");
+			InCollection.AddElements(10, "Shared");
+			TManagedArray<float>& IntInShared = InCollection.AddAttribute<float>("Float", "Shared");
+			for (int i = 0; i < 10; i++)
+			{
+				IntInShared[i] = i * 0.1 + Offset;
+			}
+
+			// two arrays, where Int is dependent on Shared, and Intf is a copy of shared.
+			// Intf is used to check the reordering of the groups. 
+			InCollection.AddGroup(GroupName);
+			InCollection.AddElements(10, GroupName);
+			TManagedArray<int32>& IntIn = InCollection.AddAttribute<int32>("Int", GroupName, DependencyOnShared);
+			TManagedArray<float>& SharedValue = InCollection.AddAttribute<float>("SharedValue", GroupName);
+			for (int i = 0; i < 10; i++)
+			{
+				// int32
+				IntIn[i] = i;
+				SharedValue[i] = IntInShared[IntIn[i]];
+			}
+		};
+
+		FManagedArrayCollection CollectionA; Build(CollectionA, 0, FName("GroupA"));
+		FManagedArrayCollection CollectionB; Build(CollectionB, 1, FName("GroupB"));
+
+		CollectionA.Append(CollectionB);
+
+		// check the union of shared
+		const TManagedArray<float>& Shared = CollectionA.GetAttribute<float>("Float", "Shared");
+		EXPECT_TRUE(CollectionA.NumElements("Shared") == 20);
+		for (int i = 0; i < 10; i++) // 0 to 9 greater than one
+		{
+			float val = (float)i * 0.1 + 1;
+			EXPECT_TRUE( FMath::IsNearlyEqual(Shared[i], val));
+		}
+		for (int i = 10; i < 20; i++) // 10 to 19 lass than one
+		{
+			float val = (float)(i-10) * 0.1;
+			EXPECT_TRUE(FMath::IsNearlyEqual(Shared[i], val));
+		}
+		
+		// check reindexing of GroupA
+		const TManagedArray<int32>& IntInA = CollectionA.GetAttribute<int32>("Int", "GroupA");
+		for (int i = 0; i < 10; i++)
+		{
+			EXPECT_TRUE(IntInA[i] == i + 10);
+		}
+
+		// check indexing of GroupB
+		const TManagedArray<int32>& IntInB = CollectionA.GetAttribute<int32>("Int", "GroupB");
+		for (int i = 0; i < 10; i++)
+		{
+			EXPECT_TRUE(IntInB[i] == i);
+		}
+
+		// check reindexing into shared for A
+		const TManagedArray<float>& SharedValueA = CollectionA.GetAttribute<float>("SharedValue", "GroupA");
+		for (int i = 0; i < 10; i++)
+		{
+			EXPECT_TRUE(SharedValueA[i] == Shared[IntInA[i]]);
+		}
+
+		// check reindexing into shared for B
+		const TManagedArray<float>& SharedValueB = CollectionA.GetAttribute<float>("SharedValue", "GroupB");
+		for (int i = 0; i < 10; i++)
+		{
+			EXPECT_TRUE(SharedValueB[i] == Shared[IntInB[i]]);
+		}
+
+		EXPECT_TRUE(true);
+	}
+
+
+	void AppendTransformCollectionTest()
+	{
+		//--gtest_filter=GeometryCollection_CreationTest.AppendTransformCollectionTest
+
+		//
+		// Build two TranformCollections then append CollectionA to CollectionB. 
+		//
+		// At the end of the test CollectionA should be moved and CollectionB
+		// should be at the start of the arrays. The parenting of CollectionA
+		// should be updated to maintain the parent-child relationships.  
+		//
+		auto Build = [](FTransformCollection& InCollection, int32 Offset)
+		{
+			InCollection.AddElements(8, FTransformCollection::TransformGroup);
+
+			//  0
+			//  ...1
+			//  ...5
+			//  ......6
+			//  ......3
+			//  ...2
+			//  ......7
+			//  .........4
+			(InCollection.Parent)[0] = -1;
+			(InCollection.Children)[0].Add(1);
+			(InCollection.Children)[0].Add(5);
+			(InCollection.Children)[0].Add(2);
+			(InCollection.Parent)[1] = 0;
+			(InCollection.Parent)[2] = 0;
+			(InCollection.Children)[2].Add(7);
+			(InCollection.Parent)[3] = 5;
+			(InCollection.Parent)[4] = 7;
+			(InCollection.Parent)[5] = 0;
+			(InCollection.Children)[5].Add(6);
+			(InCollection.Children)[5].Add(3);
+			(InCollection.Parent)[6] = 5;
+			(InCollection.Parent)[7] = 2;
+			(InCollection.Children)[7].Add(4);
+
+			for (int i = 0; i < 8; i++)
+			{
+				(InCollection.BoneName)[i] = FString::FromInt(i);
+			}
+		};
+
+		FTransformCollection CollectionA; Build(CollectionA, 0);
+		FTransformCollection CollectionB; Build(CollectionB, 10);
+
+		// add some random dependent attribute to the transform group. 
+		FManagedArrayCollection::FConstructionParameters TransformDependency(FTransformCollection::TransformGroup);
+		TManagedArray<int32>& ExtraA = CollectionA.AddAttribute<int32>("UserAttr", FTransformCollection::TransformGroup, TransformDependency);
+		TManagedArray<int32>& ExtraB = CollectionB.AddAttribute<int32>("UserAttr", FTransformCollection::TransformGroup, TransformDependency);
+		for (int i = ExtraA.Num()-1, j = 0; i >= 0; i--, j++) { ExtraA[i] = j; ExtraB[i] = j; } // count down
+
+		// add some random transform dependent attribute to a different group. 
+		FManagedArrayCollection::FConstructionParameters OtherDependency("Other");
+		CollectionA.AddGroup("Other"); CollectionB.AddGroup("Other");
+		CollectionA.AddElements(5, "Other"); CollectionB.AddElements(6, "Other");
+		TManagedArray<int32>& OtherA = CollectionA.AddAttribute<int32>("UserAttr", "Other", TransformDependency);
+		TManagedArray<int32>& OtherB = CollectionB.AddAttribute<int32>("UserAttr", "Other", TransformDependency);
+		for (int i = 0; i < OtherA.Num();  i++) { OtherA[i] = i;} // count up
+		for (int i = 0; i < OtherB.Num(); i++) { OtherB[i] = i; } // count up
+
+		// What we are testing! - Append
+		CollectionA.Append(CollectionB);
+
+		// validate that the collectionB is at the start of A, and A is indexed correctly
+		int32 Size = CollectionA.NumElements(FTransformCollection::TransformGroup);
+		EXPECT_TRUE(Size==16);
+
+		// validate appended transform hierarchy
+		for (int32 Idx = 0, Sdx = Size/2; Idx < Size/2; Idx++, Sdx++)
+		{
+			EXPECT_TRUE(CollectionA.BoneName[Idx]== CollectionA.BoneName[Sdx]);
+			if (CollectionA.Parent[Idx] == INDEX_NONE)
+				EXPECT_TRUE(CollectionA.Parent[Sdx] == INDEX_NONE);
+			else
+				EXPECT_TRUE(CollectionA.Parent[Idx] == CollectionA.Parent[Sdx] - Size / 2);
+			EXPECT_TRUE(CollectionA.Children[Idx].Num() == CollectionA.Children[Sdx].Num());
+			for (int i : CollectionA.Children[Sdx])
+				EXPECT_TRUE(CollectionA.Children[Idx].Contains(i - Size / 2) );
+
+			// validate dependent user attributes in transform group
+			EXPECT_TRUE(ExtraA[Idx] == ExtraA[Sdx]-Size/2);
+
+		}
+
+		EXPECT_TRUE(CollectionA.NumElements("Other") == 11);
+		for (int32 Idx = 0; Idx < CollectionA.NumElements("Other"); Idx++)
+		{
+			if (Idx <= 5)
+			{
+				EXPECT_TRUE(OtherA[Idx] == OtherB[Idx]);
+			}
+			else
+			{
+				EXPECT_TRUE(OtherA[Idx] == Idx+2);
+			}
+		}
+
+
+		EXPECT_TRUE(true);
+
+	}
+
+
+	void CollectionCycleTest()
+	{
+		//--gtest_filter=GeometryCollection_CreationTest.CollectionCycleTest
+
+		//
+		// Build dependent attributes and check for connections using HasCycle
+		// 
+		//
+		FManagedArrayCollection Collection;
+		Collection.AddGroup("A");
+		Collection.AddGroup("B");
+		Collection.AddGroup("C");
+		Collection.AddGroup("D");
+
+		EXPECT_FALSE(Collection.IsConnected("A", "C"));
+		EXPECT_FALSE(Collection.IsConnected("A", "B"));
+		EXPECT_FALSE(Collection.IsConnected("A", "D"));
+		EXPECT_FALSE(Collection.IsConnected("B", "C"));
+
+		//  A
+		//  |
+		//  |--->B "A depends on B"
+		//  |    |
+		//  |    |--->C "B depends on C"
+		//  | 
+		//  |--->D  "A depends on D
+		FManagedArrayCollection::FConstructionParameters ADependency("A");
+		FManagedArrayCollection::FConstructionParameters BDependency("B");
+		FManagedArrayCollection::FConstructionParameters CDependency("C");
+		FManagedArrayCollection::FConstructionParameters DDependency("D");
+
+		Collection.AddAttribute<int32>("AOnB", "A", BDependency);
+		Collection.AddAttribute<int32>("AOnD", "A", DDependency);
+		Collection.AddAttribute<int32>("BOnC", "B", CDependency);
+		Collection.AddAttribute<int32>("C", "C"); // None
+		Collection.AddAttribute<int32>("D", "D"); // None
+
+		// A cycle would be "C depending on A", which would fail attribute creation.
+		EXPECT_FALSE(Collection.IsConnected("C", "A"));
+
+		// So if we tried to create a dependency on A from C, we would first want to
+		// know if A already connects to C, and throw an error if so. 
+		EXPECT_TRUE(Collection.IsConnected("A","C"));
+
+		EXPECT_TRUE(Collection.IsConnected("A", "B"));
+		EXPECT_TRUE(Collection.IsConnected("A", "D"));
+		EXPECT_TRUE(Collection.IsConnected("B", "C"));
+
+		EXPECT_FALSE(Collection.IsConnected("A", "A"));
+		EXPECT_FALSE(Collection.IsConnected("B", "A"));
+		EXPECT_FALSE(Collection.IsConnected("B","D"));
+		EXPECT_FALSE(Collection.IsConnected("D", "B"));
+		EXPECT_FALSE(Collection.IsConnected("B", "D"));
+		EXPECT_FALSE(Collection.IsConnected("D", "B"));
+		EXPECT_FALSE(Collection.IsConnected("D", "A"));
+	}
+
 }
diff --git a/Engine/Source/Programs/HeadlessChaos/Private/HeadlessChaosApp.cpp b/Engine/Source/Programs/HeadlessChaos/Private/HeadlessChaosApp.cpp
--- a/Engine/Source/Programs/HeadlessChaos/Private/HeadlessChaosApp.cpp
+++ b/Engine/Source/Programs/HeadlessChaos/Private/HeadlessChaosApp.cpp
@@ -317,7 +317,7 @@ TEST(GeometryCollection_MatricesTest,BasicGlobalMatrices) { GeometryCollectionTe
 TEST(GeometryCollection_MatricesTest,TransformMatrixElement) { GeometryCollectionTest::TransformMatrixElement(); SUCCEED(); }
 TEST(GeometryCollection_MatricesTest,ReparentingMatrices) { GeometryCollectionTest::ReparentingMatrices(); SUCCEED(); }
 
-// Creation Tests
+// Creation Tests CollectionCycleTest
 TEST(GeometryCollection_CreationTest, CheckClassTypes) { GeometryCollectionTest::CheckClassTypes(); SUCCEED(); }
 TEST(GeometryCollection_CreationTest, CheckIncrementMask) { GeometryCollectionTest::CheckIncrementMask(); SUCCEED(); }
 TEST(GeometryCollection_CreationTest,Creation) { GeometryCollectionTest::Creation(); SUCCEED(); }
@@ -334,6 +334,9 @@ TEST(GeometryCollection_CreationTest,ReindexMaterialsTest) { GeometryCollectionT
 TEST(GeometryCollection_CreationTest,ContiguousElementsTest) { GeometryCollectionTest::ContiguousElementsTest(); SUCCEED(); }
 TEST(GeometryCollection_CreationTest, AttributeDependencyTest) { GeometryCollectionTest::AttributeDependencyTest(); SUCCEED(); }
 TEST(GeometryCollection_CreationTest, IntListReindexOnDeletionTest) { GeometryCollectionTest::IntListReindexOnDeletionTest(); SUCCEED(); }
+TEST(GeometryCollection_CreationTest, AppendManagedArrayCollectionTest) { GeometryCollectionTest::AppendManagedArrayCollectionTest(); SUCCEED(); }
+TEST(GeometryCollection_CreationTest, AppendTransformCollectionTest) { GeometryCollectionTest::AppendTransformCollectionTest(); SUCCEED(); }
+TEST(GeometryCollection_CreationTest, CollectionCycleTest) { GeometryCollectionTest::CollectionCycleTest(); SUCCEED(); }
 
 
 // Proximity Tests
diff --git a/Engine/Source/Programs/HeadlessChaos/Public/GeometryCollection/GeometryCollectionTestCreation.h b/Engine/Source/Programs/HeadlessChaos/Public/GeometryCollection/GeometryCollectionTestCreation.h
--- a/Engine/Source/Programs/HeadlessChaos/Public/GeometryCollection/GeometryCollectionTestCreation.h
+++ b/Engine/Source/Programs/HeadlessChaos/Public/GeometryCollection/GeometryCollectionTestCreation.h
@@ -39,4 +39,10 @@ namespace GeometryCollectionTest
 
 	void IntListReindexOnDeletionTest();
 
+	void AppendManagedArrayCollectionTest();
+
+	void AppendTransformCollectionTest();
+
+	void CollectionCycleTest();
+
 }
diff --git a/Engine/Source/Runtime/Experimental/Chaos/Private/GeometryCollection/ManagedArrayCollection.cpp b/Engine/Source/Runtime/Experimental/Chaos/Private/GeometryCollection/ManagedArrayCollection.cpp
--- a/Engine/Source/Runtime/Experimental/Chaos/Private/GeometryCollection/ManagedArrayCollection.cpp
+++ b/Engine/Source/Runtime/Experimental/Chaos/Private/GeometryCollection/ManagedArrayCollection.cpp
@@ -194,6 +194,79 @@ int32 FManagedArrayCollection::InsertElements(int32 NumberElements, int32 Positi
 	return Position;
 }
 
+void FManagedArrayCollection::Append(const FManagedArrayCollection& InCollection)
+{
+	bool bMatchingAttributes = true;
+	for (const TTuple<FKeyType, FValueType>& Entry : InCollection.Map)
+	{
+		if (HasAttribute(Entry.Key.Get<0>(), Entry.Key.Get<1>()))
+		{
+			const FValueType& OriginalValue = InCollection.Map[Entry.Key];
+			const FValueType& DestValue = Map[Entry.Key];
+
+			// If we don't have a type match don't attempt the copy.
+			if (OriginalValue.ArrayType != DestValue.ArrayType)
+			{
+				bMatchingAttributes = false;
+				ensureMsgf(false, TEXT("Failed : Type error in FManagedArrayCollection::AppendCollection (%s:%s)"), 
+					*Entry.Key.Get<0>().ToString(), *Entry.Key.Get<1>().ToString());
+			}
+		}
+	}
+	if (bMatchingAttributes)
+	{
+		// make space first. 
+		for (const FName& Group : InCollection.GroupNames())
+		{
+			if (HasGroup(Group) && NumElements(Group))
+			{
+				InsertElements(InCollection.NumElements(Group), 0, Group);
+			}
+			else if (!HasGroup(Group))
+			{
+				AddGroup(Group);
+				AddElements(InCollection.NumElements(Group), Group);
+			}
+		}
+
+		// copy values
+		for (const TTuple<FKeyType, FValueType>& Entry : InCollection.Map)
+		{
+			FName AttributeName = Entry.Key.Get<0>();
+			FName GroupName = Entry.Key.Get<1>();
+
+			if (HasAttribute(AttributeName, GroupName))
+			{
+				Map[Entry.Key].Value->CopyRange(*Entry.Value.Value, 0, Entry.Value.Value->Num());
+			}
+			else
+			{
+				// Copied from FManagedArrayCollection::CopyAttribute, but this 
+				// does not manage destination size, as that is pre allocated above. 
+				FKeyType Key = FManagedArrayCollection::MakeMapKey(AttributeName, GroupName);
+
+				if (!HasAttribute(AttributeName, GroupName))
+				{
+					const FValueType& V = InCollection.Map[Key];
+					EArrayType Type = V.ArrayType;
+					FValueType Value(Type, *NewManagedTypedArray(Type));
+					Value.Value->Resize(NumElements(GroupName));
+					Value.GroupIndexDependency = V.GroupIndexDependency;
+					Value.Saved = V.Saved;
+					Value.bExternalValue = V.bExternalValue;
+					Map.Add(Key, MoveTemp(Value));
+				}
+
+				const FValueType& OriginalValue = InCollection.Map[Key];
+				const FValueType& DestValue = Map[Key];
+				check(OriginalValue.ArrayType == DestValue.ArrayType);
+				DestValue.Value->Init(*OriginalValue.Value);
+			}
+		}
+	}
+}
+
+
 void FManagedArrayCollection::RemoveAttribute(FName Name, FName Group)
 {
 	FKeyType Key = FManagedArrayCollection::MakeMapKey(Name, Group);
@@ -307,7 +380,7 @@ void FManagedArrayCollection::ReorderElements(FName Group, const TArray<int32>&
 void FManagedArrayCollection::SetDependency(FName Name, FName Group, FName DependencyGroup)
 {
 	ensure(HasAttribute(Name, Group));
-	if (ensure(!HasCycle(Group, DependencyGroup)))
+	if (ensure(!IsConnected(DependencyGroup, Group)))
 	{
 		FKeyType Key = FManagedArrayCollection::MakeMapKey(Name, Group);
 		Map[Key].GroupIndexDependency = DependencyGroup;
@@ -400,39 +473,40 @@ void FManagedArrayCollection::CopyAttribute(const FManagedArrayCollection& InCol
 	DestValue.Value->Init(*OriginalValue.Value);
 }
 
-FName FManagedArrayCollection::GetDependency(FName SearchGroup)
+bool FManagedArrayCollection::IsConnected(FName StartingNode, FName TargetNode)
 {
-	FName GroupIndexDependency = "";
-
-	for (const TTuple<FKeyType, FValueType>& Entry : Map)
+	if (!StartingNode.IsNone())
 	{
-		if (Entry.Key.Get<1>() == SearchGroup)
+		TMap<FName, TArray<FName> > DMap;
+		for (const TTuple<FKeyType, FValueType>& Entry : Map)
 		{
-			GroupIndexDependency = Entry.Value.GroupIndexDependency;
+			if (!DMap.Contains(Entry.Key.Get<1>()))
+				DMap.Add(Entry.Key.Get<1>(), TArray<FName>());
+			if (!Entry.Value.GroupIndexDependency.IsNone())
+				DMap[Entry.Key.Get<1>()].AddUnique(Entry.Value.GroupIndexDependency);
 		}
-	}
 
-	return GroupIndexDependency;
-}
-
-bool FManagedArrayCollection::HasCycle(FName NewGroup, FName DependencyGroup)
-{
-	if (!DependencyGroup.IsNone())
-	{
-		// The system relies adding a dependency on it own group in order to run the reinding methods
-		// this is why we don't include the case if (NewGroup == DependencyGroup) return true;
-
-		while (!(DependencyGroup = GetDependency(DependencyGroup)).IsNone())
+		if (DMap.Contains(StartingNode))
 		{
-			// check if we are looping back to the group we are testing against
-			if (DependencyGroup == NewGroup)
+			TSet<FName> Visited;
+			TArray<FName> SearchSet = DMap[StartingNode];
+			while (SearchSet.Num())
 			{
-				return true;
+				FName Curr = SearchSet.Pop();
+				if (Curr.IsEqual(TargetNode))
+				{
+					return true;
+				}
+
+				if (!Visited.Contains(Curr))
+				{
+					Visited.Add(Curr);
+					if (!DMap[Curr].IsEmpty())
+						SearchSet.Append(DMap[Curr]);
+				}
 			}
 		}
-
 	}
-
 	return false;
 }
 
diff --git a/Engine/Source/Runtime/Experimental/Chaos/Private/GeometryCollection/TransformCollection.cpp b/Engine/Source/Runtime/Experimental/Chaos/Private/GeometryCollection/TransformCollection.cpp
--- a/Engine/Source/Runtime/Experimental/Chaos/Private/GeometryCollection/TransformCollection.cpp
+++ b/Engine/Source/Runtime/Experimental/Chaos/Private/GeometryCollection/TransformCollection.cpp
@@ -31,8 +31,6 @@ void FTransformCollection::DefineTransformSchema(FManagedArrayCollection& InColl
 
 void FTransformCollection::Construct()
 {
-	FManagedArrayCollection::FConstructionParameters TransformDependency(FTransformCollection::TransformGroup);
-
 	// Hierarchy Group
 	AddExternalAttribute<FTransform>(FTransformCollection::TransformAttribute, FTransformCollection::TransformGroup, Transform);
 	AddExternalAttribute<FString>("BoneName", FTransformCollection::TransformGroup, BoneName);
@@ -72,6 +70,28 @@ FTransformCollection FTransformCollection::SingleTransform(const FTransform& Tra
 	return TransformCollection;
 }
 
+void FTransformCollection::Append(const FManagedArrayCollection& InCollection)
+{
+	if (const FTransformCollection* InTypedCollection = InCollection.Cast<FTransformCollection>())
+	{
+		int32 Offset = InCollection.NumElements(TransformGroup);
+		Super::Append(InCollection);
+
+		int32 Size = NumElements(TransformGroup);
+		for (int Idx = Offset; Idx < Size; Idx++)
+		{
+			if (Parent[Idx] != INDEX_NONE)
+				Parent[Idx] += Offset;			
+			for (int32& Val : Children[Idx])
+				Val += Offset;
+		}
+	}
+	else
+	{
+		Super::Append(InCollection);
+	}
+}
+
 
 int32 FTransformCollection::AppendTransform(const FTransformCollection & Element, const FTransform& TransformRoot)
 {
diff --git a/Engine/Source/Runtime/Experimental/Chaos/Public/GeometryCollection/ManagedArray.h b/Engine/Source/Runtime/Experimental/Chaos/Public/GeometryCollection/ManagedArray.h
--- a/Engine/Source/Runtime/Experimental/Chaos/Public/GeometryCollection/ManagedArray.h
+++ b/Engine/Source/Runtime/Experimental/Chaos/Public/GeometryCollection/ManagedArray.h
@@ -118,6 +118,12 @@ protected:
 	*/
 	virtual void Init(const FManagedArrayBase& ) {};
 
+	/**
+	* Copy a range of values from the ConstArray into this
+	*/
+	virtual void CopyRange(const FManagedArrayBase& ConstArray, int32 Start, int32 Stop, int32 Offset = 0) {};
+
+
 public:
 	FManagedArrayBase()
 	{
@@ -204,6 +210,10 @@ template <typename T>
 void InitHelper(TArray<T>& Array, const TManagedArrayBase<T>& NewTypedArray, int32 Size);
 template <typename T>
 void InitHelper(TArray<TUniquePtr<T>>& Array, const TManagedArrayBase<TUniquePtr<T>>& NewTypedArray, int32 Size);
+template <typename T>
+void CopyRangeHelper(TArray<T>& Target, const TManagedArrayBase<T>& Source, int32 Start, int32 Stop, int32 Offset);
+template <typename T>
+void CopyRangeHelper(TArray<TUniquePtr<T>>& Array, const TManagedArrayBase<TUniquePtr<T>>& ConstArray, int32 Start, int32 Stop, int32 Offset);
 
 /***
 *  Managed Array
@@ -297,7 +307,6 @@ public:
  		Array.Shrink();
 	}
 
-
 	/**
 	* Init from a predefined Array of matching type
 	*/
@@ -311,6 +320,19 @@ public:
 		InitHelper(Array, NewTypedArray, Size);
 	}
 
+	/**
+	* Copy from a predefined Array of matching type
+	*/
+	virtual void CopyRange(const FManagedArrayBase& ConstArray, int32 Start, int32 Stop, int32 Offset = 0) override
+	{
+		ensureMsgf(ConstArray.GetTypeSize() == GetTypeSize(), TEXT("TManagedArrayBase<T>::Init : Invalid array types."));
+		if (ensureMsgf(Stop + Offset < Array.Num(), TEXT("Error : Index out of bounds")))
+		{
+			const TManagedArrayBase<ElementType>& TypedConstArray = static_cast<const TManagedArrayBase<ElementType>&>(ConstArray);
+			CopyRangeHelper(Array, TypedConstArray, Start, Stop, Offset);
+		}
+	}
+
 	/**
 	 * Fill the array with \p Value.
 	 */
@@ -550,12 +572,32 @@ void InitHelper(TArray<TUniquePtr<T>>& Array, const TManagedArrayBase<TUniquePtr
 	}
 }
 
+template <typename T>
+void CopyRangeHelper(TArray<T>& Target, const TManagedArrayBase<T>& Source, int32 Start, int32 Stop, int32 Offset)
+{
+	for (int32 Sdx = Start, Tdx = Start + Offset; Sdx < Source.Num() && Tdx < Target.Num() && Sdx < Stop; Sdx++, Tdx++)
+	{
+		Target[Tdx] = Source[Sdx];
+	}
+}
+
+template <typename T>
+void CopyRangeHelper(TArray<TUniquePtr<T>>& Target, const TManagedArrayBase<TUniquePtr<T>>& Source, int32 Start, int32 Stop, int32 Offset)
+{
+	for (int32 Sdx = Start, Tdx = Start+Offset; Sdx<Source.Num() && Tdx<Target.Num() && Sdx<Stop; Sdx++, Tdx++)
+	{
+		Target[Tdx].Reset((T*)Source[Sdx]->Copy().Release());
+	}
+}
+
 //
 //
 //
 #define UNSUPPORTED_UNIQUE_ARRAY_COPIES(TYPE, NAME) \
 template<> inline void InitHelper(TArray<TYPE>& Array, const TManagedArrayBase<TYPE>& NewTypedArray, int32 Size) { \
-	UE_LOG(LogChaos,Warning, TEXT("Cannot make a copy of unique array of type (%s) within the managed array collection. Regenerate unique pointer attributes if needed."), NAME); \
+	UE_LOG(LogChaos,Warning, TEXT("Cannot make a copy of unique array of type (%s) within the managed array collection. Regenerate unique pointer attributes if needed."), NAME); }\
+template<> inline void CopyRangeHelper(TArray<TYPE>& Target, const TManagedArrayBase<TYPE>& Source, int32 Start, int32 Stop, int32 Offset) {\
+	UE_LOG(LogChaos, Warning, TEXT("Cannot make a range copy of unique array of type (%s) within the managed array collection. Regenerate unique pointer attributes if needed."), NAME); \
 }
 
 typedef TUniquePtr<Chaos::TGeometryParticle<Chaos::FReal, 3>> LOCAL_MA_UniqueTGeometryParticle;
diff --git a/Engine/Source/Runtime/Experimental/Chaos/Public/GeometryCollection/ManagedArrayCollection.h b/Engine/Source/Runtime/Experimental/Chaos/Public/GeometryCollection/ManagedArrayCollection.h
--- a/Engine/Source/Runtime/Experimental/Chaos/Public/GeometryCollection/ManagedArrayCollection.h
+++ b/Engine/Source/Runtime/Experimental/Chaos/Public/GeometryCollection/ManagedArrayCollection.h
@@ -132,7 +132,7 @@ public:
 			FValueType Value(ManagedArrayType<T>(), *(new TManagedArray<T>()));
 			Value.Value->Resize(NumElements(Group));
 			Value.Saved = Parameters.Saved;
-			if (ensure(!HasCycle(Group, Parameters.GroupIndexDependency)))
+			if (ensure(!IsConnected(Parameters.GroupIndexDependency,Group)))
 			{
 				Value.GroupIndexDependency = Parameters.GroupIndexDependency;
 			}
@@ -199,7 +199,7 @@ public:
 		FValueType Value(ManagedArrayType<T>(), ValueIn);
 		Value.Value->Resize(NumElements(Group));
 		Value.Saved = Parameters.Saved;
-		if (ensure(!HasCycle(Group, Parameters.GroupIndexDependency)))
+		if (ensure(!IsConnected(Parameters.GroupIndexDependency, Group)))
 		{
 			Value.GroupIndexDependency = Parameters.GroupIndexDependency;
 		}
@@ -241,6 +241,13 @@ public:
 	*/
 	int32 InsertElements(int32 NumberElements, int32 Position, FName Group);
 
+
+	/**
+	* Append Collection and reindex dependencies on this collection. 
+	* @param InCollection : Collection to add. 
+	*/
+	virtual void Append(const FManagedArrayCollection& Collection);
+
 	/**
 	* Returns attribute(Name) of Type(T) from the group
 	* @param Name - The name of the attribute
@@ -495,9 +502,16 @@ public:
 
 	/**
 	* Serialize
-		*/
+	*/
 	virtual void Serialize(Chaos::FChaosArchive& Ar);
 
+	/**
+	* Cycle Checking. 
+	* Search for TargetNode from the StartingNode. 
+	*/
+	bool IsConnected(FName StartingNode, FName TargetNode);
+
+
 	/**
 	* Dump the contents to a FString
 	*/
@@ -615,9 +629,6 @@ private:
 	TMap< FName, FGroupInfo> GroupInfo;
 	bool bDirty;
 
-	FName GetDependency(FName SearchGroup);
-	bool HasCycle(FName Group, FName DependencyGroup);
-
 protected:
 
 	/**
diff --git a/Engine/Source/Runtime/Experimental/Chaos/Public/GeometryCollection/TransformCollection.h b/Engine/Source/Runtime/Experimental/Chaos/Public/GeometryCollection/TransformCollection.h
--- a/Engine/Source/Runtime/Experimental/Chaos/Public/GeometryCollection/TransformCollection.h
+++ b/Engine/Source/Runtime/Experimental/Chaos/Public/GeometryCollection/TransformCollection.h
@@ -66,6 +66,12 @@ public:
 	*/
 	static FTransformCollection SingleTransform(const FTransform& TransformRoot = FTransform::Identity);
 
+	/**
+	* Append Collection and reindex dependencies on this collection.
+	* @param InCollection : Collection to add.
+	*/
+	virtual void Append(const FManagedArrayCollection& Collection);
+
 	/*
 	* AppendTransform:
 	*   Append a transform at the end of the collection without
