Fix logic for how kinematic particles can wake islands

diff --git a/Engine/Source/Programs/HeadlessChaos/Private/HeadlessChaosTestGraphEvolution.cpp b/Engine/Source/Programs/HeadlessChaos/Private/HeadlessChaosTestGraphEvolution.cpp
--- a/Engine/Source/Programs/HeadlessChaos/Private/HeadlessChaosTestGraphEvolution.cpp
+++ b/Engine/Source/Programs/HeadlessChaos/Private/HeadlessChaosTestGraphEvolution.cpp
@@ -902,4 +902,128 @@ namespace ChaosTest
 		EXPECT_TRUE(FConstGenericParticleHandle(Test.ParticleHandles[1])->Sleeping());
 		EXPECT_TRUE(Test.ConstraintHandles[0]->IsSleeping());
 	}
+
+	// Test the conditions for a kinematic particle waking an island
+	// If a kinematic is being animated by velocity or by setting a target
+	// position the island should wake but only if the target velocity is
+	// non-zero or the target transform is different from the identity
+	GTEST_TEST(GraphEvolutionTests, TestConstraintGraph_KinematicWakeIslandConditions)
+	{
+		FGraphEvolutionTest Test(4);
+		Test.MakeChain();
+
+		// Set the root of the chain to be kinematic and the rest to be sleeping
+		Test.Evolution.SetParticleObjectState(Test.ParticleHandles[0], EObjectStateType::Kinematic);
+		for (int32 ParticleIndex = 1; ParticleIndex < Test.ParticleHandles.Num(); ++ParticleIndex)
+		{
+			Test.Evolution.SetParticleObjectState(Test.ParticleHandles[ParticleIndex], EObjectStateType::Sleeping);
+		}
+
+		Test.Advance();
+
+		// Expect one sleeping island
+		EXPECT_EQ(Test.IslandManager->NumIslands(), 1);
+		EXPECT_NE(Test.ConstraintHandles[0]->GetConstraintGraphIndex(), INDEX_NONE);
+		EXPECT_NE(Test.ConstraintHandles[1]->GetConstraintGraphIndex(), INDEX_NONE);
+		EXPECT_NE(Test.ConstraintHandles[2]->GetConstraintGraphIndex(), INDEX_NONE);
+		EXPECT_TRUE(Test.IslandManager->GetIsland(0)->IsSleeping());
+
+		// Set to velocity mode and animate
+		FKinematicGeometryParticleHandle* KinematicParticle = Test.ParticleHandles[0]->CastToKinematicParticle();
+		ASSERT_NE(KinematicParticle, nullptr);
+		KinematicParticle->KinematicTarget().SetVelocityMode();
+
+		Test.Advance();
+
+		// Expect one sleeping island as the velocity of the kinematic particle is still zero
+		EXPECT_EQ(Test.IslandManager->NumIslands(), 1);
+		EXPECT_TRUE(Test.IslandManager->GetIsland(0)->IsSleeping());
+
+		KinematicParticle->SetV(FVec3(10.0f, 0.0f, 0.0f));
+
+		Test.Advance();
+
+		// Expect one awake island
+		EXPECT_EQ(Test.IslandManager->NumIslands(), 1);
+		EXPECT_FALSE(Test.IslandManager->GetIsland(0)->IsSleeping());
+
+		// Put all particles back to sleep and now set angular velocity
+		KinematicParticle->SetV(FVec3(0.0f, 0.0f, 0.0f));
+		for (int32 ParticleIndex = 1; ParticleIndex < Test.ParticleHandles.Num(); ++ParticleIndex)
+		{
+			Test.Evolution.SetParticleObjectState(Test.ParticleHandles[ParticleIndex], EObjectStateType::Sleeping);
+		}
+
+		Test.Advance();
+
+		// Check we've put the island back to sleep
+		EXPECT_EQ(Test.IslandManager->NumIslands(), 1);
+		EXPECT_TRUE(Test.IslandManager->GetIsland(0)->IsSleeping());
+
+		// Now set angular velocity. Island should wake
+		KinematicParticle->SetW(FVec3(0.0f, 1.0f, 0.0f));
+
+		Test.Advance();
+
+		EXPECT_EQ(Test.IslandManager->NumIslands(), 1);
+		EXPECT_FALSE(Test.IslandManager->GetIsland(0)->IsSleeping());
+
+		// Put all particles back to sleep
+		KinematicParticle->SetW(FVec3(0.0f, 0.0f, 0.0f));
+		for (int32 ParticleIndex = 1; ParticleIndex < Test.ParticleHandles.Num(); ++ParticleIndex)
+		{
+			Test.Evolution.SetParticleObjectState(Test.ParticleHandles[ParticleIndex], EObjectStateType::Sleeping);
+		}
+
+		Test.Advance();
+
+		// Check we've put the island back to sleep
+		EXPECT_EQ(Test.IslandManager->NumIslands(), 1);
+		EXPECT_TRUE(Test.IslandManager->GetIsland(0)->IsSleeping());
+
+		// Now set to position mode. Initially the island should stay sleeping as the target
+		// transform is the identity
+		FKinematicTarget KinematicTarget;
+		KinematicTarget.SetTargetMode(FRigidTransform3::Identity);
+		KinematicParticle->SetKinematicTarget(KinematicTarget);
+
+		Test.Advance();
+
+		EXPECT_EQ(Test.IslandManager->NumIslands(), 1);
+		EXPECT_TRUE(Test.IslandManager->GetIsland(0)->IsSleeping());
+
+		// Set a non-zero position target. Should cause the island to wake
+		KinematicTarget.SetTargetMode(FVec3(10.0f, 0.0f, 0.0f), FRotation3::Identity);
+		KinematicParticle->SetKinematicTarget(KinematicTarget);
+
+		Test.Advance();
+
+		EXPECT_EQ(Test.IslandManager->NumIslands(), 1);
+		EXPECT_FALSE(Test.IslandManager->GetIsland(0)->IsSleeping());
+
+		// Put all particles back to sleep
+		KinematicTarget.SetTargetMode(FRigidTransform3::Identity);
+		KinematicParticle->SetKinematicTarget(KinematicTarget);
+		for (int32 ParticleIndex = 1; ParticleIndex < Test.ParticleHandles.Num(); ++ParticleIndex)
+		{
+			Test.Evolution.SetParticleObjectState(Test.ParticleHandles[ParticleIndex], EObjectStateType::Sleeping);
+		}
+
+		Test.Advance();
+
+		// Check we've put the island back to sleep
+		EXPECT_EQ(Test.IslandManager->NumIslands(), 1);
+		EXPECT_TRUE(Test.IslandManager->GetIsland(0)->IsSleeping());
+
+		// Set a non-identity rotation target. Should cause the island to wake
+		FRigidTransform3 TargetTransform(FVec3(0.0f, 0.0f, 0.0f), FQuat::MakeFromEuler(FVec3(1.0f, 0.0f, 2.0f)));
+		KinematicTarget.SetTargetMode(TargetTransform);
+		KinematicParticle->SetKinematicTarget(KinematicTarget);
+
+		Test.Advance();
+
+		EXPECT_EQ(Test.IslandManager->NumIslands(), 1);
+		EXPECT_FALSE(Test.IslandManager->GetIsland(0)->IsSleeping());
+
+	}
 }
diff --git a/Engine/Source/Runtime/Experimental/Chaos/Private/Chaos/Island/IslandManager.cpp b/Engine/Source/Runtime/Experimental/Chaos/Private/Chaos/Island/IslandManager.cpp
--- a/Engine/Source/Runtime/Experimental/Chaos/Private/Chaos/Island/IslandManager.cpp
+++ b/Engine/Source/Runtime/Experimental/Chaos/Private/Chaos/Island/IslandManager.cpp
@@ -50,13 +50,22 @@ FORCEINLINE bool IsDynamicParticle(const FGeometryParticleHandle* ParticleHandle
 }
 
 /** Check if a particle is not moving */
-FORCEINLINE bool IsStationaryParticle(const FGeometryParticleHandle* ParticleHandle)
+bool IsStationaryParticle(const FGeometryParticleHandle* ParticleHandle)
 {
 	if (ParticleHandle->ObjectState() == EObjectStateType::Kinematic)
 	{
+		// For kinematic particles check whether their current mode and target will cause the particle to move
 		const FKinematicGeometryParticleHandle* KinematicParticle = ParticleHandle->CastToKinematicParticle();
-		// @todo(chaos): should this also be checking Kinematic W()?
-		return KinematicParticle->V().IsZero() && KinematicParticle->KinematicTarget().GetMode() == EKinematicTargetMode::None;
+		if (KinematicParticle->KinematicTarget().GetMode() == EKinematicTargetMode::Position)
+		{
+			return (KinematicParticle->X() - KinematicParticle->KinematicTarget().GetTargetPosition()).IsZero()
+					&& (KinematicParticle->R() * KinematicParticle->KinematicTarget().GetTargetRotation().Inverse()).IsIdentity();
+		}
+		// For all other modes (Velocity, None and Reset) check that the velocity is non-zero
+		else
+		{
+			return KinematicParticle->V().IsZero() && KinematicParticle->W().IsZero();
+		}
 	}
 	else
 	{
