Chaos: Fixed rigid bodies no longer colliding after being constrained with disable collision. After the constraint breaks we now update the ignore manager for the bodies in question to re-enable their interaction

diff --git a/Engine/Source/Programs/HeadlessChaos/Private/HeadlessChaosTestIgnoreManager.cpp b/Engine/Source/Programs/HeadlessChaos/Private/HeadlessChaosTestIgnoreManager.cpp
new file mode 100644
--- /dev/null
+++ b/Engine/Source/Programs/HeadlessChaos/Private/HeadlessChaosTestIgnoreManager.cpp
@@ -0,0 +1,270 @@
+// Copyright Epic Games, Inc. All Rights Reserved.
+
+#include "Chaos/ContactModification.h"
+#include "ChaosSolversModule.h"
+#include "HeadlessChaosTestUtility.h"
+#include "PBDRigidsSolver.h"
+#include "PhysicsProxy/SingleParticlePhysicsProxy.h"
+#include "Algo/Count.h"
+
+namespace ChaosTest
+{
+	using namespace Chaos;
+
+	class FCollisionReportingCallback : public TSimCallbackObject<FSimCallbackNoInput, FSimCallbackNoOutput>
+	{
+	public:
+		TArray<TTuple<FUniqueIdx, FUniqueIdx>> CollisionPairs;
+
+	private:
+		virtual void OnPreSimulate_Internal() override {}
+
+		virtual void OnContactModification_Internal(Chaos::FCollisionContactModifier& Modifier) override
+		{
+			for(FContactPairModifier& PairModifier : Modifier)
+			{
+				TVec2<FGeometryParticleHandle*> Particles = PairModifier.GetParticlePair();
+
+				if(Particles[0] && Particles[1])
+				{
+					FUniqueIdx Ids[2] = {Particles[0]->UniqueIdx(), Particles[1]->UniqueIdx()};
+					if(Ids[1] < Ids[0])
+					{
+						Swap(Ids[0], Ids[1]);
+					}
+					CollisionPairs.AddUnique({ Ids[0], Ids[1] });
+				}
+			}
+		}
+	};
+
+	void FlushSolver(FPBDRigidsSolver* InSolver)
+	{
+		InSolver->AdvanceAndDispatch_External(0);
+		InSolver->WaitOnPendingTasks_External();
+
+		// Populate the spacial acceleration
+		FPBDRigidsSolver::FPBDRigidsEvolution* Evolution = InSolver->GetEvolution();
+
+		if(Evolution)
+		{
+			Evolution->FlushSpatialAcceleration();
+		}
+	}
+
+	void RunSolver(FPBDRigidsSolver* InSolver, float InDt, int32 NumSteps)
+	{
+		for(int Step = 0; Step < NumSteps; ++Step)
+		{
+			InSolver->AdvanceAndDispatch_External(InDt);
+			InSolver->UpdateGameThreadStructures();
+		}
+	}
+
+	GTEST_TEST(AllTraits, IgnoreManager_IgnoresCollisions)
+	{
+		FChaosSolversModule* Module = FChaosSolversModule::GetModule();
+		auto* Solver = Module->CreateSolver(nullptr, /*AsyncDt=*/-1);
+		InitSolverSettings(Solver);
+		Solver->SetThreadingMode_External(EThreadingModeTemp::SingleThread);
+
+		// create a static floor and two boxes falling onto it.
+		// One box has collisions ignored and should fall through, one should collide.
+
+		// simulated cube with downward velocity,should collide with floor and not fall through.
+		FSingleParticlePhysicsProxy* CollidingCubeProxy = FSingleParticlePhysicsProxy::Create(Chaos::FPBDRigidParticle::CreateParticle());
+		FRigidBodyHandle_External& CollidingCubeParticle = CollidingCubeProxy->GetGameThreadAPI();
+		TSharedPtr<FImplicitObject, ESPMode::ThreadSafe> CollidingCubeGeom = TSharedPtr<FImplicitObject, ESPMode::ThreadSafe>(new TBox<FReal, 3>(FVec3(-100), FVec3(100)));
+		CollidingCubeParticle.SetGeometry(CollidingCubeGeom);
+		Solver->RegisterObject(CollidingCubeProxy);
+		CollidingCubeParticle.SetGravityEnabled(false);
+		CollidingCubeParticle.SetV(FVec3(0, 0, -100));
+		CollidingCubeParticle.SetX(FVec3(200, 0, 500));
+		SetCubeInertiaTensor(CollidingCubeParticle, /*Dimension=*/200, /*Mass=*/1);
+		ChaosTest::SetParticleSimDataToCollide({ CollidingCubeProxy->GetParticle_LowLevel() });
+
+		// Simulated cube with downward velocity, collision ignore disables collision with floor, should fall through.
+		FSingleParticlePhysicsProxy* NonCollidableProxy = FSingleParticlePhysicsProxy::Create(Chaos::FPBDRigidParticle::CreateParticle());
+		FRigidBodyHandle_External& NonCollidableParticle = NonCollidableProxy->GetGameThreadAPI();
+		TSharedPtr<FImplicitObject, ESPMode::ThreadSafe> NonCollidableCubeGeom = TSharedPtr<FImplicitObject, ESPMode::ThreadSafe>(new TBox<FReal, 3>(FVec3(-100), FVec3(100)));
+		NonCollidableParticle.SetGeometry(NonCollidableCubeGeom);
+		Solver->RegisterObject(NonCollidableProxy);
+		NonCollidableParticle.SetGravityEnabled(false);
+		NonCollidableParticle.SetV(FVec3(0, 0, -100));
+		NonCollidableParticle.SetX(FVec3(-200, 0, 500));
+		SetCubeInertiaTensor(NonCollidableParticle, /*Dimension=*/200, /*Mass=*/1);
+		ChaosTest::SetParticleSimDataToCollide({ NonCollidableProxy->GetParticle_LowLevel() });
+
+		// static floor at origin, occupying Z = [-100,0]
+		FSingleParticlePhysicsProxy* FloorProxy = FSingleParticlePhysicsProxy::Create(Chaos::FGeometryParticle::CreateParticle());
+		FRigidBodyHandle_External& FloorParticle = FloorProxy->GetGameThreadAPI();
+		TSharedPtr<FImplicitObject, ESPMode::ThreadSafe> FloorGeom = TSharedPtr<FImplicitObject, ESPMode::ThreadSafe>(new TBox<FReal, 3>(FVec3(-500, -500, -100), FVec3(500, 500, 0)));
+		FloorParticle.SetGeometry(FloorGeom);
+		Solver->RegisterObject(FloorProxy);
+		FloorParticle.SetX(FVec3(0, 0, 0));
+		ChaosTest::SetParticleSimDataToCollide({ FloorProxy->GetParticle_LowLevel() });
+
+		FCollisionReportingCallback* Callback = Solver->CreateAndRegisterSimCallbackObject_External<FCollisionReportingCallback>(true);
+
+		// Named Ids for the particles
+		const FUniqueIdx CollidableId = CollidingCubeParticle.UniqueIdx();
+		const FUniqueIdx NonCollidableId = NonCollidableParticle.UniqueIdx();
+		const FUniqueIdx FloorId = FloorParticle.UniqueIdx();
+
+		// Set up the ignore manager
+		FIgnoreCollisionManager& IgnoreManager = Solver->GetEvolution()->GetBroadPhase().GetIgnoreCollisionManager();
+		IgnoreManager.AddIgnoreCollisionsFor(NonCollidableId, FloorId);
+
+		// Flush the solver
+		FlushSolver(Solver);
+
+		// Tell the broadphase we're interested in checking the ignore manager for this object
+		NonCollidableProxy->GetHandle_LowLevel()->CastToRigidParticle()->AddCollisionConstraintFlag(ECollisionConstraintFlags::CCF_BroadPhaseIgnoreCollisions);
+
+		RunSolver(Solver, 1.0f, 10);
+
+		auto CountPred = [Callback](FUniqueIdx Idx) -> int32
+		{
+			return Algo::CountIf(Callback->CollisionPairs, [Idx](const TTuple<FUniqueIdx, FUniqueIdx>& Hit)
+			{
+				return Hit.Get<0>() == Idx || Hit.Get<1>() == Idx;
+			});
+		};
+
+		TMap<FUniqueIdx, int32> Hits;
+		Hits.Add({ CollidableId, CountPred(CollidableId) });
+		Hits.Add({ NonCollidableId, CountPred(NonCollidableId) });
+		Hits.Add({ FloorId, CountPred(FloorId) });
+
+		EXPECT_EQ(Hits[CollidableId], 1);		// Collidable hit the floor
+		EXPECT_EQ(Hits[FloorId], 1);			// Collidable hit the floor
+		EXPECT_EQ(Hits[NonCollidableId], 0);	// Non collidable ignored all collisions with the floor
+
+		Solver->UnregisterAndFreeSimCallbackObject_External(Callback);
+		Solver->UnregisterObject(CollidingCubeProxy);
+		Solver->UnregisterObject(NonCollidableProxy);
+		Solver->UnregisterObject(FloorProxy);
+		Module->DestroySolver(Solver);
+	}
+
+	GTEST_TEST(AllTraits, IgnoreManager_IgnoresCollisions_MultipleSources)
+	{
+		FChaosSolversModule* Module = FChaosSolversModule::GetModule();
+		auto* Solver = Module->CreateSolver(nullptr, /*AsyncDt=*/-1);
+		InitSolverSettings(Solver);
+		Solver->SetThreadingMode_External(EThreadingModeTemp::SingleThread);
+
+		// create a static floor and two boxes falling onto it.
+		// One box has collisions ignored and should fall through, one should collide.
+
+		// simulated cube with downward velocity,should collide with floor and not fall through.
+		FSingleParticlePhysicsProxy* CollidingCubeProxy = FSingleParticlePhysicsProxy::Create(Chaos::FPBDRigidParticle::CreateParticle());
+		FRigidBodyHandle_External& CollidingCubeParticle = CollidingCubeProxy->GetGameThreadAPI();
+		TSharedPtr<FImplicitObject, ESPMode::ThreadSafe> CollidingCubeGeom = TSharedPtr<FImplicitObject, ESPMode::ThreadSafe>(new TBox<FReal, 3>(FVec3(-100), FVec3(100)));
+		CollidingCubeParticle.SetGeometry(CollidingCubeGeom);
+		Solver->RegisterObject(CollidingCubeProxy);
+		CollidingCubeParticle.SetGravityEnabled(false);
+		CollidingCubeParticle.SetV(FVec3(0, 0, -100));
+		CollidingCubeParticle.SetX(FVec3(200, 0, 500));
+		SetCubeInertiaTensor(CollidingCubeParticle, /*Dimension=*/200, /*Mass=*/1);
+		ChaosTest::SetParticleSimDataToCollide({ CollidingCubeProxy->GetParticle_LowLevel() });
+
+		// Simulated cube with downward velocity, collision ignore disables collision with floor, should fall through.
+		FSingleParticlePhysicsProxy* NonCollidableProxy = FSingleParticlePhysicsProxy::Create(Chaos::FPBDRigidParticle::CreateParticle());
+		FRigidBodyHandle_External& NonCollidableParticle = NonCollidableProxy->GetGameThreadAPI();
+		TSharedPtr<FImplicitObject, ESPMode::ThreadSafe> NonCollidableCubeGeom = TSharedPtr<FImplicitObject, ESPMode::ThreadSafe>(new TBox<FReal, 3>(FVec3(-100), FVec3(100)));
+		NonCollidableParticle.SetGeometry(NonCollidableCubeGeom);
+		Solver->RegisterObject(NonCollidableProxy);
+		NonCollidableParticle.SetGravityEnabled(false);
+		NonCollidableParticle.SetV(FVec3(0, 0, -10));
+		NonCollidableParticle.SetX(FVec3(-200, 0, 110));
+		SetCubeInertiaTensor(NonCollidableParticle, /*Dimension=*/200, /*Mass=*/1);
+		ChaosTest::SetParticleSimDataToCollide({ NonCollidableProxy->GetParticle_LowLevel() });
+
+		// static floor at origin, occupying Z = [-100,0]
+		FSingleParticlePhysicsProxy* FloorProxy = FSingleParticlePhysicsProxy::Create(Chaos::FGeometryParticle::CreateParticle());
+		FRigidBodyHandle_External& FloorParticle = FloorProxy->GetGameThreadAPI();
+		TSharedPtr<FImplicitObject, ESPMode::ThreadSafe> FloorGeom = TSharedPtr<FImplicitObject, ESPMode::ThreadSafe>(new TBox<FReal, 3>(FVec3(-500, -500, -100), FVec3(500, 500, 0)));
+		FloorParticle.SetGeometry(FloorGeom);
+		Solver->RegisterObject(FloorProxy);
+		FloorParticle.SetX(FVec3(0, 0, 0));
+		ChaosTest::SetParticleSimDataToCollide({ FloorProxy->GetParticle_LowLevel() });
+
+		FCollisionReportingCallback* Callback = Solver->CreateAndRegisterSimCallbackObject_External<FCollisionReportingCallback>(true);
+
+		// Named Ids for the particles
+		const FUniqueIdx CollidableId = CollidingCubeParticle.UniqueIdx();
+		const FUniqueIdx NonCollidableId = NonCollidableParticle.UniqueIdx();
+		const FUniqueIdx FloorId = FloorParticle.UniqueIdx();
+
+		// Set up the ignore manager (Twice as if two systems requested it)
+		FIgnoreCollisionManager& IgnoreManager = Solver->GetEvolution()->GetBroadPhase().GetIgnoreCollisionManager();
+		IgnoreManager.AddIgnoreCollisionsFor(NonCollidableId, FloorId);
+		IgnoreManager.AddIgnoreCollisionsFor(NonCollidableId, FloorId);
+
+		// Flush the solver
+		FlushSolver(Solver);
+
+		// Tell the broadphase we're interested in checking the ignore manager for this object
+		NonCollidableProxy->GetHandle_LowLevel()->CastToRigidParticle()->AddCollisionConstraintFlag(ECollisionConstraintFlags::CCF_BroadPhaseIgnoreCollisions);
+
+		const float Dt = 1.0f;
+		const int32 Steps = 10;
+		for(int Step = 0; Step < Steps; ++Step)
+		{
+			Solver->AdvanceAndDispatch_External(Dt);
+			Solver->UpdateGameThreadStructures();
+		}
+
+		auto CountPred = [Callback](FUniqueIdx Idx) -> int32
+		{
+			return Algo::CountIf(Callback->CollisionPairs, [Idx](const TTuple<FUniqueIdx, FUniqueIdx>& Hit)
+			{
+				return Hit.Get<0>() == Idx || Hit.Get<1>() == Idx;
+			});
+		};
+
+		TMap<FUniqueIdx, int32> Hits;
+		Hits.Add({ CollidableId, CountPred(CollidableId) });
+		Hits.Add({ NonCollidableId, CountPred(NonCollidableId) });
+		Hits.Add({ FloorId, CountPred(FloorId) });
+
+		EXPECT_EQ(Hits[CollidableId], 1);		// Collidable hit the floor
+		EXPECT_EQ(Hits[FloorId], 1);			// Collidable hit the floor
+		EXPECT_EQ(Hits[NonCollidableId], 0);	// Non collidable ignored all collisions with the floor
+
+		// Remove one source
+		IgnoreManager.RemoveIgnoreCollisionsFor(NonCollidableId, FloorId);
+
+		Callback->CollisionPairs.Reset();
+		RunSolver(Solver, 1.0f, 10);
+
+		Hits[CollidableId] = CountPred(CollidableId);
+		Hits[NonCollidableId] = CountPred(NonCollidableId);
+		Hits[FloorId] = CountPred(FloorId);
+
+		EXPECT_EQ(Hits[CollidableId], 1);		// Collidable hit the floor
+		EXPECT_EQ(Hits[FloorId], 1);			// Collidable hit the floor
+		EXPECT_EQ(Hits[NonCollidableId], 0);	// Non collidable ignored all collisions with the floor
+
+		// Remove final source
+		IgnoreManager.RemoveIgnoreCollisionsFor(NonCollidableId, FloorId);
+
+		Callback->CollisionPairs.Reset();
+		RunSolver(Solver, 1.0f, 10);
+
+		Hits[CollidableId] = CountPred(CollidableId);
+		Hits[NonCollidableId] = CountPred(NonCollidableId);
+		Hits[FloorId] = CountPred(FloorId);
+
+		EXPECT_EQ(Hits[CollidableId], 1);		// Collidable hit the floor
+		EXPECT_EQ(Hits[FloorId], 2);			// Collidable and NonCollidable hit the floor
+		EXPECT_EQ(Hits[NonCollidableId], 1);	// Non collidable is now collidable and hits the floor
+
+		Solver->UnregisterAndFreeSimCallbackObject_External(Callback);
+		Solver->UnregisterObject(CollidingCubeProxy);
+		Solver->UnregisterObject(NonCollidableProxy);
+		Solver->UnregisterObject(FloorProxy);
+		Module->DestroySolver(Solver);
+	}
+}
\ No newline at end of file
diff --git a/Engine/Source/Runtime/Experimental/Chaos/Private/Chaos/Collision/CollisionConstraintFlags.cpp b/Engine/Source/Runtime/Experimental/Chaos/Private/Chaos/Collision/CollisionConstraintFlags.cpp
--- a/Engine/Source/Runtime/Experimental/Chaos/Private/Chaos/Collision/CollisionConstraintFlags.cpp
+++ b/Engine/Source/Runtime/Experimental/Chaos/Private/Chaos/Collision/CollisionConstraintFlags.cpp
@@ -7,44 +7,83 @@
 namespace Chaos
 {
 
-	bool FIgnoreCollisionManager::ContainsHandle(FHandleID Body0)
+	bool FIgnoreCollisionManager::ContainsHandle(FHandleID Body0) const
 	{
 		return IgnoreCollisionsList.Contains(Body0);
 	}
 
-	bool FIgnoreCollisionManager::IgnoresCollision(FHandleID Body0, FHandleID Body1)
+	bool FIgnoreCollisionManager::IgnoresCollision(FHandleID Body0, FHandleID Body1) const
 	{
-		if (IgnoreCollisionsList.Contains(Body0))
+		const TArray<FIgnoreEntry>* Entries = IgnoreCollisionsList.Find(Body0);
+
+		if(Entries)
 		{
-			return IgnoreCollisionsList[Body0].Contains(Body1);
+			return Entries->ContainsByPredicate([&Body1](const FIgnoreEntry& Entry)
+			{
+				return Entry.Id == Body1;
+			});
 		}
+
 		return false;
 	}
 
-	int32 FIgnoreCollisionManager::NumIgnoredCollision(FHandleID Body0)
+	int32 FIgnoreCollisionManager::NumIgnoredCollision(FHandleID Body0) const
 	{
-		if (IgnoreCollisionsList.Contains(Body0))
-		{
-			return IgnoreCollisionsList[Body0].Num();
-		}
-		return 0;
+		const TArray<FIgnoreEntry>* Entries = IgnoreCollisionsList.Find(Body0);
+
+		return Entries ? Entries->Num() : 0;
 	}
 
 	void FIgnoreCollisionManager::AddIgnoreCollisionsFor(FHandleID Body0, FHandleID Body1)
 	{
-		if (!IgnoreCollisionsList.Contains(Body0))
+		TArray<FIgnoreEntry>& Entries = IgnoreCollisionsList.FindOrAdd(Body0);
+		FIgnoreEntry* Entry = Entries.FindByPredicate([&Body1](const FIgnoreEntry& Entry)
 		{
-			IgnoreCollisionsList.Add(Body0, TArray<FHandleID>());
-		}
-		IgnoreCollisionsList[Body0].Add(Body1);
+			return Entry.Id == Body1;
+		});
 
+		if(Entry)
+		{
+			Entry->Count++;
+		}
+		else
+		{
+			Entries.Add(FIgnoreEntry(Body1));
+		}
 	}
-	void FIgnoreCollisionManager::RemoveIgnoreCollisionsFor(FHandleID Body0, FHandleID Body1)
+
+	int32 FIgnoreCollisionManager::RemoveIgnoreCollisionsFor(FHandleID Body0, FHandleID Body1)
 	{
-		if (IgnoreCollisionsList.Contains(Body0))
+		TArray<FIgnoreEntry>* Entries = IgnoreCollisionsList.Find(Body0);
+
+		if(Entries)
 		{
-			IgnoreCollisionsList[Body0].Remove(Body1);
+			int32 EntryIndex = Entries->IndexOfByPredicate([&Body1](const FIgnoreEntry& FindEntry)
+			{
+				return FindEntry.Id == Body1;
+			});
+
+			if(EntryIndex != INDEX_NONE)
+			{
+				(*Entries)[EntryIndex].Count--;
+
+				if((*Entries)[EntryIndex].Count <= 0)
+				{
+					Entries->RemoveAtSwap(EntryIndex);
+				}
+			}
+
+			if(Entries->Num() == 0)
+			{
+				IgnoreCollisionsList.Remove(Body0);
+			}
+			else
+			{
+				return Entries->Num();
+			}
 		}
+
+		return 0;
 	}
 
 	void FIgnoreCollisionManager::PopStorageData_Internal(int32 ExternalTimestamp)
@@ -91,7 +130,8 @@ namespace Chaos
 		if (PendingActivations.Num() && PendingDeactivations.Num())
 		{
 			TArray<FHandleID> DeletionList;
-			for (auto& Elem : PendingActivations)
+			
+			for (const TPair<FHandleID, TArray<FHandleID>>& Elem : PendingActivations)
 			{
 				if (PendingDeactivations.Remove(Elem.Key))
 				{
@@ -99,8 +139,11 @@ namespace Chaos
 					PreculledParticles.Add(Elem.Key);
 				}
 			}
-			for (FHandleID Del : DeletionList)
+
+			for(FHandleID Del : DeletionList)
+			{
 				PendingActivations.Remove(Del);
+			}
 		}
 
 		// add collision relationships for particles that have valid
@@ -109,7 +152,7 @@ namespace Chaos
 		if (PendingActivations.Num())
 		{
 			TArray<FHandleID> DeletionList;
-			for (auto& Elem : PendingActivations)
+			for (TPair<FHandleID, TArray<FHandleID>>& Elem : PendingActivations)
 			{
 				for (int Index = Elem.Value.Num() - 1; Index >= 0; Index--)
 				{
@@ -121,20 +164,24 @@ namespace Chaos
 					{
 						FUniqueIdx ID0 = Elem.Key;
 						FUniqueIdx ID1 = Elem.Value[Index];
-						if (!IgnoresCollision(ID0, ID1))
-						{
-							AddIgnoreCollisionsFor(ID0, ID1);
-							AddIgnoreCollisionsFor(ID1, ID0);
-						}
+						
+						AddIgnoreCollisionsFor(ID0, ID1);
+						AddIgnoreCollisionsFor(ID1, ID0);
 
 						Elem.Value.RemoveAtSwap(Index, 1);
 					}
 				}
-				if (!Elem.Value.Num())
+
+				if(!Elem.Value.Num())
+				{
 					DeletionList.Add(Elem.Key);
+				}
 			}
-			for (FHandleID Del : DeletionList)
+
+			for(FHandleID Del : DeletionList)
+			{
 				PendingActivations.Remove(Del);
+			}
 		}
 
 		// remove relationships that exist and have been initialized. 
diff --git a/Engine/Source/Runtime/Experimental/Chaos/Private/PhysicsProxy/JointConstraintProxy.cpp b/Engine/Source/Runtime/Experimental/Chaos/Private/PhysicsProxy/JointConstraintProxy.cpp
--- a/Engine/Source/Runtime/Experimental/Chaos/Private/PhysicsProxy/JointConstraintProxy.cpp
+++ b/Engine/Source/Runtime/Experimental/Chaos/Private/PhysicsProxy/JointConstraintProxy.cpp
@@ -109,6 +109,43 @@ void FJointConstraintPhysicsProxy::DestroyOnPhysicsThread(FPBDRigidsSolver* InSo
 	{
 		TVec2<FGeometryParticleHandle*> Particles = Constraint_PT->GetConstrainedParticles();
 		
+		const FPBDJointSettings Settings = Constraint_PT->GetSettings();
+		const bool bValidPair = InSolver && Particles[0] && Particles[1];
+
+		// If this constraint disables collisions - need to restore the collisions after we destroy the constraint
+		if(bValidPair && !Settings.bCollisionEnabled)
+		{
+			FPBDRigidParticleHandle* Rigid0 = Particles[0]->CastToRigidParticle();
+			FPBDRigidParticleHandle* Rigid1 = Particles[1]->CastToRigidParticle();
+
+			if(Rigid0 || Rigid1)
+			{
+				// At least one rigid so collisions are possible
+				const FUniqueIdx Id0 = Particles[0]->UniqueIdx();
+				const FUniqueIdx Id1 = Particles[1]->UniqueIdx();
+
+				FIgnoreCollisionManager& IgnoreManager = InSolver->GetEvolution()->GetBroadPhase().GetIgnoreCollisionManager();
+
+				// Remove the ignore flags for rigid particles. If we end up with no more ignore entries for this particle
+				// then remove the broadphase flag to no longer check the ignore manager for this particles.
+				if(Rigid0)
+				{
+					if(IgnoreManager.RemoveIgnoreCollisionsFor(Id0, Id1) == 0)
+					{
+						Rigid0->RemoveCollisionConstraintFlag(ECollisionConstraintFlags::CCF_BroadPhaseIgnoreCollisions);
+					}
+				}
+
+				if(Rigid1)
+				{
+					if(IgnoreManager.RemoveIgnoreCollisionsFor(Id1, Id0) == 0)
+					{
+						Rigid1->RemoveCollisionConstraintFlag(ECollisionConstraintFlags::CCF_BroadPhaseIgnoreCollisions);
+					}
+				}
+			}
+		}
+
 		// Ensure that our connected particles are aware that this constraint no longer exists
 		if(Particles[0])
 		{
@@ -184,8 +221,10 @@ void FJointConstraintPhysicsProxy::PushStateOnPhysicsThread(FPBDRigidsSolver* In
 						{
 							if(JointSettingsBuffer.bCollisionEnabled)
 							{
-								//can't remove collision constraint flag because we may still be ignoring collision for others
-								IgnoreCollisionManager.RemoveIgnoreCollisionsFor(ID0, ID1);
+								if(IgnoreCollisionManager.RemoveIgnoreCollisionsFor(ID0, ID1) == 0)
+								{
+									RigidHandle0->RemoveCollisionConstraintFlag(ECollisionConstraintFlags::CCF_BroadPhaseIgnoreCollisions);
+								}
 							}
 							else
 							{
@@ -198,8 +237,10 @@ void FJointConstraintPhysicsProxy::PushStateOnPhysicsThread(FPBDRigidsSolver* In
 						{
 							if (JointSettingsBuffer.bCollisionEnabled)
 							{
-								//can't remove collision constraint flag because we may still be ignoring collision for others
-								IgnoreCollisionManager.RemoveIgnoreCollisionsFor(ID1, ID0);
+								if(IgnoreCollisionManager.RemoveIgnoreCollisionsFor(ID1, ID0) == 0)
+								{
+									RigidHandle1->RemoveCollisionConstraintFlag(ECollisionConstraintFlags::CCF_BroadPhaseIgnoreCollisions);
+								}
 							}
 							else
 							{
diff --git a/Engine/Source/Runtime/Experimental/Chaos/Public/Chaos/Collision/CollisionConstraintFlags.h b/Engine/Source/Runtime/Experimental/Chaos/Public/Chaos/Collision/CollisionConstraintFlags.h
--- a/Engine/Source/Runtime/Experimental/Chaos/Public/Chaos/Collision/CollisionConstraintFlags.h
+++ b/Engine/Source/Runtime/Experimental/Chaos/Public/Chaos/Collision/CollisionConstraintFlags.h
@@ -30,8 +30,9 @@ namespace Chaos
 	public:
 		using FHandleID = FUniqueIdx;
 		using FDeactivationSet = TSet<FUniqueIdx>;
-		using FActiveMap = TMap<FHandleID, TArray<FHandleID> >;
-		using FPendingMap = TMap<FHandleID, TArray<FHandleID> >;
+		using FActiveMap UE_DEPRECATED(5.1, "This type is no longer used") = TMap<FHandleID, TArray<FHandleID>>;
+		using FPendingMap = TMap<FHandleID, TArray<FHandleID>>;
+
 		struct FStorageData
 		{
 			FPendingMap PendingActivations;
@@ -52,15 +53,22 @@ namespace Chaos
 			StorageDataProducer = GetNewStorageData();
 		}
 
-		bool ContainsHandle(FHandleID Body0);
+		bool ContainsHandle(FHandleID Body0) const;
 
-		bool IgnoresCollision(FHandleID Body0, FHandleID Body1);
+		bool IgnoresCollision(FHandleID Body0, FHandleID Body1) const;
 
-		int32 NumIgnoredCollision(FHandleID Body0);
+		int32 NumIgnoredCollision(FHandleID Body0) const;
 
 		void AddIgnoreCollisionsFor(FHandleID Body0, FHandleID Body1);
 
-		void RemoveIgnoreCollisionsFor(FHandleID Body0, FHandleID Body1);
+		/**
+		 * Remove an ignore entry for collisions of Body0->Body1
+		 * Note, a reversed ignore entry could exist
+		 * @param Body0 First body in the collision
+		 * @param Body1 Second body in the collision
+		 * @return Number of collisions Body0 ignores after the removal
+		 */
+		int32 RemoveIgnoreCollisionsFor(FHandleID Body0, FHandleID Body1);
 
 		FPendingMap& GetPendingActivationsForGameThread(int32 ExternalTimestamp) 
 		{
@@ -112,6 +120,21 @@ namespace Chaos
 
 	private:
 
+		// Multiple sources can request an ignore pair, so we handle a simple count of
+		// ignore requests to ensure we don't prematurely remove an ignore pair
+		// #CHAOSTODO replace bi-directional map to pair hash to reduce required storage and dependency
+		struct FIgnoreEntry
+		{
+			FIgnoreEntry() = delete;
+			explicit FIgnoreEntry(FHandleID InId)
+				: Id(InId)
+				, Count(1) // Begin at 1 for convenience when adding entries
+			{}
+
+			FHandleID Id;
+			int32 Count;
+		};
+
 		FStorageData* GetNewStorageData()
 		{
 			FStorageData* StorageData;
@@ -130,7 +153,9 @@ namespace Chaos
 			StorageDataFreePool.Enqueue(InStorageData);
 		}
 
-		FActiveMap IgnoreCollisionsList;
+		// Maps collision body0 to a list of other bodies and a count for how many
+		// sources have requested for this pair to be ignored
+		TMap<FHandleID, TArray<FIgnoreEntry>> IgnoreCollisionsList;
 
 		FPendingMap PendingActivations;
 		FDeactivationSet PendingDeactivations;
