Chaos - fix HeadlessChaos
- add a tolerance to the mid-point rejection (single-sided tri mesh collision)
- was failing in Release builds on the exact alignment test because the capsule axis is different in Release versus Debug

diff --git a/Engine/Source/Programs/HeadlessChaos/Private/HeadlessChaosTestCapsuleTriangleCollision.cpp b/Engine/Source/Programs/HeadlessChaos/Private/HeadlessChaosTestCapsuleTriangleCollision.cpp
--- a/Engine/Source/Programs/HeadlessChaos/Private/HeadlessChaosTestCapsuleTriangleCollision.cpp
+++ b/Engine/Source/Programs/HeadlessChaos/Private/HeadlessChaosTestCapsuleTriangleCollision.cpp
@@ -37,7 +37,7 @@ namespace ChaosTest
 	//
 	//  ________
 	//  \       |
-	//   \ ==== |
+	//   \ +==+ |
 	//    \     |
 	//     \    |
 	//      \   |
@@ -109,7 +109,7 @@ namespace ChaosTest
 	//
 	//  ________
 	//  \       |
-	//   \   ======
+	//   \   +==+==
 	//    \     |
 	//     \    |
 	//      \   |
@@ -183,10 +183,10 @@ namespace ChaosTest
 	// Capsule at a 45 degree angle to the triangle, and the capsule axis passes right through an edge.
 	// We should get a single contact pointing out of the triangle at 45 degrees.
 	// 
-	//     \.\
-	//      \.\
+	//     \ \
+	//      \ \
 	//  -----\+\
-	//        \.\
+	//        \ \
 	//
 	GTEST_TEST(CapsuleTriangleTests, TestEdgeZeroSeparation)
 	{
diff --git a/Engine/Source/Runtime/Experimental/Chaos/Private/Chaos/Collision/CapsuleTriangleContactPoint.cpp b/Engine/Source/Runtime/Experimental/Chaos/Private/Chaos/Collision/CapsuleTriangleContactPoint.cpp
--- a/Engine/Source/Runtime/Experimental/Chaos/Private/Chaos/Collision/CapsuleTriangleContactPoint.cpp
+++ b/Engine/Source/Runtime/Experimental/Chaos/Private/Chaos/Collision/CapsuleTriangleContactPoint.cpp
@@ -98,7 +98,7 @@ namespace Chaos
 
 		// Reject if the middle of the capsule is inside the face (single-sided collision)
 		const FReal FaceDMid = FReal(0.5) * (FaceD0 + FaceD1);
-		if (FaceDMid < FReal(0))
+		if (FaceDMid < -FReal(DistanceTolerance))
 		{
 			// Far inside triangle
 			return;
