Fix new issues with probe constraints
* Don't add probe constraints to constraint graph
* Remove probe constraints from constraint graph
* Don't account for CCD probe constraints in CCD post-solve fixup
* Don't update CCD constraint when TOI > 1

diff --git a/Engine/Source/Programs/HeadlessChaos/Private/HeadlessChaosTestContactModification.cpp b/Engine/Source/Programs/HeadlessChaos/Private/HeadlessChaosTestContactModification.cpp
--- a/Engine/Source/Programs/HeadlessChaos/Private/HeadlessChaosTestContactModification.cpp
+++ b/Engine/Source/Programs/HeadlessChaos/Private/HeadlessChaosTestContactModification.cpp
@@ -10,7 +10,10 @@ namespace ChaosTest
 {
 	using namespace Chaos;
 
-	class FContactModificationTestCallback : public Chaos::TSimCallbackObject<Chaos::FSimCallbackNoInput, Chaos::FSimCallbackNoOutput>
+	class FContactModificationTestCallback : public Chaos::TSimCallbackObject<
+		Chaos::FSimCallbackNoInput,
+		Chaos::FSimCallbackNoOutput,
+		Chaos::ESimCallbackOptions::Presimulate | Chaos::ESimCallbackOptions::ContactModification>
 	{
 	public:
 		TUniqueFunction<void(Chaos::FCollisionContactModifier&)> TestLambda;
@@ -74,7 +77,7 @@ namespace ChaosTest
 		// Save Unique indices of floor and modified cube to disable in contact mod.
 		TVec2<FUniqueIdx> UniqueIndices({ModifiedCubeParticle.UniqueIdx(), FloorParticle.UniqueIdx()});
 
-		FContactModificationTestCallback* Callback = Solver->CreateAndRegisterSimCallbackObject_External<FContactModificationTestCallback>(true);
+		FContactModificationTestCallback* Callback = Solver->CreateAndRegisterSimCallbackObject_External<FContactModificationTestCallback>();
 		Callback->TestLambda = [UniqueIndices](Chaos::FCollisionContactModifier& Modifier)
 		{
 			for (FContactPairModifier& PairModifier : Modifier)
@@ -118,6 +121,85 @@ namespace ChaosTest
 		Module->DestroySolver(Solver);
 	}
 
+	GTEST_TEST(AllTraits, ContactModification_Probe)
+	{
+		FChaosSolversModule* Module = FChaosSolversModule::GetModule();
+		auto* Solver = Module->CreateSolver(nullptr, /*AsyncDt=*/-1);
+		InitSolverSettings(Solver);
+		Solver->SetThreadingMode_External(EThreadingModeTemp::SingleThread);
+
+		// create a static floor and two boxes falling onto it.
+		// Both boxes should turn all contacts to probes, one of them has CCD enabled
+		// and the other one doesn't.
+
+		auto CubeGeom = TSharedPtr<FImplicitObject, ESPMode::ThreadSafe>(new TBox<FReal, 3>(FVec3(-100), FVec3(100)));
+
+		// Fall through the floor, no ccd
+		FSingleParticlePhysicsProxy* CubeProxyA = FSingleParticlePhysicsProxy::Create(Chaos::FPBDRigidParticle::CreateParticle());
+		auto& CubeParticleA = CubeProxyA->GetGameThreadAPI();
+		CubeParticleA.SetGeometry(CubeGeom);
+		Solver->RegisterObject(CubeProxyA);
+		CubeParticleA.SetGravityEnabled(false);
+		CubeParticleA.SetCCDEnabled(false);
+		CubeParticleA.SetV(FVec3(0, 0, -100));
+		CubeParticleA.SetX(FVec3(200, 0, 500));
+		SetCubeInertiaTensor(CubeParticleA, /*Dimension=*/200, /*Mass=*/1);
+		ChaosTest::SetParticleSimDataToCollide({ CubeProxyA->GetParticle_LowLevel() });
+
+		// Fall through the floor, with ccd
+		FSingleParticlePhysicsProxy* CubeProxyB = FSingleParticlePhysicsProxy::Create(Chaos::FPBDRigidParticle::CreateParticle());
+		auto& CubeParticleB = CubeProxyB->GetGameThreadAPI();
+		CubeParticleB.SetGeometry(CubeGeom);
+		Solver->RegisterObject(CubeProxyB);
+		CubeParticleB.SetGravityEnabled(false);
+		CubeParticleB.SetCCDEnabled(true);
+		CubeParticleB.SetV(FVec3(0, 0, -1000));
+		CubeParticleB.SetX(FVec3(-200, 0, 500));
+		SetCubeInertiaTensor(CubeParticleB, /*Dimension=*/200, /*Mass=*/1);
+		ChaosTest::SetParticleSimDataToCollide({ CubeProxyB->GetParticle_LowLevel() });
+
+
+		// static floor at origin, occupying Z = [-100,0]
+		FSingleParticlePhysicsProxy* FloorProxy = FSingleParticlePhysicsProxy::Create(Chaos::FGeometryParticle::CreateParticle());
+		auto& FloorParticle = FloorProxy->GetGameThreadAPI();
+		auto FloorGeom = TSharedPtr<FImplicitObject, ESPMode::ThreadSafe>(new TBox<FReal, 3>(FVec3(-500, -500, -100), FVec3(500, 500, 0)));
+		FloorParticle.SetGeometry(FloorGeom);
+		Solver->RegisterObject(FloorProxy);
+		FloorParticle.SetX(FVec3(0, 0, 0));
+		ChaosTest::SetParticleSimDataToCollide({ FloorProxy->GetParticle_LowLevel() });
+
+		FContactModificationTestCallback* Callback = Solver->CreateAndRegisterSimCallbackObject_External<FContactModificationTestCallback>();
+		Callback->TestLambda = [](Chaos::FCollisionContactModifier& Modifier)
+		{
+			for (FContactPairModifier& PairModifier : Modifier)
+			{
+				PairModifier.ConvertToProbe();
+			}
+		};
+
+		const float Dt = 1.0f;
+		const int32 Steps = 10;
+		for (int Step = 0; Step < Steps; ++Step)
+		{
+			Solver->AdvanceAndDispatch_External(Dt);
+			Solver->UpdateGameThreadStructures();
+		}
+
+		// TODO: Check for hit callbacks?
+
+		// Both cubes should be below the floor
+		EXPECT_LT(CubeParticleA.X().Z, FloorParticle.X().Z);
+		EXPECT_LT(CubeParticleA.X().Z, FloorParticle.X().Z);
+
+		// Floor should be at origin.
+		EXPECT_EQ(FloorParticle.X().Z, 0);
+
+		Solver->UnregisterAndFreeSimCallbackObject_External(Callback);
+		Solver->UnregisterObject(CubeProxyA);
+		Solver->UnregisterObject(CubeProxyB);
+		Solver->UnregisterObject(FloorProxy);
+		Module->DestroySolver(Solver);
+	}
 
 	GTEST_TEST(AllTraits, ContactModification_ModifySeparation)
 	{
@@ -171,7 +253,7 @@ namespace ChaosTest
 		// Save Unique indices of floor and modified cube to disable in contact mod.
 		TVec2<FUniqueIdx> UniqueIndices({ ModifiedCubeParticle.UniqueIdx(), FloorParticle.UniqueIdx() });
 
-		FContactModificationTestCallback* Callback = Solver->CreateAndRegisterSimCallbackObject_External<FContactModificationTestCallback>(true);
+		FContactModificationTestCallback* Callback = Solver->CreateAndRegisterSimCallbackObject_External<FContactModificationTestCallback>();
 		Callback->TestLambda = [UniqueIndices, SeparationPadding](Chaos::FCollisionContactModifier& Modifier)
 		{
 			for (FContactPairModifier& PairModifier : Modifier)
@@ -268,7 +350,7 @@ namespace ChaosTest
 		// Save Unique indices of floor and modified cube to disable in contact mod.
 		TVec2<FUniqueIdx> UniqueIndices({ ModifiedCubeParticle.UniqueIdx(), FloorParticle.UniqueIdx() });
 
-		FContactModificationTestCallback* Callback = Solver->CreateAndRegisterSimCallbackObject_External<FContactModificationTestCallback>(true);
+		FContactModificationTestCallback* Callback = Solver->CreateAndRegisterSimCallbackObject_External<FContactModificationTestCallback>();
 		Callback->TestLambda = [UniqueIndices](Chaos::FCollisionContactModifier& Modifier)
 		{
 			for (FContactPairModifier& PairModifier : Modifier)
@@ -362,7 +444,7 @@ namespace ChaosTest
 		// Save Unique indices of floor and modified cube to disable in contact mod.
 		TVec2<FUniqueIdx> UniqueIndices({ ModifiedCubeParticle.UniqueIdx(), FloorParticle.UniqueIdx() });
 
-		FContactModificationTestCallback* Callback = Solver->CreateAndRegisterSimCallbackObject_External<FContactModificationTestCallback>(true);
+		FContactModificationTestCallback* Callback = Solver->CreateAndRegisterSimCallbackObject_External<FContactModificationTestCallback>();
 		Callback->TestLambda = [UniqueIndices](Chaos::FCollisionContactModifier& Modifier)
 		{
 			for (FContactPairModifier& PairModifier : Modifier)
@@ -450,7 +532,7 @@ namespace ChaosTest
 		// Save Unique indices of floor and modified cube to disable in contact mod.
 		TVec2<FUniqueIdx> NoBounceUniqueIndices({ NoBounceCubeParticle.UniqueIdx(), FloorParticle.UniqueIdx() });
 
-		FContactModificationTestCallback* Callback = Solver->CreateAndRegisterSimCallbackObject_External<FContactModificationTestCallback>(true);
+		FContactModificationTestCallback* Callback = Solver->CreateAndRegisterSimCallbackObject_External<FContactModificationTestCallback>();
 		Callback->TestLambda = [NoBounceUniqueIndices](Chaos::FCollisionContactModifier& Modifier)
 		{
 			for (FContactPairModifier& PairModifier : Modifier)
@@ -528,7 +610,7 @@ namespace ChaosTest
 		// Save Unique indices of floor and modified cube to disable in contact mod.
 		TVec2<FUniqueIdx> BounceUniqueIndices({ BounceCubeParticle.UniqueIdx(), FloorParticle.UniqueIdx() });
 
-		FContactModificationTestCallback* Callback = Solver->CreateAndRegisterSimCallbackObject_External<FContactModificationTestCallback>(true);
+		FContactModificationTestCallback* Callback = Solver->CreateAndRegisterSimCallbackObject_External<FContactModificationTestCallback>();
 		Callback->TestLambda = [BounceUniqueIndices](Chaos::FCollisionContactModifier& Modifier)
 		{
 			for (FContactPairModifier& PairModifier : Modifier)
@@ -616,7 +698,7 @@ namespace ChaosTest
 		// Save Unique indices of floor and modified cube to disable in contact mod.
 		TVec2<FUniqueIdx> UniqueIndices({ ModifiedCubeParticle.UniqueIdx(), FloorParticle.UniqueIdx() });
 
-		FContactModificationTestCallback* Callback = Solver->CreateAndRegisterSimCallbackObject_External<FContactModificationTestCallback>(true);
+		FContactModificationTestCallback* Callback = Solver->CreateAndRegisterSimCallbackObject_External<FContactModificationTestCallback>();
 		Callback->TestLambda = [UniqueIndices](Chaos::FCollisionContactModifier& Modifier)
 		{
 			for (FContactPairModifier& PairModifier : Modifier)
@@ -692,7 +774,7 @@ namespace ChaosTest
 		ChaosTest::SetParticleSimDataToCollide({ FloorProxy->GetParticle_LowLevel() });
 
 		FVec3 NewVelocity(100,0,100);
-		FContactModificationTestCallback* Callback = Solver->CreateAndRegisterSimCallbackObject_External<FContactModificationTestCallback>(true);
+		FContactModificationTestCallback* Callback = Solver->CreateAndRegisterSimCallbackObject_External<FContactModificationTestCallback>();
 		Callback->TestLambda = [NewVelocity](Chaos::FCollisionContactModifier& Modifier)
 		{
 			for (FContactPairModifier& PairModifier : Modifier)
@@ -754,7 +836,7 @@ namespace ChaosTest
 		FloorParticle.SetX(FVec3(0, 0, 0));
 		ChaosTest::SetParticleSimDataToCollide({ FloorProxy->GetParticle_LowLevel() });
 
-		FContactModificationTestCallback* Callback = Solver->CreateAndRegisterSimCallbackObject_External<FContactModificationTestCallback>(true);
+		FContactModificationTestCallback* Callback = Solver->CreateAndRegisterSimCallbackObject_External<FContactModificationTestCallback>();
 		Callback->TestLambda = [](Chaos::FCollisionContactModifier& Modifier)
 		{
 			for (FContactPairModifier& PairModifier : Modifier)
@@ -815,7 +897,7 @@ namespace ChaosTest
 		ChaosTest::SetParticleSimDataToCollide({ FloorProxy->GetParticle_LowLevel() });
 
 		FVec3 TeleportPosition(1000,2000,3000);
-		FContactModificationTestCallback* Callback = Solver->CreateAndRegisterSimCallbackObject_External<FContactModificationTestCallback>(true);
+		FContactModificationTestCallback* Callback = Solver->CreateAndRegisterSimCallbackObject_External<FContactModificationTestCallback>();
 		Callback->TestLambda = [TeleportPosition](Chaos::FCollisionContactModifier& Modifier)
 		{
 			for (FContactPairModifier& PairModifier : Modifier)
@@ -886,7 +968,7 @@ namespace ChaosTest
 		ChaosTest::SetParticleSimDataToCollide({ FloorProxy->GetParticle_LowLevel() });
 
 		FRotation3 ModificationRotation(FQuat::MakeFromEuler(FVec3(0, 0, 45)));
-		FContactModificationTestCallback* Callback = Solver->CreateAndRegisterSimCallbackObject_External<FContactModificationTestCallback>(true);
+		FContactModificationTestCallback* Callback = Solver->CreateAndRegisterSimCallbackObject_External<FContactModificationTestCallback>();
 		Callback->TestLambda = [ModificationRotation](Chaos::FCollisionContactModifier& Modifier)
 		{
 			for (FContactPairModifier& PairModifier : Modifier)
@@ -954,7 +1036,7 @@ namespace ChaosTest
 		FloorParticle.SetX(FVec3(0, 0, 0));
 		ChaosTest::SetParticleSimDataToCollide({ FloorProxy->GetParticle_LowLevel() });
 
-		FContactModificationTestCallback* Callback = Solver->CreateAndRegisterSimCallbackObject_External<FContactModificationTestCallback>(true);
+		FContactModificationTestCallback* Callback = Solver->CreateAndRegisterSimCallbackObject_External<FContactModificationTestCallback>();
 		Callback->TestLambda = [](Chaos::FCollisionContactModifier& Modifier)
 		{
 			for (FContactPairModifier& PairModifier : Modifier)
diff --git a/Engine/Source/Runtime/Experimental/Chaos/Private/Chaos/CCDUtilities.cpp b/Engine/Source/Runtime/Experimental/Chaos/Private/Chaos/CCDUtilities.cpp
--- a/Engine/Source/Runtime/Experimental/Chaos/Private/Chaos/CCDUtilities.cpp
+++ b/Engine/Source/Runtime/Experimental/Chaos/Private/Chaos/CCDUtilities.cpp
@@ -967,7 +967,7 @@ namespace Chaos
 
 			for (FPBDCollisionConstraint* Constraint : Constraints)
 			{
-				if (!Constraint->IsEnabled())
+				if (Constraint->IsProbe() || !Constraint->IsEnabled())
 				{
 					continue;
 				}
diff --git a/Engine/Source/Runtime/Experimental/Chaos/Private/Chaos/Collision/ParticlePairMidPhase.cpp b/Engine/Source/Runtime/Experimental/Chaos/Private/Chaos/Collision/ParticlePairMidPhase.cpp
--- a/Engine/Source/Runtime/Experimental/Chaos/Private/Chaos/Collision/ParticlePairMidPhase.cpp
+++ b/Engine/Source/Runtime/Experimental/Chaos/Private/Chaos/Collision/ParticlePairMidPhase.cpp
@@ -423,10 +423,15 @@ namespace Chaos
 			// If we did get a hit but it's at TOI = 1, treat this constraint as a regular non-swept constraint (skip the rewind)
 			if (Constraint->GetCCDTimeOfImpact() == FReal(1))
 			{
-				Constraint->SetCCDSweepEnabled(false);
 				Collisions::UpdateConstraint(*Constraint.Get(), Constraint->GetShapeWorldTransform0(), Constraint->GetShapeWorldTransform1(), Dt);
 			}
 
+			// If the sweep did not find a hit, or we are treating it like a regular contact, we will skip the CCD rewind step for this contact
+			if (Constraint->GetCCDTimeOfImpact() >= FReal(1))
+			{
+				Constraint->SetCCDSweepEnabled(false);
+			}
+
 			Context.GetAllocator()->ActivateConstraint(Constraint.Get());
 			LastUsedEpoch = Context.GetAllocator()->GetCurrentEpoch();
 
diff --git a/Engine/Source/Runtime/Experimental/Chaos/Private/Chaos/PBDCollisionConstraints.cpp b/Engine/Source/Runtime/Experimental/Chaos/Private/Chaos/PBDCollisionConstraints.cpp
--- a/Engine/Source/Runtime/Experimental/Chaos/Private/Chaos/PBDCollisionConstraints.cpp
+++ b/Engine/Source/Runtime/Experimental/Chaos/Private/Chaos/PBDCollisionConstraints.cpp
@@ -389,7 +389,7 @@ namespace Chaos
 				[this](FConstraintHandle* ConstraintHandle)
 				{
 					FPBDCollisionConstraintHandle* CollisionHandle = ConstraintHandle->AsUnsafe<FPBDCollisionConstraintHandle>();
-					if (!CollisionHandle->IsEnabled() || ConstraintAllocator.IsConstraintExpired(CollisionHandle->GetContact()))
+					if (!CollisionHandle->IsEnabled() || CollisionHandle->IsProbe() || ConstraintAllocator.IsConstraintExpired(CollisionHandle->GetContact()))
 					{
 						TempCollisions.Add(CollisionHandle);
 					}
@@ -407,7 +407,7 @@ namespace Chaos
 		for (FPBDCollisionConstraintHandle* ConstraintHandle : ConstraintAllocator.GetConstraints())
 		{
 			FPBDCollisionConstraint& Constraint = ConstraintHandle->GetContact();
-			if (!Constraint.IsInConstraintGraph() && Constraint.IsEnabled())
+			if (!Constraint.IsInConstraintGraph() && Constraint.IsEnabled() && !Constraint.IsProbe())
 			{
 				TempCollisions.Add(ConstraintHandle);
 			}
