Manual merge of 22524100 to resolve horde: https://horde.devtools.epicgames.com/job/634920a32b33326da3d5ae4f?step=1fc8

diff --git a/Engine/Source/Programs/HeadlessChaos/Private/HeadlessChaosTestIgnoreManager.cpp b/Engine/Source/Programs/HeadlessChaos/Private/HeadlessChaosTestIgnoreManager.cpp
--- a/Engine/Source/Programs/HeadlessChaos/Private/HeadlessChaosTestIgnoreManager.cpp
+++ b/Engine/Source/Programs/HeadlessChaos/Private/HeadlessChaosTestIgnoreManager.cpp
@@ -11,7 +11,7 @@ namespace ChaosTest
 {
 	using namespace Chaos;
 
-	class FCollisionReportingCallback : public TSimCallbackObject<FSimCallbackNoInput, FSimCallbackNoOutput>
+	class FCollisionReportingCallback : public TSimCallbackObject<FSimCallbackNoInput, FSimCallbackNoOutput, Chaos::ESimCallbackOptions::Presimulate | Chaos::ESimCallbackOptions::ContactModification>
 	{
 	public:
 		TArray<TTuple<FUniqueIdx, FUniqueIdx>> CollisionPairs;
@@ -104,7 +104,7 @@ namespace ChaosTest
 		FloorParticle.SetX(FVec3(0, 0, 0));
 		ChaosTest::SetParticleSimDataToCollide({ FloorProxy->GetParticle_LowLevel() });
 
-		FCollisionReportingCallback* Callback = Solver->CreateAndRegisterSimCallbackObject_External<FCollisionReportingCallback>(true);
+		FCollisionReportingCallback* Callback = Solver->CreateAndRegisterSimCallbackObject_External<FCollisionReportingCallback>();
 
 		// Named Ids for the particles
 		const FUniqueIdx CollidableId = CollidingCubeParticle.UniqueIdx();
@@ -190,7 +190,7 @@ namespace ChaosTest
 		FloorParticle.SetX(FVec3(0, 0, 0));
 		ChaosTest::SetParticleSimDataToCollide({ FloorProxy->GetParticle_LowLevel() });
 
-		FCollisionReportingCallback* Callback = Solver->CreateAndRegisterSimCallbackObject_External<FCollisionReportingCallback>(true);
+		FCollisionReportingCallback* Callback = Solver->CreateAndRegisterSimCallbackObject_External<FCollisionReportingCallback>();
 
 		// Named Ids for the particles
 		const FUniqueIdx CollidableId = CollidingCubeParticle.UniqueIdx();
