headless chaos : fix geometry collection test to adjust a bug fix made at CL 22850795

diff --git a/Engine/Source/Programs/HeadlessChaos/Private/GeometryCollection/GeometryCollectionTestClustering.cpp b/Engine/Source/Programs/HeadlessChaos/Private/GeometryCollection/GeometryCollectionTestClustering.cpp
--- a/Engine/Source/Programs/HeadlessChaos/Private/GeometryCollection/GeometryCollectionTestClustering.cpp
+++ b/Engine/Source/Programs/HeadlessChaos/Private/GeometryCollection/GeometryCollectionTestClustering.cpp
@@ -97,9 +97,9 @@ namespace GeometryCollectionTest
 
 		TManagedArray<bool>& Active = Collection->DynamicCollection->Active;
 
-		EXPECT_TRUE(Active[0]);
-		EXPECT_TRUE(Active[1]);
-		EXPECT_TRUE(Active[2]);
+		EXPECT_FALSE(Active[0]);
+		EXPECT_FALSE(Active[1]);
+		EXPECT_TRUE(Active[2]); // only the root cluster should be active when using clustering 
 		UnitTest.Advance();
 		EXPECT_FALSE(Active[0]);
 		EXPECT_FALSE(Active[1]);
