Chaos: added local array hash to a unit test as relocated definitions made the previous definition inaccessible

diff --git a/Engine/Source/Programs/HeadlessChaos/Private/GeometryCollection/GeometryCollectionTestDecimation.cpp b/Engine/Source/Programs/HeadlessChaos/Private/GeometryCollection/GeometryCollectionTestDecimation.cpp
--- a/Engine/Source/Programs/HeadlessChaos/Private/GeometryCollection/GeometryCollectionTestDecimation.cpp
+++ b/Engine/Source/Programs/HeadlessChaos/Private/GeometryCollection/GeometryCollectionTestDecimation.cpp
@@ -159,7 +159,17 @@ namespace GeometryCollectionTest
 		WriteImportanceOrderObjs(
 			TestCollection, Importance, CoincidentVertices, BaseName, OutputDir);
 
-		const uint32 hash = GetTypeHash(Importance);
+		auto CalcImportanceHash = [](const TArray<int32>& InImportance)
+		{
+			uint32 Hash = 0;
+			for(const int32 Value : InImportance)
+			{
+				Hash = HashCombine(Hash, GetTypeHash(Value));
+			}
+			return Hash;
+		};
+
+		const uint32 hash = CalcImportanceHash(Importance);
 #ifdef VERBOSE
 		std::cout << BaseName << " importance ordering hash: " << hash << std::endl;
 #endif
