Chaos - strip colinear edges and unused vertices from convexes
- unit test for convex box creation from default box render mesh which includes extra points

diff --git a/Engine/Source/Programs/HeadlessChaos/Private/HeadlessChaosTestConvex.cpp b/Engine/Source/Programs/HeadlessChaos/Private/HeadlessChaosTestConvex.cpp
--- a/Engine/Source/Programs/HeadlessChaos/Private/HeadlessChaosTestConvex.cpp
+++ b/Engine/Source/Programs/HeadlessChaos/Private/HeadlessChaosTestConvex.cpp
@@ -445,4 +445,105 @@ namespace ChaosTest
 			EXPECT_EQ(VertexIndex, StructureData.GetHalfEdgeVertex(VertexHalfEdgeIndex));
 		}
 	}
+
+	// The set of vertices generated from a unit box when creating a GeometryCollection from the default box in the editor.
+	// The default cube is tesselated. It has 26 vertices which include the 8 corners, plus mid-points along each edge and in the middle of each face.
+	// 
+	// This was causing the convex builder to produce a denegerate triangle (3 points in a row) leading to a zero normal and a crash in the solver.
+	//
+	// The fix was to modify TConvexHull3 to produce convex faces rather than triangles (one of which could be nearly degenerate), 
+	// and a post process on its results to eliminate colinear edges (within some tolerance)
+	//
+	GTEST_TEST(ConvexBuilderTests, TestDefaultStaticMeshBox)
+	{
+		FConvexBuilder::EBuildMethod BuildMethod = FConvexBuilder::EBuildMethod::Default;
+		const FReal Margin = 9.9999997473787516e-05;
+		TArray<FVec3f> BoxVerts =
+		{
+			{-50.0000000f, 50.0000000f, -50.0000000f},
+			{50.0000000f, 50.0000000f, -50.0000000f},
+			{50.0000000f, -50.0000000f, -50.0000000f},
+			{50.0000000f, -50.0000000f, 50.0000000f},
+			{50.0000000f, 50.0000000f, 50.0000000f},
+			{-50.0000000f, -50.0000000f, 50.0000000f},
+			{-50.0000000f, 50.0000000f, 50.0000000f},
+			{-50.0000000f, -50.0000000f, -50.0000000f},
+			{0.00000000f, 50.0000000f, -50.0000000f},
+			{50.0000000f, 0.00000000f, -50.0000000f},
+			{0.00000000f, 50.0000000f, 50.0000000f},
+			{-50.0000000f, -50.0000000f, 3.06161689e-15f},
+			{-50.0000000f, 50.0000000f, -3.06161689e-15f},
+			{-50.0000000f, 0.00000000f, -50.0000000f},
+			{50.0000000f, -50.0000000f, 3.06161689e-15f},
+			{0.00000000f, -50.0000000f, -50.0000000f},
+			{50.0000000f, 1.22464676e-14f, 50.0000000f},
+			{0.00000000f, -50.0000000f, 50.0000000f},
+			{-50.0000000f, 1.22464676e-14f, 50.0000000f},
+			{50.0000000f, 50.0000000f, -3.06161689e-15f},
+			{0.00000000f, 50.0000000f, -3.06161689e-15f},
+			{0.00000000f, 0.00000000f, -50.0000000f},
+			{0.00000000f, 1.22464676e-14f, 50.0000000f},
+			{0.00000000f, -50.0000000f, 3.06161689e-15f},
+			{50.0000000f, 6.12323379e-15f, -1.87469967e-31f},
+			{-50.0000000f, 6.12323379e-15f, -1.87469967e-31f},
+		};
+
+		FImplicitConvex3 Convex(BoxVerts, Margin, BuildMethod);
+
+		// The convex should be a box
+		EXPECT_EQ(Convex.NumVertices(), 8);
+		EXPECT_EQ(Convex.NumEdges(), 12);
+		EXPECT_EQ(Convex.NumPlanes(), 6);
+
+		// All planes normals should be...normalized
+		const FReal NormalTolerance = 1.e-4;
+		for (int32 PlaneIndex = 0; PlaneIndex < Convex.NumPlanes(); ++PlaneIndex)
+		{
+			const FVec3 PlaneN = Convex.GetPlane(PlaneIndex).Normal();
+			EXPECT_NEAR(PlaneN.Size(), FReal(1), NormalTolerance);
+		}
+
+	}
+
+	// Create a tet with an extra degenerate triangle in it. Verify that MergeColinearEdges handles this case
+	// and does not leave an invalid 2-vertex face behind.
+	// NOTE: We should not be able to create a FImplicitConvex3 that calls MergeColinearEdges in this condition
+	// but better safe than sorry.
+	GTEST_TEST(ConvexBuilderTests, TestColinearEdgeInTriangle)
+	{
+		// A right angled tet with an extra degenerate triangular face in there
+		TArray<FVec3f> TetVerts =
+		{
+			{0.0000000f, 0.0000000f, 50.0000000f},		// Top
+			{0.0000000f, 0.0000000f, 0.0000000f},		// Base0
+			{50.0000000f, 0.0000000f, 0.0000000f},		// Base1
+			{0.0000000f, 50.0000000f, 0.0000000f},		// Base2
+			{-1.e-15f, -1.e-14f, 25.0000000f},			// Extra vert along the vertical edge
+		};
+		TArray<TArray<int32>> TetFaces = 
+		{
+			{ 1, 2, 3 },								// Base
+			{ 0, 2, 1 },								// Side0
+			{ 0, 3, 2 },								// Side1
+			{ 0, 1, 3 },								// Side2
+			{ 0, 1, 4 },								// Extra degenerate face
+		};
+		TArray<TPlaneConcrete<FRealSingle>> TetPlanes =
+		{
+			// Values don't matter for this test
+			TPlaneConcrete<FRealSingle>(FVec3(0), FVec3(0,0,1)),
+			TPlaneConcrete<FRealSingle>(FVec3(0), FVec3(0,0,1)),
+			TPlaneConcrete<FRealSingle>(FVec3(0), FVec3(0,0,1)),
+			TPlaneConcrete<FRealSingle>(FVec3(0), FVec3(0,0,1)),
+			TPlaneConcrete<FRealSingle>(FVec3(0), FVec3(0,0,1)),
+		};
+
+		const FRealSingle AngleTolerance = 1.e-6f;
+		FConvexBuilder::MergeColinearEdges(TetPlanes, TetFaces, TetVerts, AngleTolerance);
+
+		// The invalid face and its vertex should have been stripped
+		EXPECT_EQ(TetVerts.Num(), 4);
+		EXPECT_EQ(TetFaces.Num(), 4);
+		EXPECT_EQ(TetPlanes.Num(), 4);
+	}
 }
diff --git a/Engine/Source/Runtime/Experimental/Chaos/Public/Chaos/CollisionConvexMesh.h b/Engine/Source/Runtime/Experimental/Chaos/Public/Chaos/CollisionConvexMesh.h
--- a/Engine/Source/Runtime/Experimental/Chaos/Public/Chaos/CollisionConvexMesh.h
+++ b/Engine/Source/Runtime/Experimental/Chaos/Public/Chaos/CollisionConvexMesh.h
@@ -574,6 +574,104 @@ namespace Chaos
 			RemoveInvalidFaces(InOutPlanes, InOutFaceVertexIndices);
 		}
 
+		// Find edge pairs that are colinear and remove the unnecessary vertex to make a single edge.
+		// If we are left with an invalid face (2 verts or less), remove it.
+		// NOTE: a vertex in the middle of two colinear edges on one face may still be required by some other face, 
+		// although technically those faces could (and should) have been merged.
+		static void MergeColinearEdges(TArray<FPlaneType>& InOutPlanes, TArray<TArray<int32>>& InOutFaceVertexIndices, TArray<FVec3Type>& InOutVertices, FRealType AngleToleranceRad)
+		{
+			check(InOutPlanes.Num() == InOutFaceVertexIndices.Num());
+
+			// This array maps from the input vertex to the output vertex array. INDEX_NONE means removed.
+			TArray<int32> VertexIndexMap;
+			VertexIndexMap.SetNum(InOutVertices.Num());
+			for (int32& VertexIndex : VertexIndexMap)
+			{
+				VertexIndex = INDEX_NONE;
+			}
+
+			// See if we have any co-linear edges in any faces, where the center vertex
+			// is not required by some other face. Assume we already removed coincident vertices.
+			// NOTE: after this loop the vertex index map contains INDEX_NONE for items to remove
+			// and other vertex has its original index. We pack and re-index later.
+			for (int32 PlaneIndex0 = 0; PlaneIndex0 < InOutFaceVertexIndices.Num(); ++PlaneIndex0)
+			{
+				TArray<int32>& FaceVertexIndices = InOutFaceVertexIndices[PlaneIndex0];
+				if (FaceVertexIndices.Num() > 2)
+				{
+					// Visit all set of 3-vertex chains in the face
+					int32 VertexIndex0 = FaceVertexIndices[FaceVertexIndices.Num() - 2];
+					int32 VertexIndex1 = FaceVertexIndices[FaceVertexIndices.Num() - 1];
+					for (int32 FaceVertexIndex = 0; FaceVertexIndex < FaceVertexIndices.Num(); ++FaceVertexIndex)
+					{
+						int32 VertexIndex2 = FaceVertexIndices[FaceVertexIndex];
+
+						// Calculate the sine of the angle between the two edges formed by the 3 vertices
+						const FVec3 Edge0 = (InOutVertices[VertexIndex1] - InOutVertices[VertexIndex0]).GetSafeNormal();
+						const FVec3 Edge1 = (InOutVertices[VertexIndex2] - InOutVertices[VertexIndex1]).GetSafeNormal();
+						const FReal CosAngle = FVec3::DotProduct(Edge0, Edge1);
+
+						// See if we need the vertex.
+						if (CosAngle < (FReal(1) - AngleToleranceRad))
+						{
+							VertexIndexMap[VertexIndex1] = VertexIndex1;
+						}
+
+						// Move to next edge pair
+						VertexIndex0 = VertexIndex1;
+						VertexIndex1 = VertexIndex2;
+					}
+				}
+			}
+
+			// Remove unused vertices from all faces and update the index map to account for removals
+			int32 NumVerticesRemoved = 0;
+			for (int32 VertexIndex = 0; VertexIndex < InOutVertices.Num(); ++VertexIndex)
+			{
+				if (VertexIndexMap[VertexIndex] == INDEX_NONE)
+				{
+					// Remove vertices we don't need from faces that use them
+					// If we end up with less than 3 verts, remove the face
+					for (int32 PlaneIndex0 = InOutFaceVertexIndices.Num() - 1; PlaneIndex0 >= 0; --PlaneIndex0)
+					{
+						InOutFaceVertexIndices[PlaneIndex0].Remove(VertexIndex);
+						if (InOutFaceVertexIndices[PlaneIndex0].Num() < 3)
+						{
+							InOutFaceVertexIndices.RemoveAt(PlaneIndex0);
+							InOutPlanes.RemoveAt(PlaneIndex0);
+						}
+					}
+					++NumVerticesRemoved;
+				}
+				else
+				{
+					VertexIndexMap[VertexIndex] = VertexIndex - NumVerticesRemoved;
+				}
+			}
+
+			if (NumVerticesRemoved > 0)
+			{
+				// Remove unused verts from the array
+				for (int32 VertexIndex = InOutVertices.Num() - 1; VertexIndex >= 0; --VertexIndex)
+				{
+					if (VertexIndexMap[VertexIndex] == INDEX_NONE)
+					{
+						InOutVertices.RemoveAt(VertexIndex);
+					}
+				}
+
+				// Remap vertex indices in all faces
+				for (int32 PlaneIndex0 = 0; PlaneIndex0 < InOutFaceVertexIndices.Num(); ++PlaneIndex0)
+				{
+					TArray<int32>& FaceVertexIndices = InOutFaceVertexIndices[PlaneIndex0];
+					for (int32 FaceVertexIndex = FaceVertexIndices.Num() - 1; FaceVertexIndex >= 0; --FaceVertexIndex)
+					{
+						FaceVertexIndices[FaceVertexIndex] = VertexIndexMap[FaceVertexIndices[FaceVertexIndex]];
+					}
+				}
+			}
+		}
+
 		// IMPORTANT : vertices are assumed to be sorted CCW
 		static void RemoveInsideFaceVertices(const FPlaneType& Face, TArray<int32>& InOutFaceVertexIndices, const TArray<FVec3Type>& Vertices, const FVec3Type& Centroid)
 		{
diff --git a/Engine/Source/Runtime/Experimental/Chaos/Public/Chaos/Convex.h b/Engine/Source/Runtime/Experimental/Chaos/Public/Chaos/Convex.h
--- a/Engine/Source/Runtime/Experimental/Chaos/Public/Chaos/Convex.h
+++ b/Engine/Source/Runtime/Experimental/Chaos/Public/Chaos/Convex.h
@@ -116,10 +116,12 @@ namespace Chaos
 			// @todo(chaos): this only works with triangles. Fix that an we can run MergeFaces before calling this
 			CalculateVolumeAndCenterOfMass(Vertices, FaceIndices, Volume, CenterOfMass);
 
-			// it appears that this code path can leave the convex in an
-			// undefined state.
+			// it appears that this code path can leave the convex in an undefined state.
+			// @todo(chaos): Tolerances should be based on size, or passed in
 			const FRealType DistanceTolerance = 1.0f;
+			const FRealType AngleTolerance = 1.e-6f;
 			FConvexBuilder::MergeFaces(Planes, FaceIndices, Vertices, DistanceTolerance);
+			FConvexBuilder::MergeColinearEdges(Planes, FaceIndices, Vertices, AngleTolerance);
 			CHAOS_ENSURE(Planes.Num() == FaceIndices.Num());
 
 			CreateStructureData(MoveTemp(FaceIndices));
@@ -1007,9 +1009,11 @@ namespace Chaos
 			TArray<TArray<int32>> FaceIndices;
 			FConvexBuilder::Simplify(Planes, FaceIndices, Vertices, LocalBoundingBox);
 
-			// @todo(chaos): DistanceTolerance should be based on size, or passed in
+			// @todo(chaos): Tolerances should be based on size, or passed in
 			const FRealType DistanceTolerance = 1.0f;
+			const FRealType AngleTolerance = 1.e-6f;
 			FConvexBuilder::MergeFaces(Planes, FaceIndices, Vertices, DistanceTolerance);
+			FConvexBuilder::MergeColinearEdges(Planes, FaceIndices, Vertices, AngleTolerance);
 
 			CreateStructureData(MoveTemp(FaceIndices));
 		}
diff --git a/Engine/Source/Runtime/Experimental/ChaosCore/Public/Chaos/Core.h b/Engine/Source/Runtime/Experimental/ChaosCore/Public/Chaos/Core.h
--- a/Engine/Source/Runtime/Experimental/ChaosCore/Public/Chaos/Core.h
+++ b/Engine/Source/Runtime/Experimental/ChaosCore/Public/Chaos/Core.h
@@ -51,7 +51,7 @@ namespace Chaos
 	// NOTE: if you get a merge conflict on the GUID, you must replace it with a new GUID - do not accept the source or target
 	// or you will likely get DDC version conflicts resulting in crashes during load.
 	// Core version string for Chaos data. Any DDC builder dependent on Chaos for serialization should depend on this version
-	static const TCHAR* const ChaosVersionGUID = TEXT("3366E325-8C79-4356-916E-B8432589888D");
+	static const TCHAR* const ChaosVersionGUID = TEXT("FC60014F-0799-4BFD-B01C-44C1DEE81F32");
 
 	inline FString GetChaosVersionStringInner()
 	{
