Optimization : Adding OBB Query in AABBTree for Triangle Mesh

diff --git a/Engine/Source/Programs/HeadlessChaos/Private/HeadlessChaosTestObbTree.cpp b/Engine/Source/Programs/HeadlessChaos/Private/HeadlessChaosTestObbTree.cpp
new file mode 100644
--- /dev/null
+++ b/Engine/Source/Programs/HeadlessChaos/Private/HeadlessChaosTestObbTree.cpp
@@ -0,0 +1,65 @@
+// Copyright Epic Games, Inc. All Rights Reserved.
+
+#include "HeadlessChaos.h"
+#include "HeadlessChaosTestUtility.h"
+#include "Chaos/OBBVectorized.h"
+#include "Chaos/TriangleMeshImplicitObject.h"
+
+namespace ChaosTest
+{
+	using namespace Chaos;
+
+
+	GTEST_TEST(ObbTest, IntersectAABB)
+	{
+		{
+			// Original points to create OBB // TArray<FVec3f> PointsA{ FVec3f(0, 0, 0), FVec3f(5, 0, 0), FVec3f(0, 5, 0) };
+			FRigidTransform3 Transform = FRigidTransform3(FVec3(1.25, 1.25, 0.0), FQuat(FVec3(0.0, 0.0, 1.0), -3.1415926 / 4.0));
+			FOBBVectorized ObbA(Transform, FVec3f(UE_INV_SQRT_2*5.0, UE_INV_SQRT_2 * 5.0 /2.0, 0.0f));
+			{
+				TAABB<FRealSingle, 3> Bounds(FVec3f(-2.0f, 1.5f, 0.0f), FVec3f(-1.0f, 2.5f, 0.0f));
+				EXPECT_TRUE(ObbA.IntersectAABB(FAABBVectorized(Bounds)));
+			}
+			{
+				TAABB<FRealSingle, 3> Bounds(FVec3f(-2.0f, 4.5f, 0.0f), FVec3f(-1.0f, 5.5f, 0.0f));
+				EXPECT_FALSE(ObbA.IntersectAABB(FAABBVectorized(Bounds)));
+			}
+			{
+				TAABB<FRealSingle, 3> Bounds(FVec3f(2.0f, 2.0f, -10.0f), FVec3f(3.0f, 3.0f, 10.0f));
+				EXPECT_TRUE(ObbA.IntersectAABB(FAABBVectorized(Bounds)));
+			}
+			{
+				TAABB<FRealSingle, 3> Bounds(FVec3f(-10.0f, -10.0f, -10.0f), FVec3f(10.0f, 10.0f, 10.0f));
+				EXPECT_TRUE(ObbA.IntersectAABB(FAABBVectorized(Bounds)));
+			}
+			{
+				TAABB<FRealSingle, 3> Bounds(FVec3f(4.0f, -1.0f, 0.0f), FVec3f(6.0f, 1.0f, 0.0f));
+				EXPECT_TRUE(ObbA.IntersectAABB(FAABBVectorized(Bounds)));
+			}
+			{
+				TAABB<FRealSingle, 3> Bounds(FVec3f(6.0f, -1.0f, 0.0f), FVec3f(7.0f, 1.0f, 0.0f));
+				EXPECT_FALSE(ObbA.IntersectAABB(FAABBVectorized(Bounds)));
+			}
+			{
+				TAABB<FRealSingle, 3> Bounds(FVec3f(-1.0f, 6.0f, 0.0f), FVec3f(1.0f, 7.0f, 0.0f));
+				EXPECT_FALSE(ObbA.IntersectAABB(FAABBVectorized(Bounds)));
+			}
+			{
+				TAABB<FRealSingle, 3> Bounds(FVec3f(4.0f, 4.0f, 0.0f), FVec3f(10.0f, 10.0f, 0.0f));
+				EXPECT_FALSE(ObbA.IntersectAABB(FAABBVectorized(Bounds)));
+			}
+			{
+				TAABB<FRealSingle, 3> Bounds(FVec3f(-1.5f, -1.0f, 0.0f), FVec3f(-0.5f, -0.5f, 0.0f));
+				EXPECT_FALSE(ObbA.IntersectAABB(FAABBVectorized(Bounds)));
+			}
+		}
+		{
+			
+			FRigidTransform3 Transform = FRigidTransform3(FVec3(0.0, 0.0, 12.0), FQuat(FVec3(1.0, 0.0, 0.0), 3.1415926 / 4.0));
+			FOBBVectorized ObbA(Transform.ToMatrixNoScale(), FVec3f(2.0f, 2.0f, 2.0f));
+			TAABB<FRealSingle, 3> Bounds(FVec3f(-10.0f, -10.0f, 0.0f), FVec3f(10.0f, 10.0f, 10.0f));
+			EXPECT_TRUE(ObbA.IntersectAABB(FAABBVectorized(Bounds)));
+		}
+	}
+
+}
\ No newline at end of file
diff --git a/Engine/Source/Programs/HeadlessChaos/Private/HeadlessChaosTestOverlap.cpp b/Engine/Source/Programs/HeadlessChaos/Private/HeadlessChaosTestOverlap.cpp
--- a/Engine/Source/Programs/HeadlessChaos/Private/HeadlessChaosTestOverlap.cpp
+++ b/Engine/Source/Programs/HeadlessChaos/Private/HeadlessChaosTestOverlap.cpp
@@ -14,348 +14,726 @@ namespace ChaosTest
 
 	void OverlapTriMesh()
 	{
-		// Trimesh is simple pyramid
-		using namespace Chaos;
-		FTriangleMeshImplicitObject::ParticlesType TrimeshParticles(
-			{
-				{-10.0, 0.0, 0.0},
-				{10.0, 0.0, 0.0},
-				{0.0, 10.0, 0.0},
-				{0.0, -10.0, 0.0},
-				{0.0,  0.0, 10.0},
-			});
-
-		TArray<TVec3<int32>> Indices(
-			{
-				{0, 1, 2},
-				{0, 3, 1},
-				{0, 1, 2},
-				{0, 2, 4},
-				{1, 2, 4},
-				{0, 3, 4},
-				{1, 3, 4}
-			});
-
-		TArray<uint16> Materials;
-		for (int32 i = 0; i < Indices.Num(); ++i)
 		{
-			Materials.Emplace(0);
-		}
+			// Trimesh is simple pyramid
+			using namespace Chaos;
+			FTriangleMeshImplicitObject::ParticlesType TrimeshParticles(
+				{
+					{-10.0, 0.0, 0.0},
+					{10.0, 0.0, 0.0},
+					{0.0, 10.0, 0.0},
+					{0.0, -10.0, 0.0},
+					{0.0,  0.0, 10.0},
+				});
 
-		TUniquePtr<FTriangleMeshImplicitObject> TriangleMesh = MakeUnique<FTriangleMeshImplicitObject>(MoveTemp(TrimeshParticles), MoveTemp(Indices), MoveTemp(Materials));
-		{
-			// Capsule test
-			const FVec3 X1 = { 0.0, 0.0, -2.0 };
-			const FVec3 X2 = { 0.0, 0.0, 2.0 };
-			const FReal Radius = 1.0;
-			const FCapsule Capsule = FCapsule(X1, X2, Radius);
+			TArray<TVec3<int32>> Indices(
+				{
+					{0, 1, 2},
+					{0, 3, 1},
+					{0, 1, 2},
+					{0, 2, 4},
+					{1, 2, 4},
+					{0, 3, 4},
+					{1, 3, 4}
+				});
 
+			TArray<uint16> Materials;
+			for (int32 i = 0; i < Indices.Num(); ++i)
 			{
-				const FRigidTransform3 QueryTM(FVec3(0.0, 0.0, 0.0), FQuat::Identity);
-				bool bResult = TriangleMesh->OverlapGeom(Capsule, QueryTM, 0.0);
-				EXPECT_EQ(bResult, true);
-			}
-			{
-				const FRigidTransform3 QueryTM(FVec3(10.0, 10.0, 0.0), FQuat::Identity);
-				bool bResult = TriangleMesh->OverlapGeom(Capsule, QueryTM, 0.0);
-				EXPECT_EQ(bResult, false);
-			}
-			{
-				const FRigidTransform3 QueryTM(FVec3(-10.0, -10.0, 0.0), FQuat::Identity);
-				bool bResult = TriangleMesh->OverlapGeom(Capsule, QueryTM, 0.0);
-				EXPECT_EQ(bResult, false);
-			}
-			{
-				const FRigidTransform3 QueryTM(FVec3(10.0, 0.0, 0.0), FQuat::Identity);
-				bool bResult = TriangleMesh->OverlapGeom(Capsule, QueryTM, 0.0);
-				EXPECT_EQ(bResult, true);
-			}
-			{
-				const FRigidTransform3 QueryTM(FVec3(0.0, 10.0, 0.0), FQuat::Identity);
-				bool bResult = TriangleMesh->OverlapGeom(Capsule, QueryTM, 0.0);
-				EXPECT_EQ(bResult, true);
-			}
-			{
-				const FRigidTransform3 QueryTM(FVec3(0.0, -10.0, 0.0), FQuat::Identity);
-				bool bResult = TriangleMesh->OverlapGeom(Capsule, QueryTM, 0.0);
-				EXPECT_EQ(bResult, true);
+				Materials.Emplace(0);
 			}
-			{
-				const FRigidTransform3 QueryTM(FVec3(-10.0, 0.0, 0.0), FQuat::Identity);
-				bool bResult = TriangleMesh->OverlapGeom(Capsule, QueryTM, 0.0);
-				EXPECT_EQ(bResult, true);
-			}
-			{
-				const FRigidTransform3 QueryTM(FVec3(0.0, 0.0, 12.0), FQuat::Identity);
-				bool bResult = TriangleMesh->OverlapGeom(Capsule, QueryTM, 0.0);
-				EXPECT_EQ(bResult, true);
-			}
-			{
-				const FRigidTransform3 QueryTM(FVec3(0.0, 0.0, 12.0), FQuat(FVec3(1.0, 0.0, 0.0), 3.1415926 / 4.0));
-				bool bResult = TriangleMesh->OverlapGeom(Capsule, QueryTM, 0.0);
-				EXPECT_EQ(bResult, false);
-			}
-			{
-				// Inside Mesh
-				const FRigidTransform3 QueryTM(FVec3(0.0, 0.0, 4.0), FQuat::Identity);
-				bool bResult = TriangleMesh->OverlapGeom(Capsule, QueryTM, 0.0);
-				EXPECT_EQ(bResult, false);
-			}
-		}
-		{
-			const FVec3 X1 = { 0.0, 0.0, -200.0 };
-			const FVec3 X2 = { 0.0, 0.0, 200.0 };
-			const FReal Radius = 100.0;
-			const FCapsule Capsule = FCapsule(X1, X2, Radius);
-			TSharedPtr<FCapsule> CapsuleShared = MakeShared<FCapsule>(X1, X2, Radius);
-			FVec3 TriMeshScale = { 0.01, 0.01, 0.01 };
-			TImplicitObjectScaled<FCapsule> ScaledCapsule = TImplicitObjectScaled<FCapsule>(CapsuleShared, TriMeshScale);
-			{
-				const FRigidTransform3 QueryTM(FVec3(0.0, 0.0, 0.0), FQuat::Identity);
-				bool bResult = TriangleMesh->OverlapGeom(ScaledCapsule, QueryTM, 0.0, nullptr, TriMeshScale);
-				EXPECT_EQ(bResult, true);
-			}
-			{
-				const FRigidTransform3 QueryTM(FVec3(10.0, 10.0, 0.0), FQuat::Identity);
-				bool bResult = TriangleMesh->OverlapGeom(ScaledCapsule, QueryTM, 0.0, nullptr, TriMeshScale);
-				EXPECT_EQ(bResult, false);
-			}
-			{
-				const FRigidTransform3 QueryTM(FVec3(-10.0, -10.0, 0.0), FQuat::Identity);
-				bool bResult = TriangleMesh->OverlapGeom(ScaledCapsule, QueryTM, 0.0, nullptr, TriMeshScale);
-				EXPECT_EQ(bResult, false);
-			}
-			{
-				const FRigidTransform3 QueryTM(FVec3(10.0, 0.0, 0.0), FQuat::Identity);
-				bool bResult = TriangleMesh->OverlapGeom(ScaledCapsule, QueryTM, 0.0, nullptr, TriMeshScale);
-				EXPECT_EQ(bResult, true);
-			}
-			{
-				const FRigidTransform3 QueryTM(FVec3(0.0, 10.0, 0.0), FQuat::Identity);
-				bool bResult = TriangleMesh->OverlapGeom(ScaledCapsule, QueryTM, 0.0, nullptr, TriMeshScale);
-				EXPECT_EQ(bResult, true);
-			}
-			{
-				const FRigidTransform3 QueryTM(FVec3(0.0, -10.0, 0.0), FQuat::Identity);
-				bool bResult = TriangleMesh->OverlapGeom(ScaledCapsule, QueryTM, 0.0, nullptr, TriMeshScale);
-				EXPECT_EQ(bResult, true);
-			}
-			{
-				const FRigidTransform3 QueryTM(FVec3(-10.0, 0.0, 0.0), FQuat::Identity);
-				bool bResult = TriangleMesh->OverlapGeom(ScaledCapsule, QueryTM, 0.0, nullptr, TriMeshScale);
-				EXPECT_EQ(bResult, true);
-			}
-			{
-				const FRigidTransform3 QueryTM(FVec3(0.0, 0.0, 12.0), FQuat::Identity);
-				bool bResult = TriangleMesh->OverlapGeom(ScaledCapsule, QueryTM, 0.0, nullptr, TriMeshScale);
-				EXPECT_EQ(bResult, true);
-			}
-			{
-				const FRigidTransform3 QueryTM(FVec3(0.0, 0.0, 12.0), FQuat(FVec3(1.0, 0.0, 0.0), 3.1415926 / 4.0));
-				bool bResult = TriangleMesh->OverlapGeom(ScaledCapsule, QueryTM, 0.0, nullptr, TriMeshScale);
-				EXPECT_EQ(bResult, false);
-			}
-			{
-				// Inside Mesh
-				const FRigidTransform3 QueryTM(FVec3(0.0, 0.0, 4.0), FQuat::Identity);
-				bool bResult = TriangleMesh->OverlapGeom(ScaledCapsule, QueryTM, 0.0);
-				EXPECT_EQ(bResult, false);
-			}
-		}
-		{
-			// Box test
-			const TBox<FReal, 3> Box = TBox<FReal, 3>({ -1.0, -1.0, -2.0 }, { 1.0, 1.0, 2.0 });
-			{
-				const FRigidTransform3 QueryTM(FVec3(0.0, 0.0, 0.0), FQuat::Identity);
-				bool bResult = TriangleMesh->OverlapGeom(Box, QueryTM, 0.0);
-				EXPECT_EQ(bResult, true);
-			}
-			{
-				const FRigidTransform3 QueryTM(FVec3(10.0, 10.0, 0.0), FQuat::Identity);
-				bool bResult = TriangleMesh->OverlapGeom(Box, QueryTM, 0.0);
-				EXPECT_EQ(bResult, false);
-			}
-			{
-				const FRigidTransform3 QueryTM(FVec3(-10.0, -10.0, 0.0), FQuat::Identity);
-				bool bResult = TriangleMesh->OverlapGeom(Box, QueryTM, 0.0);
-				EXPECT_EQ(bResult, false);
-			}
-			{
-				const FRigidTransform3 QueryTM(FVec3(10.0, 0.0, 0.0), FQuat::Identity);
-				bool bResult = TriangleMesh->OverlapGeom(Box, QueryTM, 0.0);
-				EXPECT_EQ(bResult, true);
-			}
-			{
-				const FRigidTransform3 QueryTM(FVec3(0.0, 10.0, 0.0), FQuat::Identity);
-				bool bResult = TriangleMesh->OverlapGeom(Box, QueryTM, 0.0);
-				EXPECT_EQ(bResult, true);
-			}
-			{
-				const FRigidTransform3 QueryTM(FVec3(0.0, -10.0, 0.0), FQuat::Identity);
-				bool bResult = TriangleMesh->OverlapGeom(Box, QueryTM, 0.0);
-				EXPECT_EQ(bResult, true);
-			}
-			{
-				const FRigidTransform3 QueryTM(FVec3(-10.0, 0.0, 0.0), FQuat::Identity);
-				bool bResult = TriangleMesh->OverlapGeom(Box, QueryTM, 0.0);
-				EXPECT_EQ(bResult, true);
-			}
-			{
-				const FRigidTransform3 QueryTM(FVec3(0.0, 0.0, 12.0), FQuat::Identity);
-				bool bResult = TriangleMesh->OverlapGeom(Box, QueryTM, 0.0);
-				EXPECT_EQ(bResult, true);
-			}
-			{
-				const FRigidTransform3 QueryTM(FVec3(0.0, 0.0, 12.0), FQuat(FVec3(1.0, 0.0, 0.0), 3.1415926 / 4.0));
-				bool bResult = TriangleMesh->OverlapGeom(Box, QueryTM, 0.0);
-				EXPECT_EQ(bResult, false);
-			}
-			{
-				// Inside Mesh
-				const FRigidTransform3 QueryTM(FVec3(0.0, 0.0, 4.0), FQuat::Identity);
-				bool bResult = TriangleMesh->OverlapGeom(Box, QueryTM, 0.0);
-				EXPECT_EQ(bResult, false);
-			}
-		}
 
-		// Box Scaled
-		{
-			TSharedPtr<TBox<FReal, 3>> BigBoxSafe = MakeShared<TBox<FReal, 3>>(FVec3(-100.0, -100.0, -100.0), FVec3(100.0, 100.0, 100.0));
-			FVec3 TriMeshScale = { 0.01, 0.01, 0.02 };
-			TImplicitObjectScaled<TBox<FReal, 3>> ScaledBox = TImplicitObjectScaled<TBox<FReal, 3>>(BigBoxSafe, TriMeshScale);
-			{
-				const FRigidTransform3 QueryTM(FVec3(0.0, 0.0, 0.0), FQuat::Identity);
-				bool bResult = TriangleMesh->OverlapGeom(ScaledBox, QueryTM, 0.0, nullptr, TriMeshScale);
-				EXPECT_EQ(bResult, true);
-			}
-			{
-				const FRigidTransform3 QueryTM(FVec3(10.0, 10.0, 0.0), FQuat::Identity);
-				bool bResult = TriangleMesh->OverlapGeom(ScaledBox, QueryTM, 0.0, nullptr, TriMeshScale);
-				EXPECT_EQ(bResult, false);
-			}
-			{
-				const FRigidTransform3 QueryTM(FVec3(-10.0, -10.0, 0.0), FQuat::Identity);
-				bool bResult = TriangleMesh->OverlapGeom(ScaledBox, QueryTM, 0.0, nullptr, TriMeshScale);
-				EXPECT_EQ(bResult, false);
-			}
+			TUniquePtr<FTriangleMeshImplicitObject> TriangleMesh = MakeUnique<FTriangleMeshImplicitObject>(MoveTemp(TrimeshParticles), MoveTemp(Indices), MoveTemp(Materials));
 			{
-				const FRigidTransform3 QueryTM(FVec3(10.0, 0.0, 0.0), FQuat::Identity);
-				bool bResult = TriangleMesh->OverlapGeom(ScaledBox, QueryTM, 0.0, nullptr, TriMeshScale);
-				EXPECT_EQ(bResult, true);
-			}
-			{
-				const FRigidTransform3 QueryTM(FVec3(0.0, 10.0, 0.0), FQuat::Identity);
-				bool bResult = TriangleMesh->OverlapGeom(ScaledBox, QueryTM, 0.0, nullptr, TriMeshScale);
-				EXPECT_EQ(bResult, true);
-			}
-			{
-				const FRigidTransform3 QueryTM(FVec3(0.0, -10.0, 0.0), FQuat::Identity);
-				bool bResult = TriangleMesh->OverlapGeom(ScaledBox, QueryTM, 0.0, nullptr, TriMeshScale);
-				EXPECT_EQ(bResult, true);
-			}
-			{
-				const FRigidTransform3 QueryTM(FVec3(-10.0, 0.0, 0.0), FQuat::Identity);
-				bool bResult = TriangleMesh->OverlapGeom(ScaledBox, QueryTM, 0.0, nullptr, TriMeshScale);
-				EXPECT_EQ(bResult, true);
-			}
-			{
-				const FRigidTransform3 QueryTM(FVec3(0.0, 0.0, 12.0), FQuat::Identity);
-				bool bResult = TriangleMesh->OverlapGeom(ScaledBox, QueryTM, 0.0, nullptr, TriMeshScale);
-				EXPECT_EQ(bResult, true);
-			}
-			{
-				const FRigidTransform3 QueryTM(FVec3(0.0, 0.0, 12.0), FQuat(FVec3(1.0, 0.0, 0.0), 3.1415926 / 4.0));
-				bool bResult = TriangleMesh->OverlapGeom(ScaledBox, QueryTM, 0.0, nullptr, TriMeshScale);
-				EXPECT_EQ(bResult, false);
-			}
-			{
-				// Inside Mesh
-				const FRigidTransform3 QueryTM(FVec3(0.0, 0.0, 4.0), FQuat::Identity);
-				bool bResult = TriangleMesh->OverlapGeom(ScaledBox, QueryTM, 0.0, nullptr, TriMeshScale);
-				EXPECT_EQ(bResult, false);
-			}
-		}
-		// Box Scaled test non uniform transform with rotation
-		{
-			TSharedPtr<TBox<FReal, 3>> BigBoxSafe = MakeShared<TBox<FReal, 3>>(FVec3(-100.0, -100.0, -100.0), FVec3(100.0, 100.0, 100.0));
-			FVec3 TriMeshScale = { 0.01, 0.01, 0.05 };
-			TImplicitObjectScaled<TBox<FReal, 3>> ScaledBox = TImplicitObjectScaled<TBox<FReal, 3>>(BigBoxSafe, TriMeshScale);
+				// Capsule test
+				const FVec3 X1 = { 0.0, 0.0, -2.0 };
+				const FVec3 X2 = { 0.0, 0.0, 2.0 };
+				const FReal Radius = 1.0;
+				const FCapsule Capsule = FCapsule(X1, X2, Radius);
 
-			{
-				const FRigidTransform3 QueryTM(FVec3(5.5, 0.0, 10.0), FQuat::Identity);
-				bool bResult = TriangleMesh->OverlapGeom(ScaledBox, QueryTM, 0.0, nullptr, TriMeshScale);
-				EXPECT_EQ(bResult, true);
-			}
-			{
-				const FRigidTransform3 QueryTM(FVec3(7.0, 0.0, 10.0), FQuat::Identity);
-				bool bResult = TriangleMesh->OverlapGeom(ScaledBox, QueryTM, 0.0, nullptr, TriMeshScale);
-				EXPECT_EQ(bResult, false);
+				{
+					const FRigidTransform3 QueryTM(FVec3(0.0, 0.0, 0.0), FQuat::Identity);
+					bool bResult = TriangleMesh->OverlapGeom(Capsule, QueryTM, 0.0);
+					EXPECT_EQ(bResult, true);
+				}
+				{
+					const FRigidTransform3 QueryTM(FVec3(12.5, 0.0, 0.0), FQuat::Identity);
+					bool bResult = TriangleMesh->OverlapGeom(Capsule, QueryTM, 0.0);
+					EXPECT_EQ(bResult, false);
+				}
+				{
+					const FRigidTransform3 QueryTM(FVec3(10.0, 10.0, 0.0), FQuat::Identity);
+					bool bResult = TriangleMesh->OverlapGeom(Capsule, QueryTM, 0.0);
+					EXPECT_EQ(bResult, false);
+				}
+				{
+					const FRigidTransform3 QueryTM(FVec3(-10.0, -10.0, 0.0), FQuat::Identity);
+					bool bResult = TriangleMesh->OverlapGeom(Capsule, QueryTM, 0.0);
+					EXPECT_EQ(bResult, false);
+				}
+				{
+					const FRigidTransform3 QueryTM(FVec3(10.0, 0.0, 0.0), FQuat::Identity);
+					bool bResult = TriangleMesh->OverlapGeom(Capsule, QueryTM, 0.0);
+					EXPECT_EQ(bResult, true);
+				}
+				{
+					const FRigidTransform3 QueryTM(FVec3(0.0, 10.0, 0.0), FQuat::Identity);
+					bool bResult = TriangleMesh->OverlapGeom(Capsule, QueryTM, 0.0);
+					EXPECT_EQ(bResult, true);
+				}
+				{
+					const FRigidTransform3 QueryTM(FVec3(0.0, -10.0, 0.0), FQuat::Identity);
+					bool bResult = TriangleMesh->OverlapGeom(Capsule, QueryTM, 0.0);
+					EXPECT_EQ(bResult, true);
+				}
+				{
+					const FRigidTransform3 QueryTM(FVec3(-10.0, 0.0, 0.0), FQuat::Identity);
+					bool bResult = TriangleMesh->OverlapGeom(Capsule, QueryTM, 0.0);
+					EXPECT_EQ(bResult, true);
+				}
+				{
+					const FRigidTransform3 QueryTM(FVec3(0.0, 0.0, 12.0), FQuat::Identity);
+					bool bResult = TriangleMesh->OverlapGeom(Capsule, QueryTM, 0.0);
+					EXPECT_EQ(bResult, true);
+				}
+				{
+					const FRigidTransform3 QueryTM(FVec3(0.0, 0.0, 12.0), FQuat(FVec3(1.0, 0.0, 0.0), 3.1415926 / 4.0));
+					bool bResult = TriangleMesh->OverlapGeom(Capsule, QueryTM, 0.0);
+					EXPECT_EQ(bResult, false);
+				}
+				{
+					// Inside Mesh
+					const FRigidTransform3 QueryTM(FVec3(0.0, 0.0, 4.0), FQuat::Identity);
+					bool bResult = TriangleMesh->OverlapGeom(Capsule, QueryTM, 0.0);
+					EXPECT_EQ(bResult, false);
+				}
+				{
+					const FRigidTransform3 QueryTM(FVec3(6.5, 0.0, 6.5), FQuat::Identity);
+					bool bResult = TriangleMesh->OverlapGeom(Capsule, QueryTM, 0.0);
+					EXPECT_TRUE(bResult);
+				}
+				{
+					const FRigidTransform3 QueryTM(FVec3(6.5, 0.0, 6.5), FQuat(FVec3(0.0, 1.0, 0.0), -3.1415926 / 4.0));
+					bool bResult = TriangleMesh->OverlapGeom(Capsule, QueryTM, 0.0);
+					EXPECT_FALSE(bResult);
+				}
+			}
+			{
+				const FVec3 X1 = { 0.0, 0.0, -200.0 };
+				const FVec3 X2 = { 0.0, 0.0, 200.0 };
+				const FReal Radius = 100.0;
+				const FCapsule Capsule = FCapsule(X1, X2, Radius);
+				TSharedPtr<FCapsule> CapsuleShared = MakeShared<FCapsule>(X1, X2, Radius);
+				FVec3 TriMeshScale = { 0.01, 0.01, 0.01 };
+				TImplicitObjectScaled<FCapsule> ScaledCapsule = TImplicitObjectScaled<FCapsule>(CapsuleShared, TriMeshScale);
+				{
+					const FRigidTransform3 QueryTM(FVec3(0.0, 0.0, 0.0), FQuat::Identity);
+					bool bResult = TriangleMesh->OverlapGeom(ScaledCapsule, QueryTM, 0.0, nullptr, TriMeshScale);
+					EXPECT_EQ(bResult, true);
+				}
+				{
+					const FRigidTransform3 QueryTM(FVec3(10.0, 10.0, 0.0), FQuat::Identity);
+					bool bResult = TriangleMesh->OverlapGeom(ScaledCapsule, QueryTM, 0.0, nullptr, TriMeshScale);
+					EXPECT_EQ(bResult, false);
+				}
+				{
+					const FRigidTransform3 QueryTM(FVec3(-10.0, -10.0, 0.0), FQuat::Identity);
+					bool bResult = TriangleMesh->OverlapGeom(ScaledCapsule, QueryTM, 0.0, nullptr, TriMeshScale);
+					EXPECT_EQ(bResult, false);
+				}
+				{
+					const FRigidTransform3 QueryTM(FVec3(10.0, 0.0, 0.0), FQuat::Identity);
+					bool bResult = TriangleMesh->OverlapGeom(ScaledCapsule, QueryTM, 0.0, nullptr, TriMeshScale);
+					EXPECT_EQ(bResult, true);
+				}
+				{
+					const FRigidTransform3 QueryTM(FVec3(0.0, 10.0, 0.0), FQuat::Identity);
+					bool bResult = TriangleMesh->OverlapGeom(ScaledCapsule, QueryTM, 0.0, nullptr, TriMeshScale);
+					EXPECT_EQ(bResult, true);
+				}
+				{
+					const FRigidTransform3 QueryTM(FVec3(0.0, -10.0, 0.0), FQuat::Identity);
+					bool bResult = TriangleMesh->OverlapGeom(ScaledCapsule, QueryTM, 0.0, nullptr, TriMeshScale);
+					EXPECT_EQ(bResult, true);
+				}
+				{
+					const FRigidTransform3 QueryTM(FVec3(-10.0, 0.0, 0.0), FQuat::Identity);
+					bool bResult = TriangleMesh->OverlapGeom(ScaledCapsule, QueryTM, 0.0, nullptr, TriMeshScale);
+					EXPECT_EQ(bResult, true);
+				}
+				{
+					const FRigidTransform3 QueryTM(FVec3(0.0, 0.0, 12.0), FQuat::Identity);
+					bool bResult = TriangleMesh->OverlapGeom(ScaledCapsule, QueryTM, 0.0, nullptr, TriMeshScale);
+					EXPECT_EQ(bResult, true);
+				}
+				{
+					const FRigidTransform3 QueryTM(FVec3(0.0, 0.0, 12.0), FQuat(FVec3(1.0, 0.0, 0.0), 3.1415926 / 4.0));
+					bool bResult = TriangleMesh->OverlapGeom(ScaledCapsule, QueryTM, 0.0, nullptr, TriMeshScale);
+					EXPECT_EQ(bResult, false);
+				}
+				{
+					// Inside Mesh
+					const FRigidTransform3 QueryTM(FVec3(0.0, 0.0, 4.0), FQuat::Identity);
+					bool bResult = TriangleMesh->OverlapGeom(ScaledCapsule, QueryTM, 0.0, nullptr, TriMeshScale);
+					EXPECT_EQ(bResult, false);
+				}
+				{
+					const FRigidTransform3 QueryTM(FVec3(6.5, 0.0, 6.5), FQuat::Identity);
+					bool bResult = TriangleMesh->OverlapGeom(ScaledCapsule, QueryTM, 0.0, nullptr, TriMeshScale);
+					EXPECT_TRUE(bResult);
+				}
+				{
+					const FRigidTransform3 QueryTM(FVec3(6.5, 0.0, 6.5), FQuat(FVec3(0.0, 1.0, 0.0), -3.1415926 / 4.0));
+					bool bResult = TriangleMesh->OverlapGeom(ScaledCapsule, QueryTM, 0.0, nullptr, TriMeshScale);
+					EXPECT_FALSE(bResult);
+				}
+			}
+			{
+				const FVec3 X1 = { 0.0, 0.0, -20.0 };
+				const FVec3 X2 = { 0.0, 0.0, 20.0 };
+				const FReal Radius = 10.0;
+				TSharedPtr<FCapsule> CapsuleShared = MakeShared<FCapsule>(X1, X2, Radius);
+				FVec3 TriMeshScale = { 10.0, 10.0, 10.0 };
+				FVec3 InvScale = 1.0 / TriMeshScale;
+				TImplicitObjectScaled<FCapsule> ScaledCapsule = TImplicitObjectScaled<FCapsule>(CapsuleShared, InvScale);
+				{
+					FRigidTransform3 QueryTM(FVec3(0.0, 0.0, 0.0), FQuat::Identity);
+					QueryTM = FRigidTransform3(QueryTM.GetLocation() * InvScale, QueryTM.GetRotation());
+					bool bResult = TriangleMesh->OverlapGeom(ScaledCapsule, QueryTM, 0.0, nullptr, TriMeshScale);
+					EXPECT_EQ(bResult, true);
+				}
+				{
+					FRigidTransform3 QueryTM(FVec3(100.0, 100.0, 0.0), FQuat::Identity);
+					QueryTM = FRigidTransform3(QueryTM.GetLocation() * InvScale, QueryTM.GetRotation());
+					bool bResult = TriangleMesh->OverlapGeom(ScaledCapsule, QueryTM, 0.0, nullptr, TriMeshScale);
+					EXPECT_EQ(bResult, false);
+				}
+				{
+					FRigidTransform3 QueryTM(FVec3(125.0, 0.0, 0.0), FQuat::Identity);
+					QueryTM = FRigidTransform3(QueryTM.GetLocation() * InvScale, QueryTM.GetRotation());
+					bool bResult = TriangleMesh->OverlapGeom(ScaledCapsule, QueryTM, 0.0, nullptr, TriMeshScale);
+					EXPECT_EQ(bResult, false);
+				}
+				{
+					FRigidTransform3 QueryTM(FVec3(-100.0, -100.0, 0.0), FQuat::Identity);
+					QueryTM = FRigidTransform3(QueryTM.GetLocation() * InvScale, QueryTM.GetRotation());
+					bool bResult = TriangleMesh->OverlapGeom(ScaledCapsule, QueryTM, 0.0, nullptr, TriMeshScale);
+					EXPECT_EQ(bResult, false);
+				}
+				{
+					FRigidTransform3 QueryTM(FVec3(100.0, 0.0, 0.0), FQuat::Identity);
+					QueryTM = FRigidTransform3(QueryTM.GetLocation() * InvScale, QueryTM.GetRotation());
+					bool bResult = TriangleMesh->OverlapGeom(ScaledCapsule, QueryTM, 0.0, nullptr, TriMeshScale);
+					EXPECT_EQ(bResult, true);
+				}
+				{
+					FRigidTransform3 QueryTM(FVec3(0.0, 100.0, 0.0), FQuat::Identity);
+					QueryTM = FRigidTransform3(QueryTM.GetLocation() * InvScale, QueryTM.GetRotation());
+					bool bResult = TriangleMesh->OverlapGeom(ScaledCapsule, QueryTM, 0.0, nullptr, TriMeshScale);
+					EXPECT_EQ(bResult, true);
+				}
+				{
+					FRigidTransform3 QueryTM(FVec3(0.0, -100.0, 0.0), FQuat::Identity);
+					QueryTM = FRigidTransform3(QueryTM.GetLocation() * InvScale, QueryTM.GetRotation());
+					bool bResult = TriangleMesh->OverlapGeom(ScaledCapsule, QueryTM, 0.0, nullptr, TriMeshScale);
+					EXPECT_EQ(bResult, true);
+				}
+				{
+					FRigidTransform3 QueryTM(FVec3(-100.0, 0.0, 0.0), FQuat::Identity);
+					QueryTM = FRigidTransform3(QueryTM.GetLocation() * InvScale, QueryTM.GetRotation());
+					bool bResult = TriangleMesh->OverlapGeom(ScaledCapsule, QueryTM, 0.0, nullptr, TriMeshScale);
+					EXPECT_EQ(bResult, true);
+				}
+				{
+					FRigidTransform3 QueryTM(FVec3(0.0, 0.0, 120.0), FQuat::Identity);
+					QueryTM = FRigidTransform3(QueryTM.GetLocation() * InvScale, QueryTM.GetRotation());
+					bool bResult = TriangleMesh->OverlapGeom(ScaledCapsule, QueryTM, 0.0, nullptr, TriMeshScale);
+					EXPECT_EQ(bResult, true);
+				}
+				{
+					FRigidTransform3 QueryTM(FVec3(0.0, 0.0, 120.0), FQuat(FVec3(1.0, 0.0, 0.0), 3.1415926 / 4.0));
+					QueryTM = FRigidTransform3(QueryTM.GetLocation() * InvScale, QueryTM.GetRotation());
+					bool bResult = TriangleMesh->OverlapGeom(ScaledCapsule, QueryTM, 0.0, nullptr, TriMeshScale);
+					EXPECT_EQ(bResult, false);
+				}
+				{
+					// Inside Mesh
+					FRigidTransform3 QueryTM(FVec3(0.0, 0.0, 40.0), FQuat::Identity);
+					QueryTM = FRigidTransform3(QueryTM.GetLocation() * InvScale, QueryTM.GetRotation());
+					bool bResult = TriangleMesh->OverlapGeom(ScaledCapsule, QueryTM, 0.0, nullptr, TriMeshScale);
+					EXPECT_EQ(bResult, false);
+				}
+				{
+					FRigidTransform3 QueryTM(FVec3(65.0, 0.0, 65.0), FQuat::Identity);
+					QueryTM = FRigidTransform3(QueryTM.GetLocation() * InvScale, QueryTM.GetRotation());
+					bool bResult = TriangleMesh->OverlapGeom(ScaledCapsule, QueryTM, 0.0, nullptr, TriMeshScale);
+					EXPECT_TRUE(bResult);
+				}
+				{
+					FRigidTransform3 QueryTM(FVec3(65.0, 0.0, 65.0), FQuat(FVec3(0.0, 1.0, 0.0), -3.1415926 / 4.0));
+					QueryTM = FRigidTransform3(QueryTM.GetLocation() * InvScale, QueryTM.GetRotation());
+					bool bResult = TriangleMesh->OverlapGeom(ScaledCapsule, QueryTM, 0.0, nullptr, TriMeshScale);
+					EXPECT_FALSE(bResult);
+				}
+			}
+			// Non uniform scale
+			{
+				const FVec3 X1 = { 0.0, 0.0, -2.0 };
+				const FVec3 X2 = { 0.0, 0.0, 2.0 };
+				const FReal Radius = 1.0;
+				TSharedPtr<FCapsule> CapsuleShared = MakeShared<FCapsule>(X1, X2, Radius);
+				FVec3 TriMeshScale = { 1.0, 1.0, 2.0 };
+				FVec3 InvScale = 1.0 / TriMeshScale;
+				TImplicitObjectScaled<FCapsule> ScaledCapsule = TImplicitObjectScaled<FCapsule>(CapsuleShared, InvScale);
+				{
+					FRigidTransform3 QueryTM(FVec3(0.0, 0.0, 0.0), FQuat::Identity);
+					QueryTM = FRigidTransform3(QueryTM.GetLocation() * InvScale, QueryTM.GetRotation());
+					bool bResult = TriangleMesh->OverlapGeom(ScaledCapsule, QueryTM, 0.0, nullptr, TriMeshScale);
+					EXPECT_EQ(bResult, true);
+				}
+				{
+					FRigidTransform3 QueryTM(FVec3(10.0, 10.0, 0.0), FQuat::Identity);
+					QueryTM = FRigidTransform3(QueryTM.GetLocation() * InvScale, QueryTM.GetRotation());
+					bool bResult = TriangleMesh->OverlapGeom(ScaledCapsule, QueryTM, 0.0, nullptr, TriMeshScale);
+					EXPECT_EQ(bResult, false);
+				}
+				{
+					FRigidTransform3 QueryTM(FVec3(12.5, 0.0, 0.0), FQuat::Identity);
+					QueryTM = FRigidTransform3(QueryTM.GetLocation() * InvScale, QueryTM.GetRotation());
+					bool bResult = TriangleMesh->OverlapGeom(ScaledCapsule, QueryTM, 0.0, nullptr, TriMeshScale);
+					EXPECT_EQ(bResult, false);
+				}
+				{
+					FRigidTransform3 QueryTM(FVec3(-10.0, -10.0, 0.0), FQuat::Identity);
+					QueryTM = FRigidTransform3(QueryTM.GetLocation() * InvScale, QueryTM.GetRotation());
+					bool bResult = TriangleMesh->OverlapGeom(ScaledCapsule, QueryTM, 0.0, nullptr, TriMeshScale);
+					EXPECT_EQ(bResult, false);
+				}
+				{
+					FRigidTransform3 QueryTM(FVec3(10.0, 0.0, 0.0), FQuat::Identity);
+					QueryTM = FRigidTransform3(QueryTM.GetLocation() * InvScale, QueryTM.GetRotation());
+					bool bResult = TriangleMesh->OverlapGeom(ScaledCapsule, QueryTM, 0.0, nullptr, TriMeshScale);
+					EXPECT_EQ(bResult, true);
+				}
+				{
+					FRigidTransform3 QueryTM(FVec3(0.0, 10.0, 0.0), FQuat::Identity);
+					QueryTM = FRigidTransform3(QueryTM.GetLocation() * InvScale, QueryTM.GetRotation());
+					bool bResult = TriangleMesh->OverlapGeom(ScaledCapsule, QueryTM, 0.0, nullptr, TriMeshScale);
+					EXPECT_EQ(bResult, true);
+				}
+				{
+					FRigidTransform3 QueryTM(FVec3(0.0, -10.0, 0.0), FQuat::Identity);
+					QueryTM = FRigidTransform3(QueryTM.GetLocation() * InvScale, QueryTM.GetRotation());
+					bool bResult = TriangleMesh->OverlapGeom(ScaledCapsule, QueryTM, 0.0, nullptr, TriMeshScale);
+					EXPECT_EQ(bResult, true);
+				}
+				{
+					FRigidTransform3 QueryTM(FVec3(-10.0, 0.0, 0.0), FQuat::Identity);
+					QueryTM = FRigidTransform3(QueryTM.GetLocation() * InvScale, QueryTM.GetRotation());
+					bool bResult = TriangleMesh->OverlapGeom(ScaledCapsule, QueryTM, 0.0, nullptr, TriMeshScale);
+					EXPECT_EQ(bResult, true);
+				}
+				{
+					FRigidTransform3 QueryTM(FVec3(0.0, 0.0, 22.0), FQuat::Identity);
+					QueryTM = FRigidTransform3(QueryTM.GetLocation() * InvScale, QueryTM.GetRotation());
+					bool bResult = TriangleMesh->OverlapGeom(ScaledCapsule, QueryTM, 0.0, nullptr, TriMeshScale);
+					EXPECT_EQ(bResult, true);
+				}
+				{
+					FRigidTransform3 QueryTM(FVec3(0.0, 0.0, 22.0), FQuat(FVec3(1.0, 0.0, 0.0), 3.1415926 / 4.0));
+					QueryTM = FRigidTransform3(QueryTM.GetLocation() * InvScale, QueryTM.GetRotation());
+					bool bResult = TriangleMesh->OverlapGeom(ScaledCapsule, QueryTM, 0.0, nullptr, TriMeshScale);
+					EXPECT_EQ(bResult, false);
+				}
+				{
+					// Inside Mesh
+					FRigidTransform3 QueryTM(FVec3(0.0, 0.0, 8.0), FQuat::Identity);
+					QueryTM = FRigidTransform3(QueryTM.GetLocation() * InvScale, QueryTM.GetRotation());
+					bool bResult = TriangleMesh->OverlapGeom(ScaledCapsule, QueryTM, 0.0, nullptr, TriMeshScale);
+					EXPECT_EQ(bResult, false);
+				}
+				{
+					FRigidTransform3 QueryTM(FVec3(6.0, 0.0, 11.5), FQuat::Identity);
+					QueryTM = FRigidTransform3(QueryTM.GetLocation() * InvScale, QueryTM.GetRotation());
+					bool bResult = TriangleMesh->OverlapGeom(ScaledCapsule, QueryTM, 0.0, nullptr, TriMeshScale);
+					EXPECT_TRUE(bResult);
+				}
+				{
+					// Triangles parallel to the capsule axis with non uniform scale
+					FRigidTransform3 QueryTM(FVec3(6.0, 0.0, 11.5), FQuat(FVec3(0.0, 1.0, 0.0), -3.1415926 / 4.0 + 3.1415926 / 8.0));
+					QueryTM = FRigidTransform3(QueryTM.GetLocation() * InvScale, QueryTM.GetRotation());
+					bool bResult = TriangleMesh->OverlapGeom(ScaledCapsule, QueryTM, 0.0, nullptr, TriMeshScale);
+					EXPECT_FALSE(bResult);
+				}
 			}
+
 			{
-				const FRigidTransform3 QueryTM(FVec3(7.0, 0.0, 10.0), FQuat(FVec3(0.0, 1.0, 0.0), -3.1415926 / 4.0));
-				bool bResult = TriangleMesh->OverlapGeom(ScaledBox, QueryTM, 0.0, nullptr, TriMeshScale);
-				EXPECT_EQ(bResult, false);
+				// Box test
+				const TBox<FReal, 3> Box = TBox<FReal, 3>({ -1.0, -1.0, -2.0 }, { 1.0, 1.0, 2.0 });
+				{
+					const FRigidTransform3 QueryTM(FVec3(0.0, 0.0, 0.0), FQuat::Identity);
+					bool bResult = TriangleMesh->OverlapGeom(Box, QueryTM, 0.0);
+					EXPECT_EQ(bResult, true);
+				}
+				{
+					const FRigidTransform3 QueryTM(FVec3(10.0, 10.0, 0.0), FQuat::Identity);
+					bool bResult = TriangleMesh->OverlapGeom(Box, QueryTM, 0.0);
+					EXPECT_EQ(bResult, false);
+				}
+				{
+					const FRigidTransform3 QueryTM(FVec3(-10.0, -10.0, 0.0), FQuat::Identity);
+					bool bResult = TriangleMesh->OverlapGeom(Box, QueryTM, 0.0);
+					EXPECT_EQ(bResult, false);
+				}
+				{
+					const FRigidTransform3 QueryTM(FVec3(10.0, 0.0, 0.0), FQuat::Identity);
+					bool bResult = TriangleMesh->OverlapGeom(Box, QueryTM, 0.0);
+					EXPECT_EQ(bResult, true);
+				}
+				{
+					const FRigidTransform3 QueryTM(FVec3(0.0, 10.0, 0.0), FQuat::Identity);
+					bool bResult = TriangleMesh->OverlapGeom(Box, QueryTM, 0.0);
+					EXPECT_EQ(bResult, true);
+				}
+				{
+					const FRigidTransform3 QueryTM(FVec3(0.0, -10.0, 0.0), FQuat::Identity);
+					bool bResult = TriangleMesh->OverlapGeom(Box, QueryTM, 0.0);
+					EXPECT_EQ(bResult, true);
+				}
+				{
+					const FRigidTransform3 QueryTM(FVec3(-10.0, 0.0, 0.0), FQuat::Identity);
+					bool bResult = TriangleMesh->OverlapGeom(Box, QueryTM, 0.0);
+					EXPECT_EQ(bResult, true);
+				}
+				{
+					const FRigidTransform3 QueryTM(FVec3(0.0, 0.0, 12.0), FQuat::Identity);
+					bool bResult = TriangleMesh->OverlapGeom(Box, QueryTM, 0.0);
+					EXPECT_EQ(bResult, true);
+				}
+				{
+					const FRigidTransform3 QueryTM(FVec3(0.0, 0.0, 12.0), FQuat(FVec3(1.0, 0.0, 0.0), 3.1415926 / 4.0));
+					bool bResult = TriangleMesh->OverlapGeom(Box, QueryTM, 0.0);
+					EXPECT_EQ(bResult, false);
+				}
+				{
+					// Inside Mesh
+					const FRigidTransform3 QueryTM(FVec3(0.0, 0.0, 4.0), FQuat::Identity);
+					bool bResult = TriangleMesh->OverlapGeom(Box, QueryTM, 0.0);
+					EXPECT_EQ(bResult, false);
+				}
 			}
-			{
-				const FRigidTransform3 QueryTM(FVec3(7.0, 0.0, 10.0), FQuat(FVec3(0.0, 1.0, 0.0), 3.1415926 / 4.0));
-				bool bResult = TriangleMesh->OverlapGeom(ScaledBox, QueryTM, 0.0, nullptr, TriMeshScale);
-				EXPECT_EQ(bResult, true);
+
+			// Box Scaled
+			{
+				TSharedPtr<TBox<FReal, 3>> BigBoxSafe = MakeShared<TBox<FReal, 3>>(FVec3(-100.0, -100.0, -100.0), FVec3(100.0, 100.0, 100.0));
+				FVec3 TriMeshScale = { 0.01, 0.01, 0.02 };
+				TImplicitObjectScaled<TBox<FReal, 3>> ScaledBox = TImplicitObjectScaled<TBox<FReal, 3>>(BigBoxSafe, TriMeshScale);
+				{
+					const FRigidTransform3 QueryTM(FVec3(0.0, 0.0, 0.0), FQuat::Identity);
+					bool bResult = TriangleMesh->OverlapGeom(ScaledBox, QueryTM, 0.0, nullptr, TriMeshScale);
+					EXPECT_EQ(bResult, true);
+				}
+				{
+					const FRigidTransform3 QueryTM(FVec3(10.0, 10.0, 0.0), FQuat::Identity);
+					bool bResult = TriangleMesh->OverlapGeom(ScaledBox, QueryTM, 0.0, nullptr, TriMeshScale);
+					EXPECT_EQ(bResult, false);
+				}
+				{
+					const FRigidTransform3 QueryTM(FVec3(-10.0, -10.0, 0.0), FQuat::Identity);
+					bool bResult = TriangleMesh->OverlapGeom(ScaledBox, QueryTM, 0.0, nullptr, TriMeshScale);
+					EXPECT_EQ(bResult, false);
+				}
+				{
+					const FRigidTransform3 QueryTM(FVec3(10.0, 0.0, 0.0), FQuat::Identity);
+					bool bResult = TriangleMesh->OverlapGeom(ScaledBox, QueryTM, 0.0, nullptr, TriMeshScale);
+					EXPECT_EQ(bResult, true);
+				}
+				{
+					const FRigidTransform3 QueryTM(FVec3(0.0, 10.0, 0.0), FQuat::Identity);
+					bool bResult = TriangleMesh->OverlapGeom(ScaledBox, QueryTM, 0.0, nullptr, TriMeshScale);
+					EXPECT_EQ(bResult, true);
+				}
+				{
+					const FRigidTransform3 QueryTM(FVec3(0.0, -10.0, 0.0), FQuat::Identity);
+					bool bResult = TriangleMesh->OverlapGeom(ScaledBox, QueryTM, 0.0, nullptr, TriMeshScale);
+					EXPECT_EQ(bResult, true);
+				}
+				{
+					const FRigidTransform3 QueryTM(FVec3(-10.0, 0.0, 0.0), FQuat::Identity);
+					bool bResult = TriangleMesh->OverlapGeom(ScaledBox, QueryTM, 0.0, nullptr, TriMeshScale);
+					EXPECT_EQ(bResult, true);
+				}
+				{
+					const FRigidTransform3 QueryTM(FVec3(0.0, 0.0, 12.0), FQuat::Identity);
+					bool bResult = TriangleMesh->OverlapGeom(ScaledBox, QueryTM, 0.0, nullptr, TriMeshScale);
+					EXPECT_EQ(bResult, true);
+				}
+				{
+					const FRigidTransform3 QueryTM(FVec3(0.0, 0.0, 12.0), FQuat(FVec3(1.0, 0.0, 0.0), 3.1415926 / 4.0));
+					bool bResult = TriangleMesh->OverlapGeom(ScaledBox, QueryTM, 0.0, nullptr, TriMeshScale);
+					EXPECT_EQ(bResult, false);
+				}
+				{
+					// Inside Mesh
+					const FRigidTransform3 QueryTM(FVec3(0.0, 0.0, 4.0), FQuat::Identity);
+					bool bResult = TriangleMesh->OverlapGeom(ScaledBox, QueryTM, 0.0, nullptr, TriMeshScale);
+					EXPECT_EQ(bResult, false);
+				}
+			}
+			// Box Scaled test non uniform transform with rotation
+			{
+				TSharedPtr<TBox<FReal, 3>> BigBoxSafe = MakeShared<TBox<FReal, 3>>(FVec3(-100.0, -100.0, -100.0), FVec3(100.0, 100.0, 100.0));
+				FVec3 TriMeshScale = { 0.01, 0.01, 0.05 };
+				TImplicitObjectScaled<TBox<FReal, 3>> ScaledBox = TImplicitObjectScaled<TBox<FReal, 3>>(BigBoxSafe, TriMeshScale);
+
+				{
+					const FRigidTransform3 QueryTM(FVec3(5.5, 0.0, 10.0), FQuat::Identity);
+					bool bResult = TriangleMesh->OverlapGeom(ScaledBox, QueryTM, 0.0, nullptr, TriMeshScale);
+					EXPECT_EQ(bResult, true);
+				}
+				{
+					const FRigidTransform3 QueryTM(FVec3(7.0, 0.0, 10.0), FQuat::Identity);
+					bool bResult = TriangleMesh->OverlapGeom(ScaledBox, QueryTM, 0.0, nullptr, TriMeshScale);
+					EXPECT_EQ(bResult, false);
+				}
+				{
+					const FRigidTransform3 QueryTM(FVec3(7.0, 0.0, 10.0), FQuat(FVec3(0.0, 1.0, 0.0), -3.1415926 / 4.0));
+					bool bResult = TriangleMesh->OverlapGeom(ScaledBox, QueryTM, 0.0, nullptr, TriMeshScale);
+					EXPECT_EQ(bResult, false);
+				}
+				{
+					const FRigidTransform3 QueryTM(FVec3(7.0, 0.0, 10.0), FQuat(FVec3(0.0, 1.0, 0.0), 3.1415926 / 4.0));
+					bool bResult = TriangleMesh->OverlapGeom(ScaledBox, QueryTM, 0.0, nullptr, TriMeshScale);
+					EXPECT_EQ(bResult, true);
+				}
+			}
+			{
+				TSharedPtr<TBox<FReal, 3>> BigBoxSafe = MakeShared<TBox<FReal, 3>>(FVec3(-1.0, -1.0, -1.0), FVec3(1.0, 1.0, 1.0));
+				FVec3 TriMeshScale = { 10.0, 10.0, 2.0 };
+				FVec3 InvScale = 1.0 / TriMeshScale;
+				TImplicitObjectScaled<TBox<FReal, 3>> ScaledBox = TImplicitObjectScaled<TBox<FReal, 3>>(BigBoxSafe, InvScale);
+
+				{
+					FRigidTransform3 QueryTM(FVec3(50, 0.0, 10.0), FQuat::Identity);
+					QueryTM = FRigidTransform3(QueryTM.GetLocation() * InvScale, QueryTM.GetRotation());
+					bool bResult = TriangleMesh->OverlapGeom(ScaledBox, QueryTM, 0.0, nullptr, TriMeshScale);
+					EXPECT_EQ(bResult, true);
+				}
+				{
+					// Inside mesh
+					FRigidTransform3 QueryTM(FVec3(40.0, 0.0, 10.0), FQuat::Identity);
+					QueryTM = FRigidTransform3(QueryTM.GetLocation() * InvScale, QueryTM.GetRotation());
+					bool bResult = TriangleMesh->OverlapGeom(ScaledBox, QueryTM, 0.0, nullptr, TriMeshScale);
+					EXPECT_EQ(bResult, false);
+				}
+				{
+					FRigidTransform3 QueryTM(FVec3(52.0, 0.0, 11.0), FQuat::Identity);
+					QueryTM = FRigidTransform3(QueryTM.GetLocation() * InvScale, QueryTM.GetRotation());
+					bool bResult = TriangleMesh->OverlapGeom(ScaledBox, QueryTM, 0.0, nullptr, TriMeshScale);
+					EXPECT_EQ(bResult, false);
+				}
+				{
+					FRigidTransform3 QueryTM(FVec3(52.0, 0.0, 11.0), FQuat(FVec3(0.0, 1.0, 0.0), -3.1415926 / 4.0));
+					QueryTM = FRigidTransform3(QueryTM.GetLocation() * InvScale, QueryTM.GetRotation());
+					bool bResult = TriangleMesh->OverlapGeom(ScaledBox, QueryTM, 0.0, nullptr, TriMeshScale);
+					EXPECT_EQ(bResult, true);
+				}
+				{
+					FRigidTransform3 QueryTM(FVec3(51.0, 0.0, 11.0), FQuat::Identity);
+					QueryTM = FRigidTransform3(QueryTM.GetLocation() * InvScale, QueryTM.GetRotation());
+					bool bResult = TriangleMesh->OverlapGeom(ScaledBox, QueryTM, 0.0, nullptr, TriMeshScale);
+					EXPECT_EQ(bResult, true);
+				}
+				{
+					FRigidTransform3 QueryTM(FVec3(51.5, 0.0, 11.0), FQuat(FVec3(0.0, 1.0, 0.0), 4.0 / 5.0 * 3.1415926 / 2.0));
+					QueryTM = FRigidTransform3(QueryTM.GetLocation() * InvScale, QueryTM.GetRotation());
+					bool bResult = TriangleMesh->OverlapGeom(ScaledBox, QueryTM, 0.0, nullptr, TriMeshScale);
+					EXPECT_EQ(bResult, true);
+				}
+				{
+					FRigidTransform3 QueryTM(FVec3(51.5, 0.0, 11.0), FQuat(FVec3(0.0, 1.0, 0.0), -4.0 / 5.0 * 3.1415926 / 2.0));
+					QueryTM = FRigidTransform3(QueryTM.GetLocation() * InvScale, QueryTM.GetRotation());
+					bool bResult = TriangleMesh->OverlapGeom(ScaledBox, QueryTM, 0.0, nullptr, TriMeshScale);
+					EXPECT_EQ(bResult, false);
+				}
+			}
+			{
+				// Non uniform test with box not being a cube
+				TSharedPtr<TBox<FReal, 3>> BigBoxSafe = MakeShared<TBox<FReal, 3>>(FVec3(-1.0, -5.0, -1.0), FVec3(1.0, 5.0, 1.0));
+				FVec3 TriMeshScale = { 10.0, 10.0, 2.0 };
+				FVec3 InvScale = 1.0 / TriMeshScale;
+				TImplicitObjectScaled<TBox<FReal, 3>> ScaledBox = TImplicitObjectScaled<TBox<FReal, 3>>(BigBoxSafe, InvScale);
+
+				{
+					FRigidTransform3 QueryTM(FVec3(50, 0.0, 10.0), FQuat::Identity);
+					QueryTM = FRigidTransform3(QueryTM.GetLocation() * InvScale, QueryTM.GetRotation());
+					bool bResult = TriangleMesh->OverlapGeom(ScaledBox, QueryTM, 0.0, nullptr, TriMeshScale);
+					EXPECT_EQ(bResult, true);
+				}
+				{
+					// Inside mesh
+					FRigidTransform3 QueryTM(FVec3(40.0, 0.0, 10.0), FQuat::Identity);
+					QueryTM = FRigidTransform3(QueryTM.GetLocation() * InvScale, QueryTM.GetRotation());
+					bool bResult = TriangleMesh->OverlapGeom(ScaledBox, QueryTM, 0.0, nullptr, TriMeshScale);
+					EXPECT_EQ(bResult, true);
+				}
+				{
+					FRigidTransform3 QueryTM(FVec3(52.0, 0.0, 11.0), FQuat::Identity);
+					QueryTM = FRigidTransform3(QueryTM.GetLocation() * InvScale, QueryTM.GetRotation());
+					bool bResult = TriangleMesh->OverlapGeom(ScaledBox, QueryTM, 0.0, nullptr, TriMeshScale);
+					EXPECT_EQ(bResult, false);
+				}
+				{
+					FRigidTransform3 QueryTM(FVec3(0.0, 52.0, 11.0), FQuat::Identity);
+					QueryTM = FRigidTransform3(QueryTM.GetLocation() * InvScale, QueryTM.GetRotation());
+					bool bResult = TriangleMesh->OverlapGeom(ScaledBox, QueryTM, 0.0, nullptr, TriMeshScale);
+					EXPECT_EQ(bResult, true);
+				}
+				{
+					FRigidTransform3 QueryTM(FVec3(54.0, 0.0, 11.0), FQuat(FVec3(0.0, 1.0, 0.0), -3.1415926 / 2.0 * (1.5 / 5.0)));
+					QueryTM = FRigidTransform3(QueryTM.GetLocation() * InvScale, QueryTM.GetRotation());
+					bool bResult = TriangleMesh->OverlapGeom(ScaledBox, QueryTM, 0.0, nullptr, TriMeshScale);
+					EXPECT_EQ(bResult, false);
+				}
+				{
+					FRigidTransform3 QueryTM(FVec3(0.0, 54.0, 11.0), FQuat(FVec3(1.0, 0.0, 0.0), -3.1415926 / 2.0 * (1.5 / 5.0)));
+					QueryTM = FRigidTransform3(QueryTM.GetLocation() * InvScale, QueryTM.GetRotation());
+					bool bResult = TriangleMesh->OverlapGeom(ScaledBox, QueryTM, 0.0, nullptr, TriMeshScale);
+					EXPECT_EQ(bResult, true);
+				}
+				{
+					FRigidTransform3 QueryTM(FVec3(52.0, 0.0, 11.0), FQuat(FVec3(0.0, 1.0, 0.0), (1.5 / 5.0) * 3.1415926 / 2.0));
+					QueryTM = FRigidTransform3(QueryTM.GetLocation() * InvScale, QueryTM.GetRotation());
+					bool bResult = TriangleMesh->OverlapGeom(ScaledBox, QueryTM, 0.0, nullptr, TriMeshScale);
+					EXPECT_EQ(bResult, false);
+				}
+				{
+					FRigidTransform3 QueryTM(FVec3(0.0, 52.0, 11.0), FQuat(FVec3(1.0, 0.0, 0.0), (1.5 / 5.0) * 3.1415926 / 2.0));
+					QueryTM = FRigidTransform3(QueryTM.GetLocation() * InvScale, QueryTM.GetRotation());
+					bool bResult = TriangleMesh->OverlapGeom(ScaledBox, QueryTM, 0.0, nullptr, TriMeshScale);
+					EXPECT_EQ(bResult, true);
+				}
+			}
+			{
+				// Sphere test
+				const TSphere<FReal, 3> Sphere = TSphere<FReal, 3>({ 0.0, 0.0, 0.0 }, 2.0);
+				{
+					const FRigidTransform3 QueryTM(FVec3(0.0, 0.0, 0.0), FQuat::Identity);
+					bool bResult = TriangleMesh->OverlapGeom(Sphere, QueryTM, 0.0);
+					EXPECT_EQ(bResult, true);
+				}
+				{
+					const FRigidTransform3 QueryTM(FVec3(10.0, 10.0, 0.0), FQuat::Identity);
+					bool bResult = TriangleMesh->OverlapGeom(Sphere, QueryTM, 0.0);
+					EXPECT_EQ(bResult, false);
+				}
+				{
+					const FRigidTransform3 QueryTM(FVec3(-10.0, -10.0, 0.0), FQuat::Identity);
+					bool bResult = TriangleMesh->OverlapGeom(Sphere, QueryTM, 0.0);
+					EXPECT_EQ(bResult, false);
+				}
+				{
+					const FRigidTransform3 QueryTM(FVec3(10.0, 0.0, 0.0), FQuat::Identity);
+					bool bResult = TriangleMesh->OverlapGeom(Sphere, QueryTM, 0.0);
+					EXPECT_EQ(bResult, true);
+				}
+				{
+					const FRigidTransform3 QueryTM(FVec3(0.0, 10.0, 0.0), FQuat::Identity);
+					bool bResult = TriangleMesh->OverlapGeom(Sphere, QueryTM, 0.0);
+					EXPECT_EQ(bResult, true);
+				}
+				{
+					const FRigidTransform3 QueryTM(FVec3(0.0, -10.0, 0.0), FQuat::Identity);
+					bool bResult = TriangleMesh->OverlapGeom(Sphere, QueryTM, 0.0);
+					EXPECT_EQ(bResult, true);
+				}
+				{
+					const FRigidTransform3 QueryTM(FVec3(-10.0, 0.0, 0.0), FQuat::Identity);
+					bool bResult = TriangleMesh->OverlapGeom(Sphere, QueryTM, 0.0);
+					EXPECT_EQ(bResult, true);
+				}
+				{
+					const FRigidTransform3 QueryTM(FVec3(0.0, 0.0, 12.0), FQuat::Identity);
+					bool bResult = TriangleMesh->OverlapGeom(Sphere, QueryTM, 0.0);
+					EXPECT_EQ(bResult, true);
+				}
+				{
+					const FRigidTransform3 QueryTM(FVec3(0.0, 0.0, 12.0), FQuat(FVec3(1.0, 0.0, 0.0), 3.1415926 / 4.0));
+					bool bResult = TriangleMesh->OverlapGeom(Sphere, QueryTM, 0.0);
+					EXPECT_EQ(bResult, true);
+				}
+				{
+					// Inside Mesh
+					const FRigidTransform3 QueryTM(FVec3(0.0, 0.0, 4.0), FQuat::Identity);
+					bool bResult = TriangleMesh->OverlapGeom(Sphere, QueryTM, 0.0);
+					EXPECT_EQ(bResult, false);
+				}
 			}
 		}
 		{
-			// Sphere test
-			const TSphere<FReal, 3> Sphere = TSphere<FReal, 3>({ 0.0, 0.0, 0.0 }, 2.0);
-			{
-				const FRigidTransform3 QueryTM(FVec3(0.0, 0.0, 0.0), FQuat::Identity);
-				bool bResult = TriangleMesh->OverlapGeom(Sphere, QueryTM, 0.0);
-				EXPECT_EQ(bResult, true);
-			}
-			{
-				const FRigidTransform3 QueryTM(FVec3(10.0, 10.0, 0.0), FQuat::Identity);
-				bool bResult = TriangleMesh->OverlapGeom(Sphere, QueryTM, 0.0);
-				EXPECT_EQ(bResult, false);
-			}
-			{
-				const FRigidTransform3 QueryTM(FVec3(-10.0, -10.0, 0.0), FQuat::Identity);
-				bool bResult = TriangleMesh->OverlapGeom(Sphere, QueryTM, 0.0);
-				EXPECT_EQ(bResult, false);
-			}
-			{
-				const FRigidTransform3 QueryTM(FVec3(10.0, 0.0, 0.0), FQuat::Identity);
-				bool bResult = TriangleMesh->OverlapGeom(Sphere, QueryTM, 0.0);
-				EXPECT_EQ(bResult, true);
-			}
-			{
-				const FRigidTransform3 QueryTM(FVec3(0.0, 10.0, 0.0), FQuat::Identity);
-				bool bResult = TriangleMesh->OverlapGeom(Sphere, QueryTM, 0.0);
-				EXPECT_EQ(bResult, true);
-			}
-			{
-				const FRigidTransform3 QueryTM(FVec3(0.0, -10.0, 0.0), FQuat::Identity);
-				bool bResult = TriangleMesh->OverlapGeom(Sphere, QueryTM, 0.0);
-				EXPECT_EQ(bResult, true);
-			}
-			{
-				const FRigidTransform3 QueryTM(FVec3(-10.0, 0.0, 0.0), FQuat::Identity);
-				bool bResult = TriangleMesh->OverlapGeom(Sphere, QueryTM, 0.0);
-				EXPECT_EQ(bResult, true);
-			}
-			{
-				const FRigidTransform3 QueryTM(FVec3(0.0, 0.0, 12.0), FQuat::Identity);
-				bool bResult = TriangleMesh->OverlapGeom(Sphere, QueryTM, 0.0);
-				EXPECT_EQ(bResult, true);
-			}
+			FTriangleMeshImplicitObject::ParticlesType TrimeshParticles(
+				{
+					{-10.0, -10.0, 0.0},
+					{10.0, -10.0, 0.0},
+					{-10.0, 10.0, 0.0},
+					{10.0, 10.0, 0.0},
+					{-10.0, -10.0, 10.0},
+					{10.0, -10.0, 10.0},
+					{-10.0, 10.0, 10.0},
+					{10.0, 10.0, 10.0},
+
+				});
+
+			TArray<TVec3<int32>> Indices(
+				{
+					{0, 1, 2},
+					{0, 3, 1},
+					{0, 1, 2},
+					{4, 5, 6},
+					{4, 7, 5},
+					{4, 5, 6},
+				});
+
+			TArray<uint16> Materials;
+			for (int32 i = 0; i < Indices.Num(); ++i)
 			{
-				const FRigidTransform3 QueryTM(FVec3(0.0, 0.0, 12.0), FQuat(FVec3(1.0, 0.0, 0.0), 3.1415926 / 4.0));
-				bool bResult = TriangleMesh->OverlapGeom(Sphere, QueryTM, 0.0);
-				EXPECT_EQ(bResult, true);
+				Materials.Emplace(0);
 			}
+
+			TUniquePtr<FTriangleMeshImplicitObject> TriangleMesh = MakeUnique<FTriangleMeshImplicitObject>(MoveTemp(TrimeshParticles), MoveTemp(Indices), MoveTemp(Materials));
 			{
-				// Inside Mesh
-				const FRigidTransform3 QueryTM(FVec3(0.0, 0.0, 4.0), FQuat::Identity);
-				bool bResult = TriangleMesh->OverlapGeom(Sphere, QueryTM, 0.0);
-				EXPECT_EQ(bResult, false);
+				// Capsule test
+				const FVec3 X1 = { 0.0, 0.0, -2.0 };
+				const FVec3 X2 = { 0.0, 0.0, 2.0 };
+				const FReal Radius = 1.0;
+				const FCapsule Capsule = FCapsule(X1, X2, Radius);
+
+				{
+					const FRigidTransform3 QueryTM(FVec3(0.0, 0.0, 5.0), FQuat::Identity);
+					bool bResult = TriangleMesh->OverlapGeom(Capsule, QueryTM, 0.0);
+					EXPECT_EQ(bResult, false);
+				}
+				{
+					const FRigidTransform3 QueryTM(FVec3(10.0, 0.0, 11.0), FQuat::Identity);
+					bool bResult = TriangleMesh->OverlapGeom(Capsule, QueryTM, 0.0);
+					EXPECT_EQ(bResult, true);
+				}
+				{
+					const FRigidTransform3 QueryTM(FVec3(11.5, 0.0, 11.0), FQuat::Identity);
+					bool bResult = TriangleMesh->OverlapGeom(Capsule, QueryTM, 0.0);
+					EXPECT_EQ(bResult, false);
+				}
+				{
+					const FRigidTransform3 QueryTM(FVec3(11.5, 0.0, 11.0), FQuat(FVec3(0.0, 1.0, 0.0), 3.1415926 / 4.0));
+					bool bResult = TriangleMesh->OverlapGeom(Capsule, QueryTM, 0.0);
+					EXPECT_EQ(bResult, true);
+				}
+				{
+					const FRigidTransform3 QueryTM(FVec3(11.5, 0.0, 11.0), FQuat(FVec3(0.0, 1.0, 0.0), -3.1415926 / 4.0));
+					bool bResult = TriangleMesh->OverlapGeom(Capsule, QueryTM, 0.0);
+					EXPECT_EQ(bResult, false);
+				}
 			}
 		}
 	}
-}
\ No newline at end of file
+}
diff --git a/Engine/Source/Runtime/Experimental/Chaos/Private/Chaos/OBBVectorized.cpp b/Engine/Source/Runtime/Experimental/Chaos/Private/Chaos/OBBVectorized.cpp
new file mode 100644
--- /dev/null
+++ b/Engine/Source/Runtime/Experimental/Chaos/Private/Chaos/OBBVectorized.cpp
@@ -0,0 +1,63 @@
+// Copyright Epic Games, Inc. All Rights Reserved.
+
+#include "Chaos/OBBVectorized.h"
+#include "Chaos/Core.h"
+#include "Chaos/Transform.h"
+#include "Chaos/AABBVectorized.h"
+
+namespace Chaos
+{
+
+	FOBBVectorized::FOBBVectorized(const FRigidTransform3& Transform, const FVec3f& HalfExtentsIn)
+	{
+		TRotation<FRealSingle, 3> Rotation = Transform.GetRotation();
+		PMatrix<FRealSingle, 3, 3> Matrix = Rotation.ToMatrix();
+		XAxis = VectorLoadFloat3(Matrix.M[0]);
+		YAxis = VectorLoadFloat3(Matrix.M[1]);
+		ZAxis = VectorLoadFloat3(Matrix.M[2]);
+
+		FVec3 Translation = Transform.GetTranslation();
+		Position = MakeVectorRegisterFloatFromDouble(VectorLoadDouble3(&Translation.X));
+
+		HalfExtents = VectorLoadFloat3(&HalfExtentsIn.X);
+		XHalfExtent = VectorReplicate(HalfExtents, 0);
+		YHalfExtent = VectorReplicate(HalfExtents, 1);
+		ZHalfExtent = VectorReplicate(HalfExtents, 2);
+
+		VectorRegister4Float FurthestPoint = VectorMultiplyAdd(VectorAbs(XAxis), XHalfExtent, VectorMultiplyAdd(VectorAbs(YAxis), YHalfExtent, VectorMultiply(VectorAbs(ZAxis), ZHalfExtent)));
+		MaxObb = VectorAdd(Position, FurthestPoint);
+		MinObb = VectorSubtract(Position, FurthestPoint);
+	}
+
+	bool FOBBVectorized::IntersectAABB(const FAABBVectorized& Bounds) const
+	{
+		VectorRegister4Float HasSeparationAxis = VectorBitwiseOr(VectorCompareGT(MinObb, Bounds.GetMax()), VectorCompareGT(Bounds.GetMin(), MaxObb));
+
+		if (VectorMaskBits(HasSeparationAxis) != 0)
+		{
+			return false;
+		}
+
+		VectorRegister4Float HalfAabb = VectorMultiply(VectorSubtract(Bounds.GetMax(), Bounds.GetMin()), GlobalVectorConstants::FloatOneHalf);
+		VectorRegister4Float CenterAabb = VectorMultiply(VectorAdd(Bounds.GetMax(), Bounds.GetMin()), GlobalVectorConstants::FloatOneHalf);
+		VectorRegister4Float DirCenter = VectorSubtract(CenterAabb, Position);
+
+		VectorRegister4Float DistCenterX = VectorAbs(VectorDot3(XAxis, DirCenter));
+		VectorRegister4Float DistCenterY = VectorAbs(VectorDot3(YAxis, DirCenter));
+		VectorRegister4Float DistCenterZ = VectorAbs(VectorDot3(ZAxis, DirCenter));
+
+		VectorRegister4Float DistCenter = VectorUnpackLo(DistCenterX, DistCenterY);
+		DistCenter = VectorMoveLh(DistCenter, DistCenterZ);
+		
+		VectorRegister4Float DistCornerX = VectorDot3(VectorAbs(XAxis), HalfAabb);
+		VectorRegister4Float DistCornerY = VectorDot3(VectorAbs(YAxis), HalfAabb);
+		VectorRegister4Float DistCornerZ = VectorDot3(VectorAbs(ZAxis), HalfAabb);
+
+		VectorRegister4Float DistCorner = VectorUnpackLo(DistCornerX, DistCornerY);
+		DistCorner = VectorMoveLh(DistCorner, DistCornerZ);
+
+		uint32 MaskBitTrue = VectorMaskBits(VectorCompareGT(VectorSubtract(DistCenter, DistCorner), HalfExtents));
+
+		return !(MaskBitTrue & 7);
+	}
+}
\ No newline at end of file
diff --git a/Engine/Source/Runtime/Experimental/Chaos/Private/Chaos/TriangleMeshImplicitObject.cpp b/Engine/Source/Runtime/Experimental/Chaos/Private/Chaos/TriangleMeshImplicitObject.cpp
--- a/Engine/Source/Runtime/Experimental/Chaos/Private/Chaos/TriangleMeshImplicitObject.cpp
+++ b/Engine/Source/Runtime/Experimental/Chaos/Private/Chaos/TriangleMeshImplicitObject.cpp
@@ -382,42 +382,6 @@ struct FTriangleMeshOverlapVisitorNoMTD
 		Translation = MakeVectorRegisterFloatFromDouble(MakeVectorRegister(TranslationDouble.X, TranslationDouble.Y, TranslationDouble.Z, 0.0));
 	}
 
-	// Default shape type specific culling
-	// Type specialization to follow
-	template <typename ShapeType>
-	bool ShapeTypeAdditionalCulling(const VectorRegister4Float& A, const VectorRegister4Float& B, const VectorRegister4Float& C)
-	{
-		return false;
-	}
-
-	template <>
-	bool ShapeTypeAdditionalCulling<Chaos::FCapsule>(const VectorRegister4Float& A, const VectorRegister4Float& B, const VectorRegister4Float& C)
-	{
-		const VectorRegister4Float InvRotation = VectorQuaternionInverse(Rotation);
-		const VectorRegister4Float NegTranslation = VectorNegate(Translation);
-		const VectorRegister4Float ATxSimd = VectorAdd(VectorQuaternionRotateVector(InvRotation, A), NegTranslation);
-		const VectorRegister4Float BTxSimd = VectorAdd(VectorQuaternionRotateVector(InvRotation, B), NegTranslation);
-		const VectorRegister4Float CTxSimd = VectorAdd(VectorQuaternionRotateVector(InvRotation, C), NegTranslation);
-
-		const VectorRegister4Float MinBounds = VectorMin(VectorMin(ATxSimd, BTxSimd), CTxSimd);
-		const VectorRegister4Float MaxBounds = VectorMax(VectorMax(ATxSimd, BTxSimd), CTxSimd);
-		FAABBVectorized GeometrySpaceAABB(MinBounds, MaxBounds);
-		FAABB3 GeometryAABB = QueryGeom.BoundingBox();
-		FAABBVectorized VecGeomAABB(GeometryAABB);
-
-		if (!VecGeomAABB.Intersects(GeometrySpaceAABB))
-		{
-			return true;
-		}
-		return false;
-	}
-
-	template <>
-	bool ShapeTypeAdditionalCulling<Chaos::TImplicitObjectScaled<Chaos::FCapsule, 1>>(const VectorRegister4Float& A, const VectorRegister4Float& B, const VectorRegister4Float& C)
-	{
-		return ShapeTypeAdditionalCulling<Chaos::FCapsule>(A,B,C);
-	}
-	
 	bool VisitOverlap(int32 TriIdx)
 	{
 		FVec3 A, B, C;
@@ -434,11 +398,6 @@ struct FTriangleMeshOverlapVisitorNoMTD
 		const VectorRegister4Float BSimd = MakeVectorRegisterFloatFromDouble(MakeVectorRegister(B.X, B.Y, B.Z, 0.0));
 		const VectorRegister4Float CSimd = MakeVectorRegisterFloatFromDouble(MakeVectorRegister(C.X, C.Y, C.Z, 0.0));
 
-		if (ShapeTypeAdditionalCulling<QueryGeomType>(ASimd, BSimd, CSimd))
-		{
-			return true;
-		}
-
 		const VectorRegister4Float AB = VectorSubtract(BSimd, ASimd);
 		const VectorRegister4Float AC = VectorSubtract(CSimd, ASimd);
 
@@ -481,6 +440,16 @@ bool FTrimeshBVH::FindAllIntersectionsNoMTD(const FAABB3& Intersection, const TR
 {
 	FTriangleMeshOverlapVisitorNoMTD Visitor(Transform, QueryGeom, Thickness, TriMeshScale, TriMesh);
 	Overlap(FAABBVectorized(Intersection), Visitor);
+
+	return Visitor.bFoundIntersection;
+}
+
+template <typename QueryGeomType>
+bool FTrimeshBVH::FindAllIntersectionsNoMTD(const FOBBVectorized& Intersection, const TRigidTransform<FReal, 3>& Transform, const QueryGeomType& QueryGeom, FReal Thickness, const FVec3& TriMeshScale, const FTriangleMeshImplicitObject* TriMesh) const
+{
+	FTriangleMeshOverlapVisitorNoMTD Visitor(Transform, QueryGeom, Thickness, TriMeshScale, TriMesh);
+	OverlapOBB(Intersection, Visitor);
+
 	return Visitor.bFoundIntersection;
 }
 
@@ -744,26 +713,24 @@ bool FTriangleMeshImplicitObject::Overlap(const FVec3& Point, const FReal Thickn
 }
 
 
-template <typename QueryGeomType>
+template <bool IsSpherical, typename QueryGeomType>
 bool FTriangleMeshImplicitObject::OverlapGeomImp(const QueryGeomType& QueryGeom, const FRigidTransform3& QueryTM, const FReal Thickness, FMTDInfo* OutMTD, FVec3 TriMeshScale) const
 {
-	bool bResult = false;
 
 	const QueryGeomType& WorldScaleQueryGeom = ScaleGeomIntoWorldHelper(QueryGeom, TriMeshScale);
 
-	const FVec3 InvTriMeshScale = FVec3(FReal(1) / TriMeshScale.X, FReal(1) / TriMeshScale.Y, FReal(1) / TriMeshScale.Z);
-
-	// IMPORTANT QueryTM comes with a invscaled translation so we need a version of the TM with world space translation to properly compute the bounds
-	FRigidTransform3 TriMeshToGeomNoScale{ QueryTM };
-	TriMeshToGeomNoScale.SetTranslation(TriMeshToGeomNoScale.GetTranslation() * TriMeshScale);
-	// NOTE: BVH test is done in tri-mesh local space (whereas collision detection is done in world space because you can't non-uniformly scale all shapes)
-	FAABB3 QueryBounds = WorldScaleQueryGeom.CalculateTransformedBounds(TriMeshToGeomNoScale);
-	QueryBounds.ThickenSymmetrically(FVec3(Thickness));
-	QueryBounds.ScaleWithNegative(InvTriMeshScale);
-
-
 	if (OutMTD)
 	{
+		const FVec3 InvTriMeshScale = FVec3(FReal(1) / TriMeshScale.X, FReal(1) / TriMeshScale.Y, FReal(1) / TriMeshScale.Z);
+		// IMPORTANT QueryTM comes with a invscaled translation so we need a version of the TM with world space translation to properly compute the bounds
+		FRigidTransform3 TriMeshToGeomNoScale{ QueryTM };
+		TriMeshToGeomNoScale.SetTranslation(TriMeshToGeomNoScale.GetTranslation() * TriMeshScale);
+
+		// NOTE: BVH test is done in tri-mesh local space (whereas collision detection is done in world space because you can't non-uniformly scale all shapes)
+		FAABB3 QueryBounds = WorldScaleQueryGeom.CalculateTransformedBounds(TriMeshToGeomNoScale);
+		QueryBounds.ThickenSymmetrically(FVec3(Thickness));
+		QueryBounds.ScaleWithNegative(InvTriMeshScale);
+
 		const TArray<int32> PotentialIntersections = FastBVH.FindAllIntersections(QueryBounds);
 
 		OutMTD->Normal = FVec3(0.0);
@@ -813,50 +780,71 @@ bool FTriangleMeshImplicitObject::OverlapGeomImp(const QueryGeomType& QueryGeom,
 	}
 	else
 	{
-		TRigidTransform<FReal, 3> WorldScaleQueryTM;
-		ScaleTransformHelper(TriMeshScale, QueryTM, WorldScaleQueryTM);
-		return FastBVH.FindAllIntersectionsNoMTD(QueryBounds, WorldScaleQueryTM, WorldScaleQueryGeom, Thickness, TriMeshScale, this);
+		if (IsSpherical || FMath::Abs(QueryTM.GetRotation().W) > 0.9)
+		{
+			const FVec3 InvTriMeshScale = FVec3(FReal(1) / TriMeshScale.X, FReal(1) / TriMeshScale.Y, FReal(1) / TriMeshScale.Z);
+			// IMPORTANT QueryTM comes with a invscaled translation so we need a version of the TM with world space translation to properly compute the bounds
+			FRigidTransform3 TriMeshToGeomNoScale{ QueryTM };
+			TriMeshToGeomNoScale.SetTranslation(TriMeshToGeomNoScale.GetTranslation() * TriMeshScale);
+
+			// NOTE: BVH test is done in tri-mesh local space (whereas collision detection is done in world space because you can't non-uniformly scale all shapes)
+			FAABB3 QueryBounds = WorldScaleQueryGeom.CalculateTransformedBounds(TriMeshToGeomNoScale);
+			QueryBounds.ThickenSymmetrically(FVec3(Thickness));
+			QueryBounds.ScaleWithNegative(InvTriMeshScale);
+			TRigidTransform<FReal, 3> WorldScaleQueryTM;
+			ScaleTransformHelper(TriMeshScale, QueryTM, WorldScaleQueryTM);
+			return FastBVH.FindAllIntersectionsNoMTD(QueryBounds, WorldScaleQueryTM, WorldScaleQueryGeom, Thickness, TriMeshScale, this);
+		}
+		else
+		{
+			FAABB3 GeometryAABB = QueryGeom.BoundingBox();
+			GeometryAABB.ThickenSymmetrically(FVec3(Thickness));
+			FOBBVectorized QueryObb(QueryTM, (GeometryAABB.Max() - GeometryAABB.Min()) * 0.5);
+			TRigidTransform<FReal, 3> WorldScaleQueryTM;
+			ScaleTransformHelper(TriMeshScale, QueryTM, WorldScaleQueryTM);
+			return FastBVH.FindAllIntersectionsNoMTD(QueryObb, WorldScaleQueryTM, WorldScaleQueryGeom, Thickness, TriMeshScale, this);
+		}
 	}
 }
 
 bool FTriangleMeshImplicitObject::OverlapGeom(const TSphere<FReal, 3>& QueryGeom, const FRigidTransform3& QueryTM, const FReal Thickness, FMTDInfo* OutMTD) const
 {
-	return OverlapGeomImp(QueryGeom, QueryTM, Thickness, OutMTD);
+	return OverlapGeomImp<true>(QueryGeom, QueryTM, Thickness, OutMTD);
 }
 
 bool FTriangleMeshImplicitObject::OverlapGeom(const TBox<FReal, 3>& QueryGeom, const FRigidTransform3& QueryTM, const FReal Thickness, FMTDInfo* OutMTD) const
 {
-	return OverlapGeomImp(QueryGeom, QueryTM, Thickness, OutMTD);
+	return OverlapGeomImp<false>(QueryGeom, QueryTM, Thickness, OutMTD);
 }
 
 bool FTriangleMeshImplicitObject::OverlapGeom(const FCapsule& QueryGeom, const FRigidTransform3& QueryTM, const FReal Thickness, FMTDInfo* OutMTD) const
 {
-	return OverlapGeomImp(QueryGeom, QueryTM, Thickness, OutMTD);
+	return OverlapGeomImp<false>(QueryGeom, QueryTM, Thickness, OutMTD);
 }
 
 bool FTriangleMeshImplicitObject::OverlapGeom(const FConvex& QueryGeom, const FRigidTransform3& QueryTM, const FReal Thickness, FMTDInfo* OutMTD) const
 {
-	return OverlapGeomImp(QueryGeom, QueryTM, Thickness, OutMTD);
+	return OverlapGeomImp<false>(QueryGeom, QueryTM, Thickness, OutMTD);
 }
 
 bool FTriangleMeshImplicitObject::OverlapGeom(const TImplicitObjectScaled<TSphere<FReal, 3>>& QueryGeom, const FRigidTransform3& QueryTM, const FReal Thickness, FMTDInfo* OutMTD, FVec3 TriMeshScale) const
 {
-	return OverlapGeomImp(QueryGeom, QueryTM, Thickness, OutMTD, TriMeshScale);
+	return OverlapGeomImp<true>(QueryGeom, QueryTM, Thickness, OutMTD, TriMeshScale);
 }
 
 bool FTriangleMeshImplicitObject::OverlapGeom(const TImplicitObjectScaled<TBox<FReal, 3>>& QueryGeom, const FRigidTransform3& QueryTM, const FReal Thickness, FMTDInfo* OutMTD, FVec3 TriMeshScale) const
 {
-	return OverlapGeomImp(QueryGeom, QueryTM, Thickness, OutMTD, TriMeshScale);
+	return OverlapGeomImp<false>(QueryGeom, QueryTM, Thickness, OutMTD, TriMeshScale);
 }
 
 bool FTriangleMeshImplicitObject::OverlapGeom(const TImplicitObjectScaled<FCapsule>& QueryGeom, const FRigidTransform3& QueryTM, const FReal Thickness, FMTDInfo* OutMTD, FVec3 TriMeshScale) const
 {
-	return OverlapGeomImp(QueryGeom, QueryTM, Thickness, OutMTD, TriMeshScale);
+	return OverlapGeomImp<false>(QueryGeom, QueryTM, Thickness, OutMTD, TriMeshScale);
 }
 
 bool FTriangleMeshImplicitObject::OverlapGeom(const TImplicitObjectScaled<FConvex>& QueryGeom, const FRigidTransform3& QueryTM, const FReal Thickness, FMTDInfo* OutMTD, FVec3 TriMeshScale) const
 {
-	return OverlapGeomImp(QueryGeom, QueryTM, Thickness, OutMTD, TriMeshScale);
+	return OverlapGeomImp<false>(QueryGeom, QueryTM, Thickness, OutMTD, TriMeshScale);
 }
 
 template <typename QueryGeomType, typename IdxType>
diff --git a/Engine/Source/Runtime/Experimental/Chaos/Public/Chaos/OBBVectorized.h b/Engine/Source/Runtime/Experimental/Chaos/Public/Chaos/OBBVectorized.h
new file mode 100644
--- /dev/null
+++ b/Engine/Source/Runtime/Experimental/Chaos/Public/Chaos/OBBVectorized.h
@@ -0,0 +1,32 @@
+// Copyright Epic Games, Inc. All Rights Reserved.
+#pragma once
+
+#include "Chaos/Core.h"
+
+
+namespace Chaos
+{
+	class FAABBVectorized;
+
+	class FOBBVectorized
+	{
+	public:
+		FOBBVectorized(const FRigidTransform3& Transform, const FVec3f& HalfExtentsIn);
+		bool IntersectAABB(const FAABBVectorized& Bounds) const;
+
+	private:
+		VectorRegister4Float Position;
+		VectorRegister4Float XAxis;
+		VectorRegister4Float YAxis;
+		VectorRegister4Float ZAxis;
+		VectorRegister4Float HalfExtents;
+		VectorRegister4Float XHalfExtent;
+		VectorRegister4Float YHalfExtent;
+		VectorRegister4Float ZHalfExtent;
+
+		VectorRegister4Float MaxObb;
+		VectorRegister4Float MinObb;
+		
+
+	};
+}
\ No newline at end of file
diff --git a/Engine/Source/Runtime/Experimental/Chaos/Public/Chaos/TriangleMeshImplicitObject.h b/Engine/Source/Runtime/Experimental/Chaos/Public/Chaos/TriangleMeshImplicitObject.h
--- a/Engine/Source/Runtime/Experimental/Chaos/Public/Chaos/TriangleMeshImplicitObject.h
+++ b/Engine/Source/Runtime/Experimental/Chaos/Public/Chaos/TriangleMeshImplicitObject.h
@@ -6,6 +6,7 @@
 #include "Chaos/GeometryParticles.h"
 #include "Chaos/ImplicitFwd.h"
 #include "Chaos/ImplicitObjectScaled.h"
+#include "Chaos/OBBVectorized.h"
 #include "Chaos/SegmentMesh.h"
 #include "Chaos/Triangle.h"
 
@@ -182,6 +183,26 @@ namespace Chaos
 			VisitTree(BoundsFilter, FaceVisitor);
 		}
 
+		template <typename SQVisitor>
+		FORCEINLINE_DEBUGGABLE void OverlapOBB(const FOBBVectorized& Obb, SQVisitor& Visitor) const
+		{
+			const auto BoundsFilter = [&Obb](const FAABBVectorized& Bounds) -> EFilterResult
+			{
+				const bool bHit = Obb.IntersectAABB(Bounds);
+				return bHit ? EFilterResult::Keep : EFilterResult::Skip;
+			};
+
+			const auto FaceVisitor = [&Visitor](int32 FaceIndex)
+			{
+				const bool bContinueVisiting = Visitor.VisitOverlap(FaceIndex);
+				return bContinueVisiting ? EVisitorResult::Continue : EVisitorResult::Stop;
+			};
+
+			VisitTree(BoundsFilter, FaceVisitor);
+		}
+
+		template <typename QueryGeomType>
+		bool FindAllIntersectionsNoMTD(const FOBBVectorized& Intersection, const TRigidTransform<FReal, 3>& Transform, const QueryGeomType& QueryGeom, FReal Thickness, const FVec3& TriMeshScale, const FTriangleMeshImplicitObject* TriMesh) const;
 		template <typename QueryGeomType>
 		bool FindAllIntersectionsNoMTD(const FAABB3& Intersection, const TRigidTransform<FReal, 3>& Transform, const QueryGeomType& QueryGeom, FReal Thickness, const FVec3& TriMeshScale, const FTriangleMeshImplicitObject* TriMesh) const;
 		TArray<int32> FindAllIntersections(const FAABB3& Intersection) const;
@@ -866,7 +887,7 @@ namespace Chaos
 		template <typename QueryGeomType>
 		bool GJKContactPointImp(const QueryGeomType& QueryGeom, const FRigidTransform3& QueryTM, const FReal Thickness, FVec3& Location, FVec3& Normal, FReal& Penetration, int32& FaceIndex, FVec3 TriMeshScale = FVec3(1.0)) const;
 
-		template<typename QueryGeomType>
+		template<bool IsSpherical, typename QueryGeomType>
 		bool OverlapGeomImp(const QueryGeomType& QueryGeom, const FRigidTransform3& QueryTM, const FReal Thickness, FMTDInfo* OutMTD = nullptr, FVec3 TriMeshScale = FVec3(1.0f)) const;
 
 		template<typename QueryGeomType>
