Sphere Triangle Mesh and Sphere Height Field optimization

diff --git a/Engine/Source/Programs/HeadlessChaos/Private/HeadlessChaosTestHeightfield.cpp b/Engine/Source/Programs/HeadlessChaos/Private/HeadlessChaosTestHeightfield.cpp
--- a/Engine/Source/Programs/HeadlessChaos/Private/HeadlessChaosTestHeightfield.cpp
+++ b/Engine/Source/Programs/HeadlessChaos/Private/HeadlessChaosTestHeightfield.cpp
@@ -938,8 +938,6 @@ namespace ChaosTest {
 		const int32 Columns = 10;
 		const int32 Rows = 10;
 
-		
-
 		TArray<FReal> Heights = CreateMountain(Columns, Rows);
 
 		TArray<FReal> HeightsCopy = Heights;
@@ -950,6 +948,8 @@ namespace ChaosTest {
 		{
 			TBox<FReal, 3> Box(FVec3(-1.0, -1.0, -1.0), FVec3(1.0, 1.0, 10.0));
 			FCapsule Capsule(FVec3(0.0, 0.0, 0.0), FVec3(0.0, 0.0, 9.0), 1.14);
+			Chaos::FSphere Sphere1(FVec3(0.0, 0.0, -2.0), 0.6);
+
 			for (int32 Row = 0; Row < Rows; ++Row)
 			{
 				for (int32 Col = 0; Col < Columns; ++Col)
@@ -959,11 +959,14 @@ namespace ChaosTest {
 					FVec3 Dir(1, 0, 0);
 					bool BoxResult = Heightfield.OverlapGeom(Box, QueryTM, 0.0, nullptr);
 					bool CapsuleResult = Heightfield.OverlapGeom(Capsule, QueryTM, 0.0, nullptr);
+					bool SphereResult = Heightfield.OverlapGeom(Sphere1, QueryTM, 0.0, nullptr);
 					// No collision on the side of the mountain
 					if ((Col < 3 || Col > 7) && (Row < 3 || Row > 7))
 					{
 						EXPECT_FALSE(BoxResult);
 						EXPECT_FALSE(CapsuleResult);
+						// Sphere on the floor
+						EXPECT_TRUE(SphereResult);
 					}
 					// Collision with the mountain
 					else if ((Col >= 3 && Col <= 7) && (Row >= 3 && Row <= 7))
@@ -975,6 +978,11 @@ namespace ChaosTest {
 					{
 						EXPECT_FALSE(BoxResult);
 					}
+					// Inside the mountain the sphere shouldn't collide
+					if ((Col > 3 && Col < 7) && (Row > 3 && Row < 7))
+					{
+						EXPECT_FALSE(SphereResult);
+					}
 				}
 			}
 		}
diff --git a/Engine/Source/Programs/HeadlessChaos/Private/HeadlessChaosTestOverlap.cpp b/Engine/Source/Programs/HeadlessChaos/Private/HeadlessChaosTestOverlap.cpp
--- a/Engine/Source/Programs/HeadlessChaos/Private/HeadlessChaosTestOverlap.cpp
+++ b/Engine/Source/Programs/HeadlessChaos/Private/HeadlessChaosTestOverlap.cpp
@@ -3,9 +3,11 @@
 #include "HeadlessChaosTestOverlap.h"
 
 #include "HeadlessChaos.h"
-#include "Chaos/Capsule.h"
 #include "Chaos/Box.h"
+#include "Chaos/Capsule.h"
+#include "Chaos/Sphere.h"
 #include "Chaos/TriangleMeshImplicitObject.h"
+#include "Chaos/ImplicitObject.h"
 #include "Chaos/ImplicitObjectScaled.h"
 
 namespace ChaosTest
@@ -117,6 +119,89 @@ namespace ChaosTest
 					bool bResult = TriangleMesh->OverlapGeom(Capsule, QueryTM, 0.0);
 					EXPECT_FALSE(bResult);
 				}
+				{
+					const FRigidTransform3 QueryTM(FVec3(2.5, 2.5, 0.0), FQuat::Identity);
+					bool bResult = TriangleMesh->OverlapGeom(Capsule, QueryTM, 0.0);
+					EXPECT_TRUE(bResult);
+				}
+				{
+					const FRigidTransform3 QueryTM(FVec3(-2.5, -2.5, 0.0), FQuat::Identity);
+					bool bResult = TriangleMesh->OverlapGeom(Capsule, QueryTM, 0.0);
+					EXPECT_TRUE(bResult);
+				}
+			}
+			{
+				// Sphere test
+				const FVec3 X = { 0.0, 0.0, 0.0 };
+				const FReal Radius = 1.0;
+				const Chaos::FSphere Sphere = Chaos::FSphere(X, Radius);
+
+				{
+					const FRigidTransform3 QueryTM(FVec3(0.0, 0.0, 0.0), FQuat::Identity);
+					bool bResult = TriangleMesh->OverlapGeom(Sphere, QueryTM, 0.0);
+					EXPECT_EQ(bResult, true);
+				}
+				{
+					const FRigidTransform3 QueryTM(FVec3(12.5, 0.0, 0.0), FQuat::Identity);
+					bool bResult = TriangleMesh->OverlapGeom(Sphere, QueryTM, 0.0);
+					EXPECT_EQ(bResult, false);
+				}
+				{
+					const FRigidTransform3 QueryTM(FVec3(7.0, 7.0, 0.0), FQuat::Identity);
+					bool bResult = TriangleMesh->OverlapGeom(Sphere, QueryTM, 0.0);
+					EXPECT_EQ(bResult, false);
+				}
+				{
+					const FRigidTransform3 QueryTM(FVec3(-7.0, -7.0, 0.0), FQuat::Identity);
+					bool bResult = TriangleMesh->OverlapGeom(Sphere, QueryTM, 0.0);
+					EXPECT_EQ(bResult, false);
+				}
+				{
+					const FRigidTransform3 QueryTM(FVec3(10.0, 0.0, 0.0), FQuat::Identity);
+					bool bResult = TriangleMesh->OverlapGeom(Sphere, QueryTM, 0.0);
+					EXPECT_EQ(bResult, true);
+				}
+				{
+					const FRigidTransform3 QueryTM(FVec3(0.0, 10.2, 0.0), FQuat::Identity);
+					bool bResult = TriangleMesh->OverlapGeom(Sphere, QueryTM, 0.0);
+					EXPECT_EQ(bResult, true);
+				}
+				{
+					const FRigidTransform3 QueryTM(FVec3(0.0, -10.2, 0.0), FQuat::Identity);
+					bool bResult = TriangleMesh->OverlapGeom(Sphere, QueryTM, 0.0);
+					EXPECT_EQ(bResult, true);
+				}
+				{
+					const FRigidTransform3 QueryTM(FVec3(-10.3, 0.0, 0.0), FQuat::Identity);
+					bool bResult = TriangleMesh->OverlapGeom(Sphere, QueryTM, 0.0);
+					EXPECT_EQ(bResult, true);
+				}
+				{
+					const FRigidTransform3 QueryTM(FVec3(0.0, 0.0, 10.9), FQuat::Identity);
+					bool bResult = TriangleMesh->OverlapGeom(Sphere, QueryTM, 0.0);
+					EXPECT_EQ(bResult, true);
+				}
+				{
+					const FRigidTransform3 QueryTM(FVec3(0.0, 0.0, 11.5), FQuat(FVec3(1.0, 0.0, 0.0), 3.1415926 / 4.0));
+					bool bResult = TriangleMesh->OverlapGeom(Sphere, QueryTM, 0.0);
+					EXPECT_EQ(bResult, false);
+				}
+				{
+					// Inside Mesh
+					const FRigidTransform3 QueryTM(FVec3(0.0, 0.0, 4.0), FQuat::Identity);
+					bool bResult = TriangleMesh->OverlapGeom(Sphere, QueryTM, 0.0);
+					EXPECT_EQ(bResult, false);
+				}
+				{
+					const FRigidTransform3 QueryTM(FVec3(5.5, 0.0, 5.5), FQuat::Identity);
+					bool bResult = TriangleMesh->OverlapGeom(Sphere, QueryTM, 0.0);
+					EXPECT_TRUE(bResult);
+				}
+				{
+					const FRigidTransform3 QueryTM(FVec3(5.8, 0.0, 5.8), FQuat(FVec3(0.0, 1.0, 0.0), -3.1415926 / 4.0));
+					bool bResult = TriangleMesh->OverlapGeom(Sphere, QueryTM, 0.0);
+					EXPECT_FALSE(bResult);
+				}
 			}
 			{
 				const FVec3 X1 = { 0.0, 0.0, -200.0 };
@@ -366,6 +451,95 @@ namespace ChaosTest
 					EXPECT_FALSE(bResult);
 				}
 			}
+			// Non uniform scale sphere
+			{
+				const FVec3 X = { 0.0, 0.0, 0.0 };
+				const FReal Radius = 1.0;
+				TSharedPtr<Chaos::FSphere> SphereShared = MakeShared<Chaos::FSphere>(X, Radius);
+				FVec3 TriMeshScale = { 1.0, 1.0, 2.0 };
+				FVec3 InvScale = 1.0 / TriMeshScale;
+				TImplicitObjectScaled<Chaos::FSphere> ScaledSphere = TImplicitObjectScaled<Chaos::FSphere>(SphereShared, InvScale);
+				{
+					FRigidTransform3 QueryTM(FVec3(0.0, 0.0, 0.0), FQuat::Identity);
+					QueryTM = FRigidTransform3(QueryTM.GetLocation() * InvScale, QueryTM.GetRotation());
+					bool bResult = TriangleMesh->OverlapGeom(ScaledSphere, QueryTM, 0.0, nullptr, TriMeshScale);
+					EXPECT_EQ(bResult, true);
+				}
+				{
+					FRigidTransform3 QueryTM(FVec3(10.0, 10.0, 0.0), FQuat::Identity);
+					QueryTM = FRigidTransform3(QueryTM.GetLocation() * InvScale, QueryTM.GetRotation());
+					bool bResult = TriangleMesh->OverlapGeom(ScaledSphere, QueryTM, 0.0, nullptr, TriMeshScale);
+					EXPECT_EQ(bResult, false);
+				}
+				{
+					FRigidTransform3 QueryTM(FVec3(12.5, 0.0, 0.0), FQuat::Identity);
+					QueryTM = FRigidTransform3(QueryTM.GetLocation() * InvScale, QueryTM.GetRotation());
+					bool bResult = TriangleMesh->OverlapGeom(ScaledSphere, QueryTM, 0.0, nullptr, TriMeshScale);
+					EXPECT_EQ(bResult, false);
+				}
+				{
+					FRigidTransform3 QueryTM(FVec3(-10.0, -10.0, 0.0), FQuat::Identity);
+					QueryTM = FRigidTransform3(QueryTM.GetLocation() * InvScale, QueryTM.GetRotation());
+					bool bResult = TriangleMesh->OverlapGeom(ScaledSphere, QueryTM, 0.0, nullptr, TriMeshScale);
+					EXPECT_EQ(bResult, false);
+				}
+				{
+					FRigidTransform3 QueryTM(FVec3(10.0, 0.0, 0.0), FQuat::Identity);
+					QueryTM = FRigidTransform3(QueryTM.GetLocation() * InvScale, QueryTM.GetRotation());
+					bool bResult = TriangleMesh->OverlapGeom(ScaledSphere, QueryTM, 0.0, nullptr, TriMeshScale);
+					EXPECT_EQ(bResult, true);
+				}
+				{
+					FRigidTransform3 QueryTM(FVec3(0.0, 10.0, 0.0), FQuat::Identity);
+					QueryTM = FRigidTransform3(QueryTM.GetLocation() * InvScale, QueryTM.GetRotation());
+					bool bResult = TriangleMesh->OverlapGeom(ScaledSphere, QueryTM, 0.0, nullptr, TriMeshScale);
+					EXPECT_EQ(bResult, true);
+				}
+				{
+					FRigidTransform3 QueryTM(FVec3(0.0, -10.0, 0.0), FQuat::Identity);
+					QueryTM = FRigidTransform3(QueryTM.GetLocation() * InvScale, QueryTM.GetRotation());
+					bool bResult = TriangleMesh->OverlapGeom(ScaledSphere, QueryTM, 0.0, nullptr, TriMeshScale);
+					EXPECT_EQ(bResult, true);
+				}
+				{
+					FRigidTransform3 QueryTM(FVec3(-10.0, 0.0, 0.0), FQuat::Identity);
+					QueryTM = FRigidTransform3(QueryTM.GetLocation() * InvScale, QueryTM.GetRotation());
+					bool bResult = TriangleMesh->OverlapGeom(ScaledSphere, QueryTM, 0.0, nullptr, TriMeshScale);
+					EXPECT_EQ(bResult, true);
+				}
+				{
+					FRigidTransform3 QueryTM(FVec3(0.0, 0.0, 20.5), FQuat::Identity);
+					QueryTM = FRigidTransform3(QueryTM.GetLocation() * InvScale, QueryTM.GetRotation());
+					bool bResult = TriangleMesh->OverlapGeom(ScaledSphere, QueryTM, 0.0, nullptr, TriMeshScale);
+					EXPECT_EQ(bResult, true);
+				}
+				{
+					FRigidTransform3 QueryTM(FVec3(0.0, 0.0, 22.0), FQuat(FVec3(1.0, 0.0, 0.0), 3.1415926 / 4.0));
+					QueryTM = FRigidTransform3(QueryTM.GetLocation() * InvScale, QueryTM.GetRotation());
+					bool bResult = TriangleMesh->OverlapGeom(ScaledSphere, QueryTM, 0.0, nullptr, TriMeshScale);
+					EXPECT_EQ(bResult, false);
+				}
+				{
+					// Inside Mesh
+					FRigidTransform3 QueryTM(FVec3(0.0, 0.0, 8.0), FQuat::Identity);
+					QueryTM = FRigidTransform3(QueryTM.GetLocation() * InvScale, QueryTM.GetRotation());
+					bool bResult = TriangleMesh->OverlapGeom(ScaledSphere, QueryTM, 0.0, nullptr, TriMeshScale);
+					EXPECT_EQ(bResult, false);
+				}
+				{
+					FRigidTransform3 QueryTM(FVec3(5.5, 0.0, 10.5), FQuat::Identity);
+					QueryTM = FRigidTransform3(QueryTM.GetLocation() * InvScale, QueryTM.GetRotation());
+					bool bResult = TriangleMesh->OverlapGeom(ScaledSphere, QueryTM, 0.0, nullptr, TriMeshScale);
+					EXPECT_TRUE(bResult);
+				}
+				{
+					// Triangles parallel to the capsule axis with non uniform scale
+					FRigidTransform3 QueryTM(FVec3(6.0, 0.0, 11.5), FQuat(FVec3(0.0, 1.0, 0.0), -3.1415926 / 4.0 + 3.1415926 / 8.0));
+					QueryTM = FRigidTransform3(QueryTM.GetLocation() * InvScale, QueryTM.GetRotation());
+					bool bResult = TriangleMesh->OverlapGeom(ScaledSphere, QueryTM, 0.0, nullptr, TriMeshScale);
+					EXPECT_FALSE(bResult);
+				}
+			}
 
 			{
 				// Box test
diff --git a/Engine/Source/Runtime/Experimental/Chaos/Private/Chaos/Collision/TriangleOverlap.cpp b/Engine/Source/Runtime/Experimental/Chaos/Private/Chaos/Collision/TriangleOverlap.cpp
--- a/Engine/Source/Runtime/Experimental/Chaos/Private/Chaos/Collision/TriangleOverlap.cpp
+++ b/Engine/Source/Runtime/Experimental/Chaos/Private/Chaos/Collision/TriangleOverlap.cpp
@@ -118,4 +118,85 @@ namespace Chaos
 
 		return true;
 	}
+
+	bool ComputeSphereTriangleOverlapSimd(const VectorRegister4Float& A, const VectorRegister4Float& B, const VectorRegister4Float& C, const VectorRegister4Float& X, FRealSingle Radius)
+	{
+
+		const VectorRegister4Float AB = VectorSubtract(B, A);
+		const VectorRegister4Float BC = VectorSubtract(C, B);
+
+		VectorRegister4Float Normal = VectorNormalize(VectorCross(AB, BC));
+
+		// Plane Triangle
+		const VectorRegister4Float AX = VectorSubtract(X, A);
+
+		FRealSingle AXDist = VectorDot3Scalar(AX, Normal);
+
+		if (FMath::Abs<FRealSingle>(AXDist) > Radius)
+		{
+			return false;
+		}
+
+		constexpr FRealSingle ThirdFloat = 1.0f / 3.0f;
+		constexpr VectorRegister4Float Third = MakeVectorRegisterFloatConstant(ThirdFloat, ThirdFloat, ThirdFloat, ThirdFloat);
+		const VectorRegister4Float Centroid = VectorMultiply(VectorAdd(VectorAdd(A, B), C), Third);
+		
+
+		
+		
+
+		FRealSingle SqrRadius = Radius * Radius;
+		{
+			const VectorRegister4Float SqrAB = VectorDot3(AB, AB);
+			const VectorRegister4Float ZeroMask = VectorCompareEQ(VectorZeroFloat(), SqrAB);
+			VectorRegister4Float Time = VectorClamp(VectorDivide(VectorDot3(AB, AX), SqrAB), VectorZeroFloat(), VectorOneFloat());
+			Time = VectorBitwiseNotAnd(ZeroMask, Time);
+			const VectorRegister4Float P = VectorMultiplyAdd(B, VectorSubtract(VectorOneFloat(), Time), VectorMultiply(A, Time));
+
+			const VectorRegister4Float PX = VectorSubtract(X, P);
+			if (VectorDot3Scalar(PX, VectorSubtract(P, Centroid)) > 0.0f)
+			{
+				if (VectorDot3Scalar(PX, PX) > SqrRadius)
+				{
+					return false;
+				}
+			}
+		}
+		{
+			const VectorRegister4Float SqrBC = VectorDot3(BC, BC);
+			const VectorRegister4Float ZeroMask = VectorCompareEQ(VectorZeroFloat(), SqrBC);
+			const VectorRegister4Float BX = VectorSubtract(X, B);
+			VectorRegister4Float Time = VectorClamp(VectorDivide(VectorDot3(BC, BX), SqrBC), VectorZeroFloat(), VectorOneFloat());
+			Time = VectorBitwiseNotAnd(ZeroMask, Time);
+			const VectorRegister4Float P = VectorMultiplyAdd(C, VectorSubtract(VectorOneFloat(), Time), VectorMultiply(B, Time));
+
+			const VectorRegister4Float PX = VectorSubtract(X, P);
+			if (VectorDot3Scalar(PX, VectorSubtract(P, Centroid)) > 0.0f)
+			{
+				if (VectorDot3Scalar(PX, PX) > SqrRadius)
+				{
+					return false;
+				}
+			}
+		}
+		{
+			const VectorRegister4Float CA = VectorSubtract(A, C);
+			const VectorRegister4Float SqrCA = VectorDot3(CA, CA);
+			const VectorRegister4Float ZeroMask = VectorCompareEQ(VectorZeroFloat(), SqrCA);
+			const VectorRegister4Float CX = VectorSubtract(X, C);
+			VectorRegister4Float Time = VectorClamp(VectorDivide(VectorDot3(CA, CX), SqrCA), VectorZeroFloat(), VectorOneFloat());
+			Time = VectorBitwiseNotAnd(ZeroMask, Time);
+			const VectorRegister4Float P = VectorMultiplyAdd(A, VectorSubtract(VectorOneFloat(), Time), VectorMultiply(C, Time));
+
+			const VectorRegister4Float PX = VectorSubtract(X, P);
+			if (VectorDot3Scalar(PX, VectorSubtract(P, Centroid)) > 0.0f)
+			{
+				if (VectorDot3Scalar(PX, PX) > SqrRadius)
+				{
+					return false;
+				}
+			}
+		}
+		return true;
+	}
 }
diff --git a/Engine/Source/Runtime/Experimental/Chaos/Private/Chaos/Collision/TriangleOverlap.h b/Engine/Source/Runtime/Experimental/Chaos/Private/Chaos/Collision/TriangleOverlap.h
--- a/Engine/Source/Runtime/Experimental/Chaos/Private/Chaos/Collision/TriangleOverlap.h
+++ b/Engine/Source/Runtime/Experimental/Chaos/Private/Chaos/Collision/TriangleOverlap.h
@@ -6,5 +6,6 @@
 namespace Chaos
 {
 	bool ComputeCapsuleTriangleOverlapSimd(const VectorRegister4Float& A, const VectorRegister4Float& B, const VectorRegister4Float& C, const VectorRegister4Float& X1, const VectorRegister4Float& X2, FRealSingle Radius);
+	bool ComputeSphereTriangleOverlapSimd(const VectorRegister4Float& A, const VectorRegister4Float& B, const VectorRegister4Float& C, const VectorRegister4Float& X, FRealSingle Radius);
 }
 
diff --git a/Engine/Source/Runtime/Experimental/Chaos/Private/Chaos/HeightField.cpp b/Engine/Source/Runtime/Experimental/Chaos/Private/Chaos/HeightField.cpp
--- a/Engine/Source/Runtime/Experimental/Chaos/Private/Chaos/HeightField.cpp
+++ b/Engine/Source/Runtime/Experimental/Chaos/Private/Chaos/HeightField.cpp
@@ -2009,6 +2009,18 @@ namespace Chaos
 					X1 = VectorLoadFloat3(&X1f.X);
 					X2 = VectorLoadFloat3(&X2f.X);
 				}
+				if constexpr (std::is_same<QueryGeomType, Chaos::FSphere>::value)
+				{
+					Radius = FRealSingle(QueryGeom.GetRadius());
+					const FVec3f X1f = QueryGeom.GetCenter();
+					X1 = VectorLoadFloat3(&X1f.X);
+				}
+				else if constexpr (std::is_same<QueryGeomType, TImplicitObjectScaled < Chaos::FSphere>>::value)
+				{
+					Radius = FRealSingle(QueryGeom.GetRadius());
+					const FVec3f X1f = QueryGeom.GetUnscaledObject()->GetCenter() * QueryGeom.GetScale();
+					X1 = VectorLoadFloat3(&X1f.X);
+				}
 
 				VectorRegister4Float Points[4];
 				const UE::Math::TQuat<FReal>& RotationDouble = QueryTM.GetRotation();
@@ -2045,6 +2057,17 @@ namespace Chaos
 								return true;
 							}
 						}
+						else if constexpr (std::is_same<QueryGeomType, Chaos::FSphere>::value || std::is_same<QueryGeomType, TImplicitObjectScaled < Chaos::FSphere>>::value)
+						{
+							if (ComputeSphereTriangleOverlapSimd(Points[0], Points[1], Points[3], X1, Radius))
+							{
+								return true;
+							}
+							if (ComputeSphereTriangleOverlapSimd(Points[0], Points[3], Points[2], X1, Radius))
+							{
+								return true;
+							}
+						}
 						else
 						{
 							if (OverlapTriangleNoMTD(Points[0], Points[1], Points[3]))
diff --git a/Engine/Source/Runtime/Experimental/Chaos/Private/Chaos/TriangleMeshImplicitObject.cpp b/Engine/Source/Runtime/Experimental/Chaos/Private/Chaos/TriangleMeshImplicitObject.cpp
--- a/Engine/Source/Runtime/Experimental/Chaos/Private/Chaos/TriangleMeshImplicitObject.cpp
+++ b/Engine/Source/Runtime/Experimental/Chaos/Private/Chaos/TriangleMeshImplicitObject.cpp
@@ -551,6 +551,101 @@ struct FTriangleMeshOverlapVisitorNoMTD<TImplicitObjectScaled<FCapsule>>
 	bool bFoundIntersection;
 };
 
+template <>
+struct FTriangleMeshOverlapVisitorNoMTD<Chaos::FSphere>
+{
+	FTriangleMeshOverlapVisitorNoMTD(const TRigidTransform<FReal, 3>& WorldScaleQueryTM, const Chaos::FSphere& InQueryGeom, FReal InThickness, const FVec3f& InTriMeshScale, const FTriangleMeshImplicitObject* InTriMesh)
+		: Radius(FRealSingle(InQueryGeom.GetRadius() + InThickness))
+		, TriMesh(InTriMesh)
+		, bFoundIntersection(false)
+	{
+		const FVec3f Xf = WorldScaleQueryTM.TransformPosition(InQueryGeom.GetCenter());
+		X = VectorLoadFloat3(&Xf.X);
+		TriMeshScale = VectorLoadFloat3(&InTriMeshScale.X);
+	}
+
+	bool VisitOverlap(int32 TriIdx)
+	{
+		VectorRegister4Float A, B, C;
+		if (TriMesh->MElements.RequiresLargeIndices())
+		{
+			TriangleMeshTransformVertsHelperSimd(TriMeshScale, TriIdx, TriMesh->MParticles, TriMesh->MElements.GetLargeIndexBuffer(), A, B, C);
+		}
+		else
+		{
+			TriangleMeshTransformVertsHelperSimd(TriMeshScale, TriIdx, TriMesh->MParticles, TriMesh->MElements.GetSmallIndexBuffer(), A, B, C);
+		}
+
+		bFoundIntersection = ComputeSphereTriangleOverlapSimd(A, B, C, X, Radius);
+		return !bFoundIntersection;
+	}
+
+	bool VisitSweep(int32 Instance, FRealSingle& CurDataLength)
+	{
+		check(false);
+		return true;
+	}
+	bool VisitRaycast(int32 Instance, FRealSingle& CurDataLength)
+	{
+		check(false);
+		return true;
+	}
+
+	FRealSingle Radius;
+	VectorRegister4Float X;
+	VectorRegister4Float TriMeshScale;
+	const FTriangleMeshImplicitObject* TriMesh;
+	bool bFoundIntersection;
+};
+
+
+template <>
+struct FTriangleMeshOverlapVisitorNoMTD<TImplicitObjectScaled<Chaos::FSphere>>
+{
+	FTriangleMeshOverlapVisitorNoMTD(const TRigidTransform<FReal, 3>& WorldScaleQueryTM, const TImplicitObjectScaled< Chaos::FSphere >& InQueryGeom, FReal InThickness, const FVec3f& InTriMeshScale, const FTriangleMeshImplicitObject* InTriMesh)
+		: Radius(FRealSingle(InQueryGeom.GetRadius() + InThickness))
+		, TriMesh(InTriMesh)
+		, bFoundIntersection(false)
+	{
+		const FVec3f Xf = WorldScaleQueryTM.TransformPosition(InQueryGeom.GetUnscaledObject()->GetCenter() * InQueryGeom.GetScale());
+		X = VectorLoadFloat3(&Xf.X);
+		TriMeshScale = VectorLoadFloat3(&InTriMeshScale.X);
+	}
+
+	bool VisitOverlap(int32 TriIdx)
+	{
+		VectorRegister4Float A, B, C;
+		if (TriMesh->MElements.RequiresLargeIndices())
+		{
+			TriangleMeshTransformVertsHelperSimd(TriMeshScale, TriIdx, TriMesh->MParticles, TriMesh->MElements.GetLargeIndexBuffer(), A, B, C);
+		}
+		else
+		{
+			TriangleMeshTransformVertsHelperSimd(TriMeshScale, TriIdx, TriMesh->MParticles, TriMesh->MElements.GetSmallIndexBuffer(), A, B, C);
+		}
+
+		bFoundIntersection = ComputeSphereTriangleOverlapSimd(A, B, C, X, Radius);
+		return !bFoundIntersection;
+	}
+
+	bool VisitSweep(int32 Instance, FRealSingle& CurDataLength)
+	{
+		check(false);
+		return true;
+	}
+	bool VisitRaycast(int32 Instance, FRealSingle& CurDataLength)
+	{
+		check(false);
+		return true;
+	}
+
+	FRealSingle Radius;
+	VectorRegister4Float X;
+	VectorRegister4Float TriMeshScale;
+	const FTriangleMeshImplicitObject* TriMesh;
+	bool bFoundIntersection;
+};
+
 template <typename QueryGeomType>
 bool FTrimeshBVH::FindAllIntersectionsNoMTD(const FAABB3& Intersection, const TRigidTransform<FReal, 3>& Transform, const QueryGeomType& QueryGeom, FReal Thickness, const FVec3& TriMeshScale, const FTriangleMeshImplicitObject* TriMesh) const
 {
