Chaos - fix HeadlessChaos test deprecation warning

diff --git a/Engine/Source/Programs/HeadlessChaos/Private/HeadlessChaosTestClustering.cpp b/Engine/Source/Programs/HeadlessChaos/Private/HeadlessChaosTestClustering.cpp
--- a/Engine/Source/Programs/HeadlessChaos/Private/HeadlessChaosTestClustering.cpp
+++ b/Engine/Source/Programs/HeadlessChaos/Private/HeadlessChaosTestClustering.cpp
@@ -146,7 +146,7 @@ namespace ChaosTest {
 				EXPECT_NE(Particle.Handle(), ClusteredParticles.Handle(BoxID));	//make sure boxes are not in non disabled array
 			}
 
-			EXPECT_TRUE(Evolution.GetConstraintGraph().DebugCheckParticleNotInGraph(ClusteredParticles.Handle(BoxID)));
+			EXPECT_TRUE(Evolution.GetIslandManager().DebugCheckParticleNotInGraph(ClusteredParticles.Handle(BoxID)));
 		}
 
 		for (Chaos::FPBDRigidParticleHandle* ClusterHandle : ClusterHandles)
@@ -172,7 +172,7 @@ namespace ChaosTest {
 				EXPECT_NE(Particle.Handle(), ClusterHandle);	//make sure boxes are not in non disabled array
 			}
 
-			EXPECT_TRUE(Evolution.GetConstraintGraph().DebugCheckParticleNotInGraph(ClusterHandle));
+			EXPECT_TRUE(Evolution.GetIslandManager().DebugCheckParticleNotInGraph(ClusterHandle));
 		}
 
 		EXPECT_EQ(Particles.GetNonDisabledView().Num(), NumBoxes);
@@ -270,7 +270,7 @@ namespace ChaosTest {
 		{
 			EXPECT_TRUE(ClusteredParticles.Disabled(BoxID));	//no boxes should be active yet
 			EXPECT_TRUE(Evolution.GetRigidClustering().GetTopLevelClusterParents().Contains(ClusteredParticles.Handle(BoxID)) == false);
-			EXPECT_TRUE(Evolution.GetConstraintGraph().DebugCheckParticleNotInGraph(ClusteredParticles.Handle(BoxID)));
+			EXPECT_TRUE(Evolution.GetIslandManager().DebugCheckParticleNotInGraph(ClusteredParticles.Handle(BoxID)));
 		}
 
 		SolverStrainArray[NumBoxes + NumBoxes / 4 + 1] = (FReal)1;
@@ -287,7 +287,7 @@ namespace ChaosTest {
 		{
 			EXPECT_TRUE(ClusteredParticles.Disabled(BoxID));	//no boxes should be active yet
 			EXPECT_TRUE(Evolution.GetRigidClustering().GetTopLevelClusterParents().Contains(ClusteredParticles.Handle(BoxID)) == false);
-			EXPECT_TRUE(Evolution.GetConstraintGraph().DebugCheckParticleNotInGraph(ClusteredParticles.Handle(BoxID)));
+			EXPECT_TRUE(Evolution.GetIslandManager().DebugCheckParticleNotInGraph(ClusteredParticles.Handle(BoxID)));
 		}
 		
 	}
