Chaos - fix Heightfield CCD sweep for objects not centered on their bounds

diff --git a/Engine/Source/Programs/HeadlessChaos/Private/HeadlessChaosTestHeightfield.cpp b/Engine/Source/Programs/HeadlessChaos/Private/HeadlessChaosTestHeightfield.cpp
--- a/Engine/Source/Programs/HeadlessChaos/Private/HeadlessChaosTestHeightfield.cpp
+++ b/Engine/Source/Programs/HeadlessChaos/Private/HeadlessChaosTestHeightfield.cpp
@@ -1103,4 +1103,85 @@ namespace ChaosTest {
 		SUCCEED();
 	}
 
+	// A flat heightfield and a sphere swept downwards using the CCD API.
+	// Check that the IgnoreThreshold is being used to ignore sweeps that have a penetration depth
+	// less that that threshold at T=1.
+	GTEST_TEST(HeightfieldCCDTests, TestSphere)
+	{
+		int32 Rows = 64;
+		int32 Columns = 64;
+		FVec3 Scale(100.0, 100.0, 100.0);
+		TArray<FReal> Heights;
+		Heights.AddZeroed(Rows * Columns);
+
+		TArray<FReal> HeightsCopy = Heights;
+		FHeightField Heightfield(MoveTemp(HeightsCopy), TArray<uint8>(), Rows, Columns, Scale);
+		const auto& Bounds = Heightfield.BoundingBox();	//Current API forces us to do this to cache the bounds
+
+		TSphere<FReal, 3> Sphere(FVec3(0.0, 0.0, 0.0), 50.0);
+
+		FReal TOI, Phi;
+		FVec3 Position, Normal, FaceNormal;
+		int32 FaceIdx = 0;
+
+		const FRigidTransform3 Start(FVec3(0, 0, 60), FRotation3::FromIdentity());
+		const FVec3 Dir(0, 0, -1);
+		const FReal CCDIgnorePenetration = 30;
+		const FReal CCDTargetPenetration = 0;
+
+		// The first sweep should be less than the ignore threshold so get no hit
+		const FReal CCDSweepLengthA = 30;
+		const bool bHitA = Heightfield.SweepGeomCCD(Sphere, Start, Dir, CCDSweepLengthA, CCDIgnorePenetration, CCDTargetPenetration, TOI, Phi, Position, Normal, FaceIdx, FaceNormal);
+		EXPECT_FALSE(bHitA);
+
+		// The second sweep should be greater than the ignore threshold so we get a hit
+		const FReal CCDSweepLengthB = 50;
+		const bool bHitB = Heightfield.SweepGeomCCD(Sphere, Start, Dir, CCDSweepLengthB, CCDIgnorePenetration, CCDTargetPenetration, TOI, Phi, Position, Normal, FaceIdx, FaceNormal);
+		EXPECT_TRUE(bHitB);
+		EXPECT_NEAR(TOI, 0.2, UE_KINDA_SMALL_NUMBER);	// 10 / 50
+		EXPECT_NEAR(Phi, 0.0, UE_KINDA_SMALL_NUMBER);	// We are not initially penetrating
+	}
+
+	// Same as TestSphere except the sphere has a local offset built in.
+	// CCD sweeps have an early out if the depth at T=1 is less than some threshold, and there was a bug in this.
+	// related to the fact that the sweep is effectively an AABB sweep to find the overlapping heightfield cells, 
+	// and a Shape sweep to find the shape-triangle overlaps. The bug was that the AABB sweep position was being 
+	// used in the early-rejection code that should be using the shape position.
+	GTEST_TEST(HeightfieldCCDTests, TestSphereWithOffset)
+	{
+		int32 Rows = 64;
+		int32 Columns = 64;
+		FVec3 Scale(100.0, 100.0, 100.0);
+		TArray<FReal> Heights;
+		Heights.AddZeroed(Rows * Columns);
+
+		TArray<FReal> HeightsCopy = Heights;
+		FHeightField Heightfield(MoveTemp(HeightsCopy), TArray<uint8>(), Rows, Columns, Scale);
+		const auto& Bounds = Heightfield.BoundingBox();	//Current API forces us to do this to cache the bounds
+
+		// NOTE: Sphere center 50cm up
+		TSphere<FReal, 3> Sphere(FVec3(0.0, 0.0, 50.0), 50.0);
+
+		FReal TOI, Phi;
+		FVec3 Position, Normal, FaceNormal;
+		int32 FaceIdx = 0;
+
+		// NOTE: Start height at 10, rather than 60 in TestSphere
+		const FRigidTransform3 Start(FVec3(0, 0, 10), FRotation3::FromIdentity());
+		const FVec3 Dir(0, 0, -1);
+		const FReal CCDIgnorePenetration = 30;
+		const FReal CCDTargetPenetration = 0;
+		
+		// The first sweep should be less than the ignore threshold so get no hit
+		const FReal CCDSweepLengthA = 30;
+		const bool bHitA = Heightfield.SweepGeomCCD(Sphere, Start, Dir, CCDSweepLengthA, CCDIgnorePenetration, CCDTargetPenetration, TOI, Phi, Position, Normal, FaceIdx, FaceNormal);
+		EXPECT_FALSE(bHitA);
+
+		// The second sweep should be greater than the ignore threshold so we get a hit
+		const FReal CCDSweepLengthB = 50;
+		const bool bHitB = Heightfield.SweepGeomCCD(Sphere, Start, Dir, CCDSweepLengthB, CCDIgnorePenetration, CCDTargetPenetration, TOI, Phi, Position, Normal, FaceIdx, FaceNormal);
+		EXPECT_TRUE(bHitB);
+		EXPECT_NEAR(TOI, 0.2, UE_KINDA_SMALL_NUMBER);	// 10 / 50
+		EXPECT_NEAR(Phi, 0.0, UE_KINDA_SMALL_NUMBER);	// We are not initially penetrating
+	}
 }
\ No newline at end of file
diff --git a/Engine/Source/Runtime/Experimental/Chaos/Private/Chaos/HeightField.cpp b/Engine/Source/Runtime/Experimental/Chaos/Private/Chaos/HeightField.cpp
--- a/Engine/Source/Runtime/Experimental/Chaos/Private/Chaos/HeightField.cpp
+++ b/Engine/Source/Runtime/Experimental/Chaos/Private/Chaos/HeightField.cpp
@@ -460,10 +460,12 @@ namespace Chaos
 			, TargetPenetration(InTargetPenetration)
 		{
 			const FAABB3 QueryBounds = InQueryGeom.BoundingBox();
-			const FVec3 StartPoint = StartTM.TransformPositionNoScale(QueryBounds.Center());
+			const FVec3 BoundsStartPoint = StartTM.TransformPositionNoScale(QueryBounds.Center());
+			const FVec3 StartPoint = StartTM.GetTranslation();
 			const FVec3 Inflation3D = QueryBounds.Extents() * FReal(0.5);
 
 			const VectorRegister4Float LengthSimd = MakeVectorRegisterFloatFromDouble(VectorSetFloat1(Length));
+			BoundsStartPointSimd = MakeVectorRegisterFloatFromDouble(MakeVectorRegister(BoundsStartPoint.X, BoundsStartPoint.Y, BoundsStartPoint.Z, 0.0));
 			StartPointSimd = MakeVectorRegisterFloatFromDouble(MakeVectorRegister(StartPoint.X, StartPoint.Y, StartPoint.Z, 0.0));
 			DirSimd = MakeVectorRegisterFloatFromDouble(MakeVectorRegister(InDir.X, InDir.Y, InDir.Z, 0.0));
 			EndPointSimd = VectorAdd(StartPointSimd, VectorMultiply(DirSimd, LengthSimd));
@@ -549,7 +551,7 @@ namespace Chaos
 			HfData->GetPointsAndBoundsScaledSimd(VertexIndex, Points, CellBounds);
 			CellBounds.Thicken(Inflation3DSimd);
 
-			if (CellBounds.RaycastFast(StartPointSimd, InvDirSimd, ParallelSimd, VectorSetFloat1(static_cast<FRealSingle>(CurrentLength))))
+			if (CellBounds.RaycastFast(BoundsStartPointSimd, InvDirSimd, ParallelSimd, VectorSetFloat1(static_cast<FRealSingle>(CurrentLength))))
 			{
 				const int32 TriIndex0 = 2 * CellIndex + 0;
 				const int32 TriIndex1 = 2 * CellIndex + 1;
@@ -579,6 +581,7 @@ namespace Chaos
 		FReal IgnorePenetration;
 		FReal TargetPenetration;
 
+		VectorRegister4Float BoundsStartPointSimd;
 		VectorRegister4Float StartPointSimd;
 		VectorRegister4Float EndPointSimd;
 		VectorRegister4Float InvDirSimd;
