Add support for Tetrahedra in TBoundingVolumeHierarchy.

diff --git a/Engine/Source/Programs/HeadlessChaos/Private/HeadlessChaosTestImplicits.cpp b/Engine/Source/Programs/HeadlessChaos/Private/HeadlessChaosTestImplicits.cpp
--- a/Engine/Source/Programs/HeadlessChaos/Private/HeadlessChaosTestImplicits.cpp
+++ b/Engine/Source/Programs/HeadlessChaos/Private/HeadlessChaosTestImplicits.cpp
@@ -7,6 +7,7 @@
 #include "Modules/ModuleManager.h"
 #include "Chaos/PBDRigidsEvolution.h"
 #include "Chaos/PBDRigidParticles.h"
+#include "Chaos/BoundingVolumeHierarchy.h"
 #include "Chaos/Box.h"
 #include "Chaos/Plane.h"
 #include "Chaos/Sphere.h"
@@ -315,6 +316,29 @@ namespace ChaosTest {
 		Surf = Tet.ProjectToSurface(Tris, Pt);
 		Pt -= Tris[0].GetNormal() * 0.1;
 		EXPECT_VECTOR_NEAR(Surf, Pt, 0.001);
+
+		// Bounding Volume Hierarchy
+		// Put the center of the tet at the origin, and sweep point tests across it.
+		Pt = Tet.GetCenter();
+		for (int i = 0; i < 4; i++)
+		{
+			Tet[i] -= Pt;
+		}
+		TArray<TTetrahedron<double>*> Tetrahedra;
+		Tetrahedra.Add(&Tet);
+		TBoundingVolumeHierarchy<TArray<TTetrahedron<Chaos::FReal>*>, TArray<int32>, Chaos::FReal, 3> BVH(Tetrahedra);
+		for (int32 i = -5; i < 5; i++)
+		{
+			TArray<int32> Intersections = BVH.FindAllIntersections(Chaos::TVec3<double>(i, i, i));
+			if (i == 0)
+			{
+				EXPECT_TRUE(Intersections.Num() == 1);
+			}
+			else
+			{
+				EXPECT_TRUE(Intersections.Num() == 0);
+			}
+		}
 	}
 
 	void ImplicitCube()
diff --git a/Engine/Source/Runtime/Experimental/Chaos/Private/Chaos/BoundingVolumeHierarchy.cpp b/Engine/Source/Runtime/Experimental/Chaos/Private/Chaos/BoundingVolumeHierarchy.cpp
--- a/Engine/Source/Runtime/Experimental/Chaos/Private/Chaos/BoundingVolumeHierarchy.cpp
+++ b/Engine/Source/Runtime/Experimental/Chaos/Private/Chaos/BoundingVolumeHierarchy.cpp
@@ -3,6 +3,7 @@
 #include "Chaos/BoundingVolumeHierarchy.h"
 #include "Chaos/BoundingVolume.h"
 #include "Chaos/BoundingVolumeUtilities.h"
+#include "Chaos/Tetrahedron.h"
 #include "ChaosStats.h"
 
 namespace Chaos
@@ -684,6 +685,7 @@ namespace Chaos
 
 //template class CHAOS_API Chaos::TBoundingVolume<int32, FReal, 3>;
 template class Chaos::TBoundingVolumeHierarchy<TArray<Chaos::TSphere<Chaos::FReal, 3>*>, TArray<int32>, Chaos::FReal, 3>;
+template class Chaos::TBoundingVolumeHierarchy<TArray<Chaos::TTetrahedron<Chaos::FReal>*>, TArray<int32>, Chaos::FReal, 3>;
 template class Chaos::TBoundingVolumeHierarchy<Chaos::TPBDRigidParticles<Chaos::FReal, 3>, TArray<int32>, Chaos::FReal, 3>;
 template class Chaos::TBoundingVolumeHierarchy<Chaos::FParticles, TArray<int32>, Chaos::FReal, 3>;
 template class Chaos::TBoundingVolumeHierarchy<Chaos::TGeometryParticles<Chaos::FReal, 3>, TArray<int32>, Chaos::FReal, 3>;
diff --git a/Engine/Source/Runtime/Experimental/Chaos/Public/Chaos/Tetrahedron.h b/Engine/Source/Runtime/Experimental/Chaos/Public/Chaos/Tetrahedron.h
--- a/Engine/Source/Runtime/Experimental/Chaos/Public/Chaos/Tetrahedron.h
+++ b/Engine/Source/Runtime/Experimental/Chaos/Public/Chaos/Tetrahedron.h
@@ -2,6 +2,7 @@
 
 #pragma once
 
+#include "Chaos/AABB.h"
 #include "Chaos/Core.h"
 #include "Chaos/Triangle.h"
 #include "Chaos/Math/Poisson.h" // 3x3 row major matrix functions
@@ -45,6 +46,18 @@ namespace Chaos {
 			return static_cast<T>(.25) * (X[0] + X[1] + X[2] + X[3]);
 		}
 
+		bool HasBoundingBox() const { return true; }
+		TAABB<T, 3> BoundingBox() const { return GetBoundingBox(); }
+		TAABB<T, 3> GetBoundingBox() const
+		{
+			TAABB<T, 3> Box = TAABB<T, 3>::EmptyAABB();
+			for (int32 i = 0; i < 4; i++)
+			{
+				Box.GrowToInclude(X[i]);
+			}
+			return Box;
+		}
+
 		T GetMinEdgeLengthSquared() const
 		{
 			return FMath::Min(
