Chaos - Add character ground constraint solver

diff --git a/Engine/Source/Programs/HeadlessChaos/HeadlessChaos.Build.cs b/Engine/Source/Programs/HeadlessChaos/HeadlessChaos.Build.cs
--- a/Engine/Source/Programs/HeadlessChaos/HeadlessChaos.Build.cs
+++ b/Engine/Source/Programs/HeadlessChaos/HeadlessChaos.Build.cs
@@ -12,6 +12,9 @@ public class HeadlessChaos : TestModuleRules
 		// For LaunchEngineLoop.cpp include
 		PrivateIncludePaths.Add("Runtime/Launch/Private");
 
+		// For testing access to private Chaos classes
+		PrivateIncludePaths.Add("Runtime/Experimental/Chaos/Private");
+
 		SetupModulePhysicsSupport(Target);
 
 		PrivateDependencyModuleNames.AddRange(
diff --git a/Engine/Source/Programs/HeadlessChaos/Private/Character/TestCharacterGroundConstraintSolver.cpp b/Engine/Source/Programs/HeadlessChaos/Private/Character/TestCharacterGroundConstraintSolver.cpp
new file mode 100644
--- /dev/null
+++ b/Engine/Source/Programs/HeadlessChaos/Private/Character/TestCharacterGroundConstraintSolver.cpp
@@ -0,0 +1,465 @@
+// Copyright Epic Games, Inc. All Rights Reserved.
+
+#include "HeadlessChaos.h"
+#include "HeadlessChaosTestUtility.h"
+
+#include "Chaos/Character/CharacterGroundConstraintSolver.h"
+
+namespace ChaosTest
+{
+	using namespace Chaos;
+
+	class CharacterGroundConstraintSolverTest : public ::testing::Test
+	{
+	protected:
+		void SetUp() override
+		{
+			CharacterBody = FSolverBody::MakeInitialized();
+			CharacterBody.SetInvM(0.01);
+			CharacterBody.SetInvILocal(FVec3(0.005, 0.005, 0.005));
+
+			GroundBody = FSolverBody::MakeInitialized();
+
+			Dt = FReal(1.0 / 30.0);
+		}
+
+		void SolvePosition(int NumIterations)
+		{
+			for (int It = 0; It < NumIterations; ++It)
+			{
+				Solver.SolvePosition();
+			}
+		}
+
+		void UpdateSingleBody(int NumPositionIterations, int NumVelocityIterations)
+		{
+			Solver.SetBodies(&CharacterBody);
+			Solver.GatherInput(Dt, Settings, Data);
+			SolvePosition(NumPositionIterations);
+			CharacterBody.ApplyCorrections();
+			Solver.Reset();
+		}
+
+		void UpdateTwoBody(int NumPositionIterations, int NumVelocityIterations)
+		{
+			Solver.SetBodies(&CharacterBody, &GroundBody);
+			Solver.GatherInput(Dt, Settings, Data);
+			SolvePosition(NumPositionIterations);
+			CharacterBody.ApplyCorrections();
+			GroundBody.ApplyCorrections();
+			Solver.Reset();
+		}
+
+		Private::FCharacterGroundConstraintSolver Solver;
+		FCharacterGroundConstraintSettings Settings;
+		FCharacterGroundConstraintDynamicData Data;
+		FSolverBody CharacterBody;
+		FSolverBody GroundBody;
+		FReal Dt;
+	};
+
+	TEST_F(CharacterGroundConstraintSolverTest, TestInitialization)
+	{
+		// Should be able to access the impulses before initialization and get zero
+		EXPECT_VECTOR_FLOAT_EQ(Solver.GetLinearImpulse(1.0f), FVec3::ZeroVector);
+		EXPECT_VECTOR_FLOAT_EQ(Solver.GetAngularImpulse(1.0f), FVec3::ZeroVector);
+	}
+
+	// Initial overlap at zero velocity
+	// Should be corrected in a single iteration without introducing velocity
+	TEST_F(CharacterGroundConstraintSolverTest, TestSingleBody_NormalImpulse_InitialOverlap)
+	{
+		CharacterBody.SetX(FVec3(10.0, 10.0, 20.0));
+		CharacterBody.SetP(FVec3(10.0, 10.0, 20.0));
+		Data.GroundDistance = 10.0f;
+		Settings.TargetHeight = 15.0f;
+		Settings.RadialForceLimit = 0.0f;
+		Settings.TorqueLimit = 0.0f;
+
+		UpdateSingleBody(1, 0);
+
+		EXPECT_VECTOR_FLOAT_EQ(CharacterBody.P(), FVec3(10.0, 10.0, 25.0));
+		EXPECT_VECTOR_FLOAT_EQ(CharacterBody.V(), FVec3(0.0, 0.0, 0.0));
+
+		FVec3 ExpectedImpulse = (1.0 / (CharacterBody.InvM() * Dt)) * FVec3(0.0, 0.0, 5.0);
+		EXPECT_VECTOR_FLOAT_EQ(Solver.GetLinearImpulse(Dt), ExpectedImpulse);
+		EXPECT_VECTOR_FLOAT_EQ(Solver.GetAngularImpulse(Dt), FVec3::ZeroVector);
+	}
+
+	// Initial overlap but enough velocity to move out of overlap
+	// Solver should do nothing
+	TEST_F(CharacterGroundConstraintSolverTest, TestSingleBody_NormalImpulse_InitialOverlap_NoFinalOverlap)
+	{
+		CharacterBody.SetX(FVec3(10.0, 10.0, 20.0));
+		CharacterBody.SetP(FVec3(10.0, 10.0, 25.01));
+		CharacterBody.SetV(FVec3(0.0, -100.0, 5.01 / Dt));
+		Data.GroundDistance = 10.0f;
+		Settings.TargetHeight = 15.0f;
+		Settings.RadialForceLimit = 0.0f;
+		Settings.TorqueLimit = 0.0f;
+
+		UpdateSingleBody(1, 0);
+
+		EXPECT_VECTOR_FLOAT_EQ(CharacterBody.P(), FVec3(10.0, 10.0, 25.01));
+		EXPECT_VECTOR_FLOAT_EQ(CharacterBody.V(), FVec3(0.0, -100.0, 5.01 / Dt));
+
+		EXPECT_VECTOR_FLOAT_EQ(Solver.GetLinearImpulse(Dt), FVec3::ZeroVector);
+		EXPECT_VECTOR_FLOAT_EQ(Solver.GetAngularImpulse(Dt), FVec3::ZeroVector);
+	}
+
+	// No overlap but projected overlap
+	// Solver should fix overlap using displacement not correction
+	TEST_F(CharacterGroundConstraintSolverTest, TestSingleBody_NormalImpulse_NoInitialOverlap_FinalOverlap)
+	{
+		CharacterBody.SetX(FVec3(10.0, 10.0, 26.0));
+		CharacterBody.SetP(FVec3(10.0, 10.0, 20.0));
+		CharacterBody.SetV(FVec3(0.0, 500.0, -6.0 / Dt));
+		Data.GroundDistance = 16.0f;
+		Settings.TargetHeight = 15.0f;
+		Settings.RadialForceLimit = 0.0f;
+		Settings.TorqueLimit = 0.0f;
+
+		Solver.SetBodies(&CharacterBody);
+		Solver.GatherInput(Dt, Settings, Data);
+		SolvePosition(1);
+
+		EXPECT_TRUE(CharacterBody.CP() == FSolverVec3::ZeroVector);
+		EXPECT_FALSE(CharacterBody.DP() == FSolverVec3::ZeroVector);
+
+		EXPECT_VECTOR_FLOAT_EQ(CharacterBody.CorrectedP(), FVec3(10.0, 10.0, 25.0));
+
+		FVec3 ExpectedImpulse = (1.0 / (CharacterBody.InvM() * Dt)) * FVec3(0.0, 0.0, 5.0);
+		EXPECT_VECTOR_FLOAT_EQ(Solver.GetLinearImpulse(Dt), ExpectedImpulse);
+		EXPECT_VECTOR_FLOAT_EQ(Solver.GetAngularImpulse(Dt), FVec3::ZeroVector);
+	}
+
+	// No overlap or projected overlap but overlap due to other constraints
+	// Solver should fix overlap using displacement not correction
+	TEST_F(CharacterGroundConstraintSolverTest, TestSingleBody_NormalImpulse_OtherConstraintOverlap)
+	{
+		CharacterBody.SetX(FVec3(10.0, 10.0, 26.0));
+		CharacterBody.SetP(FVec3(10.0, 10.0, 26.0));
+		Data.GroundDistance = 16.0f;
+		Settings.TargetHeight = 15.0f;
+		Settings.RadialForceLimit = 0.0f;
+		Settings.TorqueLimit = 0.0f;
+
+		Solver.SetBodies(&CharacterBody);
+		Solver.GatherInput(Dt, Settings, Data);
+		SolvePosition(1);
+
+		// Solver should do nothing first iteration as there is no projected overlap
+		EXPECT_TRUE(CharacterBody.CP() == FSolverVec3::ZeroVector);
+		EXPECT_TRUE(CharacterBody.DP() == FSolverVec3::ZeroVector);
+
+		// Add projected overlap due to another constraint
+		CharacterBody.ApplyPositionDelta(FSolverVec3(1.0f, -2.0f, -4.0f));
+
+		Solver.SolvePosition();
+
+		FSolverVec3 Impulse = Solver.GetLinearImpulse(1);
+		FSolverVec3 ExpectedSolverImpulse = (1.0 / (CharacterBody.InvM())) * 3.0f * FSolverVec3(0.0, 0.0f, 1.0f);
+		EXPECT_VECTOR_FLOAT_EQ(Impulse, ExpectedSolverImpulse);
+
+		// Run a second time to check that the impulse doesn't change - the constraint should be
+		// solved exactly in a single iteration for a single body
+		Solver.SolvePosition();
+
+		Impulse = Solver.GetLinearImpulse(1);
+		EXPECT_VECTOR_FLOAT_EQ(Impulse, ExpectedSolverImpulse);
+
+		EXPECT_VECTOR_FLOAT_EQ(CharacterBody.CorrectedP(), FVec3(11.0, 8.0, 25.0));
+
+		FVec3 ExpectedImpulse = ExpectedSolverImpulse / Dt;
+		EXPECT_VECTOR_FLOAT_EQ(Solver.GetLinearImpulse(Dt), ExpectedImpulse);
+		EXPECT_VECTOR_FLOAT_EQ(Solver.GetAngularImpulse(Dt), FVec3::ZeroVector);
+	}
+
+	// Initial overlap at zero velocity with a ground normal pointing at 45 degrees
+	// Should be corrected in a single iteration without introducing velocity
+	// Correction should be along the ground normal
+	TEST_F(CharacterGroundConstraintSolverTest, TestSingleBody_NormalImpulse_InitialOverlapOnSlope)
+	{
+		CharacterBody.SetX(FVec3(10.0, 10.0, 20.0));
+		CharacterBody.SetP(FVec3(10.0, 10.0, 20.0));
+		Data.GroundDistance = 10.0f;
+		Settings.TargetHeight = 15.0f;
+		Settings.RadialForceLimit = 0.0f;
+		Settings.TorqueLimit = 0.0f;
+		Data.GroundNormal = FVec3(1.0, 0.0, 1.0);
+		Data.GroundNormal.SafeNormalize();
+
+		UpdateSingleBody(1, 0);
+
+		FVec3 ExpectedDeltaPos = 5.0 * FMath::Sin(FMath::DegreesToRadians(45.0)) * Data.GroundNormal;
+		FVec3 ExpectedPos = CharacterBody.X() + ExpectedDeltaPos;
+		EXPECT_VECTOR_FLOAT_EQ(CharacterBody.P(), ExpectedPos);
+		EXPECT_VECTOR_FLOAT_EQ(CharacterBody.V(), FVec3(0.0, 0.0, 0.0));
+
+		FVec3 ExpectedImpulse = (1.0 / (CharacterBody.InvM() * Dt)) * ExpectedDeltaPos;
+		EXPECT_VECTOR_FLOAT_EQ(Solver.GetLinearImpulse(Dt), ExpectedImpulse);
+		EXPECT_VECTOR_FLOAT_EQ(Solver.GetAngularImpulse(Dt), FVec3::ZeroVector);
+	}
+
+	// Initial overlap with motion target
+	// Target should be reachable
+	TEST_F(CharacterGroundConstraintSolverTest, TestSingleBody_MotionTarget_LinearReachable)
+	{
+		CharacterBody.SetX(FVec3(5.0, 10.0, 20.0));
+		CharacterBody.SetP(FVec3(10.0, 10.0, 20.0));
+		CharacterBody.SetR(FRotation3::FromAxisAngle(FVec3(0.0, 0.0, 1.0), 0.0));
+		CharacterBody.SetQ(CharacterBody.R());
+		Data.GroundDistance = 10.0f;
+		Settings.TargetHeight = 15.0f;
+		Settings.RadialForceLimit = (1.0 / (CharacterBody.InvM() * Dt * Dt)) * 10.1;
+		Data.TargetDeltaPosition = FVec3(-5.0, 0.0, 0.0);
+		Data.TargetDeltaFacing = -0.4;
+		Settings.TorqueLimit = (1.0 / (CharacterBody.InvILocal().Z * Dt * Dt)) * 0.5;
+
+		UpdateSingleBody(1, 0);
+
+		EXPECT_VECTOR_FLOAT_EQ(CharacterBody.P(), FVec3(0.0, 10.0, 25.0));
+
+		FRotation3 Rot = CharacterBody.R() * FRotation3::FromAxisAngle(FVec3(0.0, 0.0, 1.0), Data.TargetDeltaFacing);
+		// Solver uses approximate formula for angular displacement to quaternion, so only expect
+		// results to be equal up to order of delta angle cubed
+		FReal ErrorTolerance = 0.16 * FMath::Abs(Data.TargetDeltaFacing * Data.TargetDeltaFacing * Data.TargetDeltaFacing);
+		EXPECT_NEAR(CharacterBody.Q().W, Rot.W, ErrorTolerance);
+		EXPECT_NEAR(CharacterBody.Q().X, Rot.X, ErrorTolerance);
+		EXPECT_NEAR(CharacterBody.Q().Y, Rot.Y, ErrorTolerance);
+		EXPECT_NEAR(CharacterBody.Q().Z, Rot.Z, ErrorTolerance);
+
+		FVec3 ExpectedImpulse = (1.0 / (CharacterBody.InvM() * Dt)) * FVec3(-10.0, 0.0, 5.0);
+		EXPECT_VECTOR_FLOAT_EQ(Solver.GetLinearImpulse(Dt), ExpectedImpulse);
+
+		FVec3 ExpectedAngularImpulse = (1.0 / (CharacterBody.InvILocal().Z * Dt)) * Data.TargetDeltaFacing * FVec3(0.0, 0.0, 1.0);
+		EXPECT_VECTOR_FLOAT_EQ(Solver.GetAngularImpulse(Dt), ExpectedAngularImpulse);
+	}
+
+	// Initial overlap with motion target
+	// Target should be unreachable. Solver should only move up to force limit
+	TEST_F(CharacterGroundConstraintSolverTest, TestSingleBody_MotionTarget_LinearUnReachable)
+	{
+		CharacterBody.SetX(FVec3(10.0, 10.0, 20.0));
+		CharacterBody.SetP(FVec3(10.0, 10.0, 20.0));
+		Data.GroundDistance = 10.0f;
+		Settings.TargetHeight = 15.0f;
+		Settings.RadialForceLimit = (1.0 / (CharacterBody.InvM() * Dt * Dt)) * 10.0;
+		Data.TargetDeltaPosition = FVec3(-11.0, 0.0, 0.0);
+		Data.TargetDeltaFacing = 0.6;
+		Settings.TorqueLimit = (1.0 / (CharacterBody.InvILocal().Z * Dt * Dt)) * 0.5;
+
+		UpdateSingleBody(5, 0);
+
+		EXPECT_VECTOR_FLOAT_EQ(CharacterBody.P(), FVec3(0.0, 10.0, 25.0));
+
+		FRotation3 Rot = CharacterBody.R() * FRotation3::FromAxisAngle(FVec3(0.0, 0.0, 1.0), 0.5);
+		// Solver uses approximate formula for angular displacement to quaternion, so only expect
+		// results to be equal up to order of delta angle cubed
+		FReal ErrorTolerance = 0.16 * 0.5 * 0.5 * 0.5;
+		EXPECT_NEAR(CharacterBody.Q().W, Rot.W, ErrorTolerance);
+		EXPECT_NEAR(CharacterBody.Q().X, Rot.X, ErrorTolerance);
+		EXPECT_NEAR(CharacterBody.Q().Y, Rot.Y, ErrorTolerance);
+		EXPECT_NEAR(CharacterBody.Q().Z, Rot.Z, ErrorTolerance);
+
+		FVec3 ExpectedImpulse = (1.0 / (CharacterBody.InvM() * Dt)) * FVec3(-10.0, 0.0, 5.0);
+		EXPECT_VECTOR_FLOAT_EQ(Solver.GetLinearImpulse(Dt), ExpectedImpulse);
+		EXPECT_FLOAT_EQ(FMath::Abs(Solver.GetLinearImpulse(Dt).X), Settings.RadialForceLimit * Dt);
+		EXPECT_FLOAT_EQ(FMath::Abs(Solver.GetAngularImpulse(Dt).Z), Settings.TorqueLimit * Dt);
+	}
+
+	// No overlap with motion target
+	// Target should be reachable but off ground so should be no movement
+	TEST_F(CharacterGroundConstraintSolverTest, TestSingleBody_MotionTarget_LinearReachableOffGround)
+	{
+		CharacterBody.SetX(FVec3(5.0, 10.0, 20.0));
+		CharacterBody.SetP(FVec3(10.0, 10.0, 20.0));
+		Data.GroundDistance = 20.0f;
+		Settings.TargetHeight = 18.0f;
+		Settings.RadialForceLimit = (1.0 / (CharacterBody.InvM() * Dt * Dt)) * 10.1;
+		Data.TargetDeltaPosition = FVec3(-5.0, 0.0, 0.0);
+		Settings.TorqueLimit = 0.0f;
+		Settings.AssumedOnGroundHeight = 1.0f;
+
+		UpdateSingleBody(1, 0);
+
+		EXPECT_VECTOR_FLOAT_EQ(CharacterBody.P(), FVec3(10.0, 10.0, 20.0));
+
+		EXPECT_VECTOR_FLOAT_EQ(Solver.GetLinearImpulse(Dt), FVec3::ZeroVector);
+		EXPECT_VECTOR_FLOAT_EQ(Solver.GetAngularImpulse(Dt), FVec3::ZeroVector);
+	}
+
+	// No overlap with motion target
+	// Target should be reachable. Off ground but within assumed on ground height so should reach target
+	TEST_F(CharacterGroundConstraintSolverTest, TestSingleBody_MotionTarget_LinearReachableSlightlyOffGround)
+	{
+		CharacterBody.SetX(FVec3(5.0, 10.0, 20.0));
+		CharacterBody.SetP(FVec3(10.0, 10.0, 20.0));
+		Data.GroundDistance = 20.0f;
+		Settings.TargetHeight = 18.0f;
+		Settings.RadialForceLimit = (1.0 / (CharacterBody.InvM() * Dt * Dt)) * 10.1;
+		Data.TargetDeltaPosition = FVec3(-5.0, 0.0, 0.0);
+		Settings.TorqueLimit = 0.0f;
+		Settings.AssumedOnGroundHeight = 2.1f;
+
+		UpdateSingleBody(1, 0);
+
+		EXPECT_VECTOR_FLOAT_EQ(CharacterBody.P(), FVec3(0.0, 10.0, 20.0));
+
+		FVec3 ExpectedImpulse = (1.0 / (CharacterBody.InvM() * Dt)) * FVec3(-10.0, 0.0, 0.0);
+		EXPECT_VECTOR_FLOAT_EQ(Solver.GetLinearImpulse(Dt), ExpectedImpulse);
+		EXPECT_VECTOR_FLOAT_EQ(Solver.GetAngularImpulse(Dt), FVec3::ZeroVector);
+	}
+
+	// No overlap with motion target
+	// Motion target into slope is projected onto the ground plane
+	TEST_F(CharacterGroundConstraintSolverTest, TestSingleBody_MotionTarget_OnSlope)
+	{
+		CharacterBody.SetX(FVec3(0.0, 0.0, 10.0));
+		CharacterBody.SetP(FVec3(0.0, 0.0, 10.0));
+		Data.GroundDistance = 10.0f;
+		Settings.TargetHeight = 10.0f;
+		Settings.RadialForceLimit = (1.0 / (CharacterBody.InvM() * Dt * Dt)) * 100.0;
+		Data.TargetDeltaPosition = FVec3(10.0, 0.0, 0.0);
+		Settings.TorqueLimit = 0.0f;
+		Settings.AssumedOnGroundHeight = 1.0f;
+		Data.GroundNormal = FVec3(-1.0, 0.0, 1.0);
+		Data.GroundNormal.SafeNormalize();
+
+		UpdateSingleBody(10, 0);
+
+		FVec3 ExpectedDeltaPos = Data.TargetDeltaPosition - FVec3::DotProduct(Data.TargetDeltaPosition, Data.GroundNormal) * Data.GroundNormal;
+		FVec3 ExpectedPos = CharacterBody.X() + ExpectedDeltaPos;
+		EXPECT_VECTOR_FLOAT_EQ(CharacterBody.P(), ExpectedPos);
+
+		FVec3 ExpectedImpulse = (1.0 / (CharacterBody.InvM() * Dt)) * ExpectedDeltaPos;
+		EXPECT_VECTOR_FLOAT_EQ(Solver.GetLinearImpulse(Dt), ExpectedImpulse);
+		EXPECT_VECTOR_FLOAT_EQ(Solver.GetAngularImpulse(Dt), FVec3::ZeroVector);
+	}
+
+	// No overlap with zero motion target and moving ground body
+	// Character body should move with the ground body
+	// Ground body should be unaffected
+	TEST_F(CharacterGroundConstraintSolverTest, TestTwoBody_MotionTarget_MovingGround)
+	{
+		CharacterBody.SetX(FVec3(0.0, 0.0, 10.0));
+		CharacterBody.SetP(FVec3(0.0, 0.0, 10.0));
+		GroundBody.SetX(FVec3(0.0, 0.0, 0.0));
+		GroundBody.SetP(FVec3(0.0, 10.0, 0.0));
+		Data.GroundDistance = 10.0f;
+		Settings.TargetHeight = 10.0f;
+		Settings.RadialForceLimit = (1.0 / (CharacterBody.InvM() * Dt * Dt)) * 100.0;
+		Data.TargetDeltaPosition = FVec3(10.0, -5.0, 0.0);
+		Settings.TorqueLimit = 0.0f;
+		Settings.AssumedOnGroundHeight = 1.0f;
+
+		UpdateTwoBody(1, 0);
+
+		FVec3 ExpectedDeltaPos =  GroundBody.P() - GroundBody.X() + Data.TargetDeltaPosition;
+		FVec3 ExpectedPos = CharacterBody.X() + ExpectedDeltaPos;
+		EXPECT_VECTOR_FLOAT_EQ(CharacterBody.P(), ExpectedPos);
+		EXPECT_VECTOR_FLOAT_EQ(GroundBody.P(), FVec3(0.0, 10.0, 0.0));
+
+		FVec3 ExpectedImpulse = (1.0 / (CharacterBody.InvM() * Dt)) * ExpectedDeltaPos;
+		EXPECT_VECTOR_FLOAT_EQ(Solver.GetLinearImpulse(Dt), ExpectedImpulse);
+		EXPECT_VECTOR_FLOAT_EQ(Solver.GetAngularImpulse(Dt), FVec3::ZeroVector);
+	}
+
+	// No overlap with zero motion target and rotating ground body
+	// Character body should rotate with the ground body
+	// Ground body should be unaffected
+	TEST_F(CharacterGroundConstraintSolverTest, TestTwoBody_MotionTarget_RotatingGround)
+	{
+		CharacterBody.SetX(FVec3(10.0, 0.0, 10.0));
+		CharacterBody.SetP(FVec3(10.0, 0.0, 10.0));
+		GroundBody.SetInvM(0.001);
+		GroundBody.SetInvILocal(FVec3(0.001, 0.001, 0.001));
+		GroundBody.SetX(FVec3(0.0, 0.0, 0.0));
+		GroundBody.SetP(FVec3(0.0, 0.0, 0.0));
+		GroundBody.SetQ(FRotation3::FromAxisAngle(FVec3(0.0, 0.0, 1.0), 0.1));
+		Data.GroundDistance = 10.0f;
+		Settings.TargetHeight = 10.0f;
+		Settings.RadialForceLimit = (1.0 / (CharacterBody.InvM() * Dt * Dt)) * 100.0;
+		Data.TargetDeltaPosition = FVec3(0.0, 0.0, 0.0);
+		Data.TargetDeltaFacing = 0.0f;
+		Settings.TorqueLimit = (1.0 / (CharacterBody.InvILocal().Z * Dt * Dt)) * 100.0;
+		Settings.AssumedOnGroundHeight = 1.0f;
+
+		UpdateTwoBody(1, 0);
+
+		FVec3 ExpectedDeltaPos = GroundBody.Q() * FVec3(10.0, 0.0, 0.0) - FVec3(10.0, 0.0, 0.0);
+		FVec3 ExpectedPos = CharacterBody.X() + ExpectedDeltaPos;
+		EXPECT_VECTOR_FLOAT_EQ(CharacterBody.P(), ExpectedPos);
+		EXPECT_VECTOR_FLOAT_EQ(GroundBody.P(), FVec3(0.0, 0.0, 0.0));
+
+		// Solver uses approximate formula for angular displacement to quaternion, so only expect
+		// results to be equal up to order of delta angle cubed
+		FReal ErrorTolerance = 0.16 * 0.1 * 0.1 * 0.1;
+		EXPECT_NEAR(CharacterBody.Q().W, GroundBody.Q().W, ErrorTolerance);
+		EXPECT_NEAR(CharacterBody.Q().X, GroundBody.Q().X, ErrorTolerance);
+		EXPECT_NEAR(CharacterBody.Q().Y, GroundBody.Q().Y, ErrorTolerance);
+		EXPECT_NEAR(CharacterBody.Q().Z, GroundBody.Q().Z, ErrorTolerance);
+
+		FVec3 ExpectedImpulse = (1.0 / (CharacterBody.InvM() * Dt)) * ExpectedDeltaPos;
+		EXPECT_VECTOR_FLOAT_EQ(Solver.GetLinearImpulse(Dt), ExpectedImpulse);
+
+		FVec3 ExpectedAngularImpulse = (1.0 / (CharacterBody.InvILocal().Z * Dt)) * FVec3(0.0, 0.0, 0.1);
+		EXPECT_VECTOR_FLOAT_EQ(Solver.GetAngularImpulse(Dt), ExpectedAngularImpulse);
+	}
+
+	// Initial overlap at zero velocity. Two body
+	// Character body should be corrected in a single iteration without introducing velocity
+	// Ground body should not move (doesn't currently get corrected)
+	TEST_F(CharacterGroundConstraintSolverTest, TestTwoBody_NormalImpulse_InitialOverlap)
+	{
+		CharacterBody.SetX(FVec3(10.0, 10.0, 20.0));
+		CharacterBody.SetP(FVec3(10.0, 10.0, 20.0));
+		GroundBody.SetInvM(0.001);
+		GroundBody.SetInvILocal(FVec3(0.0005, 0.0005, 0.0005));
+		GroundBody.SetX(FVec3(10.0, 5.0, 0.0));
+		GroundBody.SetP(FVec3(10.0, 5.0, 0.0));
+		Data.GroundDistance = 10.0f;
+		Settings.TargetHeight = 15.0f;
+		Settings.RadialForceLimit = 0.0f;
+		Settings.TorqueLimit = 0.0f;
+
+		UpdateTwoBody(1, 0);
+
+		EXPECT_VECTOR_FLOAT_EQ(CharacterBody.P(), FVec3(10.0, 10.0, 25.0));
+		EXPECT_VECTOR_FLOAT_EQ(CharacterBody.V(), FVec3(0.0, 0.0, 0.0));
+		EXPECT_VECTOR_FLOAT_EQ(GroundBody.P(), FVec3(10.0, 5.0, 0.0));
+
+		FVec3 ExpectedImpulse = (1.0 / (CharacterBody.InvM() * Dt)) * FVec3(0.0, 0.0, 5.0);
+		EXPECT_VECTOR_FLOAT_EQ(Solver.GetLinearImpulse(Dt), ExpectedImpulse);
+		EXPECT_VECTOR_FLOAT_EQ(Solver.GetAngularImpulse(Dt), FVec3::ZeroVector);
+	}
+
+	// No overlap but projected overlap. Two body
+	// Character body and ground body have the same mass
+	// Linear displacement should be equal and opposite
+	TEST_F(CharacterGroundConstraintSolverTest, TestTwoBody_NormalImpulse_NoInitialOverlap_FinalOverlap)
+	{
+		CharacterBody.SetX(FVec3(10.0, 10.0, 30.0));
+		CharacterBody.SetP(FVec3(10.0, 10.0, 24.0));
+		GroundBody.SetX(FVec3(10.0, 10.0, 0.0));
+		GroundBody.SetP(FVec3(10.0, 10.0, 6.0));
+
+		GroundBody.SetInvM(CharacterBody.InvM());
+		GroundBody.SetInvILocal(CharacterBody.InvILocal());
+
+		Data.GroundDistance = 20.0f;
+		Settings.TargetHeight = 10.0f;
+		Settings.RadialForceLimit = 0.0f;
+		Settings.TorqueLimit = 0.0f;
+
+		UpdateTwoBody(1, 0);
+
+		EXPECT_VECTOR_FLOAT_EQ(CharacterBody.P(), FVec3(10.0, 10.0, 25.0));
+		EXPECT_VECTOR_FLOAT_EQ(GroundBody.P(), FVec3(10.0, 10.0, 5.0));
+
+		FVec3 ExpectedImpulse = (1.0 / ((CharacterBody.InvM() + GroundBody.InvM()) * Dt)) * FVec3(0.0, 0.0, 2.0);
+		EXPECT_VECTOR_FLOAT_EQ(Solver.GetLinearImpulse(Dt), ExpectedImpulse);
+		EXPECT_VECTOR_FLOAT_EQ(Solver.GetAngularImpulse(Dt), FVec3::ZeroVector);
+	}
+
+} // namespace ChaosTest
\ No newline at end of file
diff --git a/Engine/Source/Runtime/Experimental/Chaos/Private/Chaos/Character/CharacterGroundConstraintSolver.cpp b/Engine/Source/Runtime/Experimental/Chaos/Private/Chaos/Character/CharacterGroundConstraintSolver.cpp
new file mode 100644
--- /dev/null
+++ b/Engine/Source/Runtime/Experimental/Chaos/Private/Chaos/Character/CharacterGroundConstraintSolver.cpp
@@ -0,0 +1,196 @@
+// Copyright Epic Games, Inc. All Rights Reserved.
+
+#include "CharacterGroundConstraintSolver.h"
+
+namespace Chaos
+{
+	namespace Private
+	{
+		//////////////////////////////////////////////////////////////////////////
+		/// FBodyData
+
+		FCharacterGroundConstraintSolver::FBodyData::FBodyData()
+		{
+			// Initializes the additional state data
+			// These modifiers are not currently being used, so initialize them once to their default values
+			CharacterBody.Init();
+			GroundBody.Init();
+		}
+
+		void FCharacterGroundConstraintSolver::FBodyData::Init(FSolverBody* InCharacterBody, FSolverBody* InGroundBody)
+		{
+			check(InCharacterBody != nullptr);
+			CharacterBody.SetSolverBody(*InCharacterBody);
+
+			if (InGroundBody != nullptr)
+			{
+				GroundBody.SetSolverBody(*InGroundBody);
+			}
+		}
+
+		void FCharacterGroundConstraintSolver::FBodyData::Reset()
+		{
+			CharacterBody.Reset();
+			GroundBody.Reset();
+		}
+
+		//////////////////////////////////////////////////////////////////////////
+		/// FImpulseData
+
+		FCharacterGroundConstraintSolver::FImpulseData::FImpulseData()
+		{
+			Reset();
+		}
+
+		void FCharacterGroundConstraintSolver::FImpulseData::Reset()
+		{
+			LinearPositionImpulse = FSolverVec3::ZeroVector;
+			AngularImpulse = FSolverReal(0.0f);
+			LinearCorrectionImpulse = FSolverReal(0.0f);
+		}
+
+		//////////////////////////////////////////////////////////////////////////
+		/// FConstraintData
+
+		FCharacterGroundConstraintSolver::FConstraintData::FConstraintData()
+			: Normal(FSolverVec3(0.0, 0.0, 1.0)) // Need to initialize normal as it is used in GetLinearImpulse
+			, VerticalAxis(FSolverVec3(0.0, 0.0, 1.0)) // Need to vertical axis as it is used in GetAngularImpulse
+			, CharacterInvM(-1.0f)
+			, CharacterInvI(-1.0f)
+		{
+		}
+
+		bool FCharacterGroundConstraintSolver::FConstraintData::IsValid()
+		{
+			return (CharacterInvM > 0.0f) && (CharacterInvI > 0.0f);
+		}
+
+		//////////////////////////////////////////////////////////////////////////
+		/// FCharacterGroundConstraintSolver
+
+		void FCharacterGroundConstraintSolver::SetBodies(FSolverBody* CharacterSolverBody, FSolverBody* GroundSolverBody)
+		{
+			BodyData.Init(CharacterSolverBody, GroundSolverBody);
+		}
+
+		void FCharacterGroundConstraintSolver::Reset()
+		{
+			BodyData.Reset();
+		}
+
+		void FCharacterGroundConstraintSolver::GatherInput(FReal Dt, const FCharacterGroundConstraintSettings& Settings, const FCharacterGroundConstraintDynamicData& Data)
+		{
+			ImpulseData.Reset();
+
+			ConstraintData.Normal = FSolverVec3(Data.GroundNormal);
+			ConstraintData.VerticalAxis = FSolverVec3(Settings.VerticalAxis);
+			ConstraintData.AssumedOnGroundHeight = FSolverReal(Settings.AssumedOnGroundHeight);
+
+			// Convert force/torque limits into solver impulse limits
+			ConstraintData.AngularImpulseLimit = FSolverReal(Settings.TorqueLimit * Dt * Dt);
+			ConstraintData.RadialImpulseLimit = FSolverReal(Settings.RadialForceLimit * Dt * Dt);
+
+			FRotation3 RelQuatCharacter = BodyData.CharacterBody.R().Inverse() * BodyData.CharacterBody.Q();
+			FReal CharacterAngularDisplacement = FVec3::DotProduct(RelQuatCharacter.ToRotationVector(), Settings.VerticalAxis);
+			FVec3 CharacterLinearDisplacement = BodyData.CharacterBody.P() - BodyData.CharacterBody.X();
+
+			// Vertical ground constraint
+			FVec3 PC_Init = BodyData.CharacterBody.X() - Settings.TargetHeight * Settings.VerticalAxis;
+			FVec3 PG_Init = BodyData.CharacterBody.X() - Data.GroundDistance * Settings.VerticalAxis;
+			FVec3 InitialError = PC_Init - PG_Init;
+			FVec3 ProjectedError = InitialError + CharacterLinearDisplacement;
+
+			// Motion Target
+			FVec3 MotionTarget = BodyData.CharacterBody.X() + Data.TargetDeltaPosition;
+
+			// If the slope is too steep adjust the target to not point up the slope
+			FReal DP = FVec3::DotProduct(Settings.VerticalAxis, Data.GroundNormal);
+			if (DP <= Settings.CosMaxWalkableSlopeAngle)
+			{
+				FVec3 UpSlope = Settings.VerticalAxis - DP * Data.GroundNormal;
+				UpSlope.Normalize();
+				FReal UpMotion = FVec3::DotProduct(MotionTarget - BodyData.CharacterBody.X(), UpSlope);
+				if (UpMotion > FReal(0.0f))
+				{
+					MotionTarget -= UpMotion * UpSlope;
+				}
+			}
+
+			FVec3 MotionTargetError = BodyData.CharacterBody.P() - MotionTarget;
+			FReal MotionTargetAngularError = CharacterAngularDisplacement - Data.TargetDeltaFacing;
+
+			if (BodyData.GroundBody.IsValid())
+			{
+				FVec3 RC_Init = PC_Init - BodyData.GroundBody.X();
+				FVec3 RG_Init = PG_Init - BodyData.GroundBody.X();
+
+				FRotation3 RelQuatGround = BodyData.GroundBody.R().Inverse() * BodyData.GroundBody.Q();
+				FReal GroundAngularDisplacement = FVec3::DotProduct(RelQuatGround.ToRotationVector(), Settings.VerticalAxis);
+				FVec3 GroundLinearDisplacement = BodyData.GroundBody.P() - BodyData.GroundBody.X() + RelQuatGround * RG_Init - RG_Init;
+
+				ConstraintData.GroundOffset = RelQuatGround * RC_Init;
+				ProjectedError -= GroundLinearDisplacement;
+
+				MotionTargetError -= GroundLinearDisplacement;
+				MotionTargetAngularError -= GroundAngularDisplacement;
+			}
+
+			bool CurrentlyOverlapping = FVec3::DotProduct(InitialError, Data.GroundNormal) < FReal(0.0);
+			FReal ProjectedOverlap = FVec3::DotProduct(ProjectedError, -Data.GroundNormal);
+			bool WillOverlap = ProjectedOverlap > FReal(0.0);
+			bool NeedsCorrection = CurrentlyOverlapping && WillOverlap;
+
+			if (WillOverlap)
+			{
+				FReal DampingFactor = FMath::Clamp(Settings.DampingFactor * Dt, 0.0, 1.0);
+				FVec3 Diff = DampingFactor * ProjectedOverlap * Data.GroundNormal;
+				InitialError += Diff;
+				ProjectedError += Diff;
+			}
+
+			ConstraintData.InitialError = FSolverVec3(InitialError);
+			ConstraintData.InitialProjectedError = FSolverVec3(ProjectedError);
+			ConstraintData.MotionTargetError = FSolverVec3(MotionTargetError);
+			ConstraintData.MotionTargetAngularError = FSolverReal(MotionTargetAngularError);
+
+			FSolverMatrix33 WorldSpaceCharacterInvI = FSolverMatrix33(Utilities::ComputeWorldSpaceInertia(BodyData.CharacterBody.Q(), BodyData.CharacterBody.InvILocal()));
+			ConstraintData.CharacterInvI = FSolverVec3::DotProduct(ConstraintData.VerticalAxis, WorldSpaceCharacterInvI * ConstraintData.VerticalAxis);
+			ConstraintData.CharacterInvM = BodyData.CharacterBody.InvM();
+			ConstraintData.SingleBodyEffectiveMass = FSolverReal(1.0f / BodyData.CharacterBody.InvM());
+			ConstraintData.EffectiveInertia = FSolverReal(1.0f / ConstraintData.CharacterInvI);
+
+			if (BodyData.IsTwoBody())
+			{
+				ConstraintData.GroundInvM = BodyData.GroundBody.InvM();
+				ConstraintData.GroundInvI = FSolverMatrix33(Utilities::ComputeWorldSpaceInertia(BodyData.GroundBody.Q(), BodyData.GroundBody.InvILocal()));
+				ConstraintData.TwoBodyEffectiveMass = FSolverReal(1.0f) / FSolverVec3::DotProduct(Data.GroundNormal, Utilities::ComputeJointFactorMatrix(ConstraintData.GroundOffset, ConstraintData.GroundInvI, ConstraintData.GroundInvM + ConstraintData.CharacterInvM) * Data.GroundNormal);
+				PositionSolveFunction = &SolvePositionTwoBody;
+			}
+			else
+			{
+				PositionSolveFunction = &SolvePositionSingleBody;
+			}
+			// Note: It would be good to have a two body correction function, but that doesn't work so well
+			// if contacts and other constraints don't support correction
+			CorrectionSolveFunction = NeedsCorrection ? &SolveCorrectionSingleBody : &NoSolve;
+		}
+
+		void FCharacterGroundConstraintSolver::ScatterOutput(const FReal Dt, FVec3& OutSolverAppliedForce, FVec3& OutSolverAppliedTorque)
+		{
+			if (Dt > UE_SMALL_NUMBER)
+			{
+				FReal FreqSq = 1.0 / (Dt * Dt);
+				OutSolverAppliedForce = FVec3(ImpulseData.LinearPositionImpulse + ImpulseData.LinearCorrectionImpulse * ConstraintData.Normal) * FreqSq;
+				OutSolverAppliedTorque = FVec3(ImpulseData.AngularImpulse * ConstraintData.VerticalAxis) * FreqSq;
+			}
+			else
+			{
+				OutSolverAppliedForce = FVec3::ZeroVector;
+				OutSolverAppliedTorque = FVec3::ZeroVector;
+			}
+			Reset();
+
+		}
+
+	} // namespace Private
+} // namespace Chaos
\ No newline at end of file
diff --git a/Engine/Source/Runtime/Experimental/Chaos/Private/Chaos/Character/CharacterGroundConstraintSolver.h b/Engine/Source/Runtime/Experimental/Chaos/Private/Chaos/Character/CharacterGroundConstraintSolver.h
new file mode 100644
--- /dev/null
+++ b/Engine/Source/Runtime/Experimental/Chaos/Private/Chaos/Character/CharacterGroundConstraintSolver.h
@@ -0,0 +1,255 @@
+// Copyright Epic Games, Inc. All Rights Reserved.
+#pragma once
+
+#include "CoreMinimal.h"
+
+#include "Chaos/Evolution/SolverBody.h"
+#include "Chaos/ParticleHandleFwd.h"
+#include "Chaos/Character/CharacterGroundConstraintSettings.h"
+#include "Chaos/Utilities.h"
+
+namespace Chaos
+{
+	namespace Private
+	{
+		//////////////////////////////////////////////////////////////////////////
+		/// FCharacterGroundConstraintSolver class
+		
+		/// Computes the and applies linear and angular displacement for a character ground constraint
+		class FCharacterGroundConstraintSolver
+		{
+		public:
+			/// Must call SetBodies and GatherInput before solve
+			void SetBodies(FSolverBody* CharacterSolverBody, FSolverBody* GroundSolverBody = nullptr);
+			void GatherInput(FReal Dt, const FCharacterGroundConstraintSettings& Settings, const FCharacterGroundConstraintDynamicData& Data);
+
+			/// Solve function performs one iteration of the solver
+			void SolvePosition();
+
+			/// Gets the solver output as a force and torque in units ML/T^2 and ML^2/T^2 respectively
+			/// and resets the solver
+			void ScatterOutput(const FReal Dt, FVec3& OutSolverAppliedForce, FVec3& OutSolverAppliedTorque);
+
+			void Reset();
+
+			/// Gets the solver linear displacement for this constraint and converts to an impulse in units of ML/T
+			FVec3 GetLinearImpulse(FReal Dt) const;
+
+			/// Gets the solver angular displacement for this constraint and converts to an impulse in units of ML/T
+			FVec3 GetAngularImpulse(FReal Dt) const;
+
+
+		private:
+			/// Utility functions
+			static FSolverVec3 ProjectOntoPlane(const FSolverVec3& Vector, const FSolverVec3& PlaneNormal);
+			static FSolverVec3 ClampMagnitude(const FSolverVec3& Vector, const FSolverReal& Max);
+			static FSolverReal ClampAbs(const FSolverReal& Value, const FSolverReal& Max);
+
+			/// FBodyData
+			struct FBodyData
+			{
+				FBodyData();
+				void Init(FSolverBody* InCharacterBody, FSolverBody* InGroundBody);
+				bool IsTwoBody();
+				void Reset();
+
+				FConstraintSolverBody CharacterBody;
+				FConstraintSolverBody GroundBody;
+			} BodyData;
+
+			/// FImpulseData
+			struct FImpulseData
+			{
+				FImpulseData();
+				void Reset();
+
+				FSolverVec3 LinearPositionImpulse;
+				FSolverReal AngularImpulse;
+				FSolverReal LinearCorrectionImpulse;
+			} ImpulseData;
+
+			/// Constraints
+			struct FConstraintData
+			{
+				FConstraintData();
+				bool IsValid();
+
+				FSolverMatrix33 GroundInvI;				/// World space ground body inverse inertia
+				FSolverVec3 GroundOffset;				/// Offset vector from ground body CoM to the constraint position
+				FSolverVec3 InitialError;				/// Constraint error pre-integration
+				FSolverVec3 InitialProjectedError;		/// Projected constraint error post-integration
+				FSolverVec3 Normal;						///	Ground plane normal direction
+				FSolverVec3 VerticalAxis;				/// World space vertical axis
+				FSolverVec3 MotionTargetError;			/// Projected constraint error for motion target constraint
+				FSolverReal MotionTargetAngularError;	/// Projected angular facing error for motion target constraint
+				FSolverReal CharacterInvM;				/// Character inverse mass
+				FSolverReal CharacterInvI;				/// Character inverse inertia about the vertical axis
+				FSolverReal GroundInvM;					/// Ground body inverse mass
+				FSolverReal TwoBodyEffectiveMass;		/// Effective mass for the two body system
+				FSolverReal SingleBodyEffectiveMass;	/// Effective mass for a single body constraint (the character mass)
+				FSolverReal EffectiveInertia;			/// Effective angular mass to rotate about the vertical axis
+				FSolverReal AngularImpulseLimit;		/// Angular impulse from the solver is clamped to this limit
+				FSolverReal RadialImpulseLimit;			/// Linear impulse from the solver is clamped to this limit
+				FSolverReal AssumedOnGroundHeight;		/// Height below which the character is assumed to be grounded
+			} ConstraintData;
+
+			using FSolveFunctionType = void (*)(const FConstraintData& ConstraintData, FBodyData& BodyData, FImpulseData& ImpulseData);
+			FSolveFunctionType PositionSolveFunction;
+			FSolveFunctionType CorrectionSolveFunction;
+
+			static void SolveCorrectionSingleBody(const FConstraintData& ConstraintData, FBodyData& BodyData, FImpulseData& ImpulseData);
+			static void SolvePositionSingleBody(const FConstraintData& ConstraintData, FBodyData& BodyData, FImpulseData& ImpulseData);
+			static void SolvePositionTwoBody(const FConstraintData& ConstraintData, FBodyData& BodyData, FImpulseData& ImpulseData);
+			static void NoSolve(const FConstraintData&, FBodyData&, FImpulseData&) {}
+		};
+
+		//////////////////////////////////////////////////////////////////////////
+		/// FCharacterGroundConstraintSolver inline functions
+
+		FORCEINLINE FVec3 FCharacterGroundConstraintSolver::GetLinearImpulse(FReal Dt) const
+		{
+			return Dt > UE_SMALL_NUMBER ? FVec3(ImpulseData.LinearPositionImpulse + ImpulseData.LinearCorrectionImpulse * ConstraintData.Normal) / Dt : FVec3::ZeroVector;
+		}
+
+		FORCEINLINE FVec3 FCharacterGroundConstraintSolver::GetAngularImpulse(FReal Dt) const
+		{
+			return Dt > UE_SMALL_NUMBER ? FVec3(ImpulseData.AngularImpulse * ConstraintData.VerticalAxis) / Dt : FVec3::ZeroVector;
+		}
+
+		//////////////////////////////////////////////////////////////////////////
+		/// Utility functions
+
+		FORCEINLINE FSolverVec3 FCharacterGroundConstraintSolver::ProjectOntoPlane(const FSolverVec3& Vector, const FSolverVec3& PlaneNormal)
+		{
+			return Vector - FSolverVec3::DotProduct(Vector, PlaneNormal) * PlaneNormal;
+		}
+
+		FORCEINLINE FSolverVec3 FCharacterGroundConstraintSolver::ClampMagnitude(const FSolverVec3& Vector, const FSolverReal& Max)
+		{
+			const FSolverReal MagSq = Vector.SizeSquared();
+			const FSolverReal MaxSq = Max * Max;
+			if (MagSq > MaxSq)
+			{
+				if (MaxSq > UE_SMALL_NUMBER)
+				{
+					return Vector * FMath::InvSqrt(MagSq) * Max;
+
+				}
+				else
+				{
+					return FSolverVec3::ZeroVector;
+				}
+			}
+			else
+			{
+				return Vector;
+			}
+		}
+
+		FORCEINLINE FSolverReal FCharacterGroundConstraintSolver::ClampAbs(const FSolverReal& Value, const FSolverReal& Max)
+		{
+			return Value > Max ? Max : (Value < -Max ? -Max : Value);
+		}
+
+		//////////////////////////////////////////////////////////////////////////
+		/// FBodyData
+
+		FORCEINLINE bool FCharacterGroundConstraintSolver::FBodyData::IsTwoBody()
+		{
+			return GroundBody.IsValid() && GroundBody.IsDynamic();
+		}
+
+		//////////////////////////////////////////////////////////////////////////
+		/// Solve Functions
+
+		FORCEINLINE void FCharacterGroundConstraintSolver::SolvePosition()
+		{
+			check(ConstraintData.IsValid());
+
+			// Note: Solving these together as part of the same loop for now but
+			// may be better to split and solve correction first for the whole
+			// system before starting the displacement solver
+			(*CorrectionSolveFunction)(ConstraintData, BodyData, ImpulseData);
+			(*PositionSolveFunction)(ConstraintData, BodyData, ImpulseData);
+		}
+
+		FORCEINLINE void FCharacterGroundConstraintSolver::SolveCorrectionSingleBody(const FConstraintData& ConstraintData, FBodyData& BodyData, FImpulseData& ImpulseData)
+		{
+			FSolverReal Error = FSolverVec3::DotProduct(BodyData.CharacterBody.CP() + ConstraintData.InitialError, ConstraintData.Normal);
+			if (Error < FSolverReal(0.0f))
+			{
+				FSolverReal Delta = -ConstraintData.SingleBodyEffectiveMass * Error;
+				ImpulseData.LinearCorrectionImpulse += Delta;
+				BodyData.CharacterBody.ApplyPositionCorrectionDelta(ConstraintData.CharacterInvM * Delta * ConstraintData.Normal);
+			}
+		}
+
+		FORCEINLINE void FCharacterGroundConstraintSolver::SolvePositionSingleBody(const FConstraintData& ConstraintData, FBodyData& BodyData, FImpulseData& ImpulseData)
+		{
+			// Normal
+			FSolverReal Error = FSolverVec3::DotProduct(BodyData.CharacterBody.DP() + ConstraintData.InitialProjectedError, ConstraintData.Normal);
+			if (Error < FSolverReal(0.0f))
+			{
+				FSolverVec3 Delta = -ConstraintData.SingleBodyEffectiveMass * Error * ConstraintData.Normal;
+				ImpulseData.LinearPositionImpulse += Delta;
+				BodyData.CharacterBody.ApplyPositionDelta(ConstraintData.CharacterInvM * Delta);
+			}
+
+			const FSolverReal NormalImpulse = FSolverVec3::DotProduct(ImpulseData.LinearPositionImpulse, ConstraintData.Normal);
+			if (((NormalImpulse + ImpulseData.LinearCorrectionImpulse) > FSolverReal(0.0f)) || Error < ConstraintData.AssumedOnGroundHeight)
+			{
+				// Target Position
+				FSolverVec3 MotionTargetError = ProjectOntoPlane(ConstraintData.MotionTargetError + BodyData.CharacterBody.DP(), ConstraintData.Normal);
+				FSolverVec3 InitialMotionTargetImpulse = ProjectOntoPlane(ImpulseData.LinearPositionImpulse, ConstraintData.Normal);
+				FSolverVec3 NewMotionTargetImpulse = ClampMagnitude(InitialMotionTargetImpulse - ConstraintData.SingleBodyEffectiveMass * MotionTargetError, ConstraintData.RadialImpulseLimit);
+				FSolverVec3 Delta = NewMotionTargetImpulse - InitialMotionTargetImpulse;
+				ImpulseData.LinearPositionImpulse += Delta;
+				BodyData.CharacterBody.ApplyPositionDelta(ConstraintData.CharacterInvM * Delta);
+
+				// Target Rotation
+				FSolverReal MotionTargetAngularError = ConstraintData.MotionTargetAngularError + FSolverVec3::DotProduct(BodyData.CharacterBody.DQ(), ConstraintData.VerticalAxis);
+				FSolverReal NewAngularImpulse = ClampAbs(ImpulseData.AngularImpulse - ConstraintData.EffectiveInertia * MotionTargetAngularError, ConstraintData.AngularImpulseLimit);
+				FSolverReal AngularDelta = NewAngularImpulse - ImpulseData.AngularImpulse;
+				ImpulseData.AngularImpulse += AngularDelta;
+				BodyData.CharacterBody.ApplyRotationDelta(ConstraintData.CharacterInvI * AngularDelta * ConstraintData.VerticalAxis);
+			}
+		}
+
+		FORCEINLINE void FCharacterGroundConstraintSolver::SolvePositionTwoBody(const FConstraintData& ConstraintData, FBodyData& BodyData, FImpulseData& ImpulseData)
+		{
+			FSolverReal Error = FSolverVec3::DotProduct(ConstraintData.InitialProjectedError + BodyData.CharacterBody.DP() - BodyData.GroundBody.DP()
+				- FSolverVec3::CrossProduct(BodyData.GroundBody.DQ(), ConstraintData.GroundOffset), ConstraintData.Normal);
+			if (Error < FSolverReal(0.0f))
+			{
+				FSolverVec3 Delta = -ConstraintData.TwoBodyEffectiveMass * Error * ConstraintData.Normal;
+				ImpulseData.LinearPositionImpulse += Delta;
+				BodyData.CharacterBody.ApplyPositionDelta(ConstraintData.CharacterInvM * Delta);
+				BodyData.GroundBody.ApplyPositionDelta(-ConstraintData.GroundInvM * Delta);
+				BodyData.GroundBody.ApplyRotationDelta(-ConstraintData.GroundInvI * FSolverVec3::CrossProduct(ConstraintData.GroundOffset, Delta));
+			}
+
+			const FSolverReal NormalImpulse = FSolverVec3::DotProduct(ImpulseData.LinearPositionImpulse, ConstraintData.Normal);
+			const FSolverReal MinNormalImpulse(0.0f);
+			if (((NormalImpulse + ImpulseData.LinearCorrectionImpulse) > MinNormalImpulse) || Error < FSolverReal(5.0f))
+			{
+				// Target Position
+				FSolverVec3 MotionTargetError = ConstraintData.MotionTargetError + BodyData.CharacterBody.DP() - BodyData.GroundBody.DP();
+				MotionTargetError -= FSolverVec3::CrossProduct(BodyData.GroundBody.DQ(), MotionTargetError);
+				MotionTargetError = ProjectOntoPlane(MotionTargetError, ConstraintData.Normal);
+				FSolverVec3 InitialMotionTargetImpulse = ProjectOntoPlane(ImpulseData.LinearPositionImpulse, ConstraintData.Normal);
+				FSolverVec3 NewMotionTargetImpulse = ClampMagnitude(InitialMotionTargetImpulse - ConstraintData.SingleBodyEffectiveMass* MotionTargetError, ConstraintData.RadialImpulseLimit);
+				FSolverVec3 Delta = NewMotionTargetImpulse - InitialMotionTargetImpulse;
+				ImpulseData.LinearPositionImpulse += Delta;
+				BodyData.CharacterBody.ApplyPositionDelta(ConstraintData.CharacterInvM * Delta);
+
+				// Target Rotation
+				FSolverReal MotionTargetAngularError = ConstraintData.MotionTargetAngularError + FSolverVec3::DotProduct(BodyData.CharacterBody.DQ() - BodyData.GroundBody.DQ(), ConstraintData.VerticalAxis);
+				FSolverReal NewAngularImpulse = ClampAbs(ImpulseData.AngularImpulse - ConstraintData.EffectiveInertia * MotionTargetAngularError, ConstraintData.AngularImpulseLimit);
+				FSolverReal AngularDelta = NewAngularImpulse - ImpulseData.AngularImpulse;
+				ImpulseData.AngularImpulse += AngularDelta;
+				BodyData.CharacterBody.ApplyRotationDelta(ConstraintData.CharacterInvI * AngularDelta * ConstraintData.VerticalAxis);
+			}
+		}
+
+	} // namespace Private
+} // namespace Chaos
\ No newline at end of file
diff --git a/Engine/Source/Runtime/Experimental/Chaos/Public/Chaos/Character/CharacterGroundConstraintSettings.h b/Engine/Source/Runtime/Experimental/Chaos/Public/Chaos/Character/CharacterGroundConstraintSettings.h
new file mode 100644
--- /dev/null
+++ b/Engine/Source/Runtime/Experimental/Chaos/Public/Chaos/Character/CharacterGroundConstraintSettings.h
@@ -0,0 +1,64 @@
+// Copyright Epic Games, Inc. All Rights Reserved.
+#pragma once
+
+#include "CoreMinimal.h"
+
+#include "Chaos/Core.h"
+#include "Chaos/Transform.h"
+#include "Chaos/Vector.h"
+
+namespace Chaos
+{
+	class CHAOS_API FCharacterGroundConstraintSettings
+	{
+	public:
+		FCharacterGroundConstraintSettings()
+			: VerticalAxis(FVec3(0, 0, 1))
+			, TargetHeight(0)
+			, RadialForceLimit(1000)
+			, TorqueLimit(1000)
+			, CosMaxWalkableSlopeAngle(0.633)
+			, DampingFactor(0.0)
+			, AssumedOnGroundHeight(5.0)
+			, UserData(nullptr)
+		{
+		}
+
+		bool operator==(const FCharacterGroundConstraintSettings& Other) const
+		{
+			return !FMemory::Memcmp(this, &Other, sizeof(*this));
+		}
+
+		FVec3 VerticalAxis;				/// World space up direction (default z axis)
+		FReal TargetHeight;				/// Desired distance from the character body to the ground
+		FReal RadialForceLimit;			/// How much force the character can apply parallel to the ground plane to reach the target position
+		FReal TorqueLimit;				///	How much torque the character can apply about the vertical axis to reach the target facing direction
+		FReal CosMaxWalkableSlopeAngle;	/// Cosine of the maximum angle in degrees that the character is allowed to walk on
+		FReal DampingFactor;			/// Applies a damping to the vertical ground constraint making it softer. Units: /T
+		FReal AssumedOnGroundHeight;	/// Below this height the character is assumed to be on the ground and can apply force/torque to reach the target position and facing 
+		void* UserData;					/// Reserved for user data
+	};
+
+	class CHAOS_API FCharacterGroundConstraintDynamicData
+	{
+	public:
+		FCharacterGroundConstraintDynamicData()
+			: GroundNormal(FVec3(0, 0, 1))
+			, TargetDeltaPosition(FVec3(0, 0, 0))
+			, TargetDeltaFacing(0)
+			, GroundDistance(0)
+		{
+		}
+
+		bool operator==(const FCharacterGroundConstraintDynamicData& Other) const
+		{
+			return !FMemory::Memcmp(this, &Other, sizeof(*this));
+		}
+
+		FVec3 GroundNormal;				/// World space ground normal
+		FVec3 TargetDeltaPosition;		///	Target linear movement vector. Will be projected onto ground plane
+		FReal TargetDeltaFacing;		/// Target rotation in radians about the vertical axis
+		FReal GroundDistance;			/// Distance from the character body to the ground
+	};
+
+} // namespace Chaos
\ No newline at end of file
