Chaos - fix kinematic body movement not waking other body through joint
- added unit tests for this

The add/remove constraint logic was removing joints between sleeping bodies, but is need to leave them in the graph for this purpose.

diff --git a/Engine/Source/Programs/HeadlessChaos/Private/HeadlessChaosTestJoint.h b/Engine/Source/Programs/HeadlessChaos/Private/HeadlessChaosTestJoint.h
--- a/Engine/Source/Programs/HeadlessChaos/Private/HeadlessChaosTestJoint.h
+++ b/Engine/Source/Programs/HeadlessChaos/Private/HeadlessChaosTestJoint.h
@@ -47,6 +47,11 @@ namespace ChaosTest
 			return Joint;
 		}
 
+		FPBDJointConstraintHandle* GetJoint(const int32 JointIndex)
+		{
+			return Evolution.GetJointConstraints().GetConstConstraintHandles()[JointIndex];
+		}
+
 		void Create()
 		{
 			for (int32 ParticleIndex = 0; ParticleIndex < ParticlePositions.Num(); ++ParticleIndex)
@@ -85,6 +90,12 @@ namespace ChaosTest
 			JointSettings.SetNum(NumParticles - 1);
 		}
 
+		void Advance(const FReal Dt)
+		{
+			Evolution.AdvanceOneTimeStep(Dt);
+			Evolution.EndFrame(Dt);
+		}
+
 		// Initial particles setup
 		TArray<FVec3> ParticlePositions;
 		TArray<FVec3> ParticleSizes;
diff --git a/Engine/Source/Programs/HeadlessChaos/Private/HeadlessChaosTestJointWake.cpp b/Engine/Source/Programs/HeadlessChaos/Private/HeadlessChaosTestJointWake.cpp
new file mode 100644
--- /dev/null
+++ b/Engine/Source/Programs/HeadlessChaos/Private/HeadlessChaosTestJointWake.cpp
@@ -0,0 +1,135 @@
+// Copyright Epic Games, Inc. All Rights Reserved.
+
+#include "HeadlessChaos.h"
+#include "HeadlessChaosTestJoint.h"
+#include "Modules/ModuleManager.h"
+#include "Chaos/ParticleHandle.h"
+
+namespace ChaosTest {
+
+	using namespace Chaos;
+
+	// Create a joint between a kinematic and dynamic body set up as a linear spring (using the drive)
+	// and wait for it to sleep. Then move the kinematic and verify that the dynamic gets woken.
+	GTEST_TEST(JointTests, TestWakeOnKinematicMove)
+	{
+		const int32 NumSolverIterations = 10;
+		const FReal Gravity = 1000;
+		const FReal Dt = 0.01f;
+
+		FJointChainTest<FPBDRigidsEvolutionGBF> Test(NumSolverIterations, Gravity);
+		Test.EnableSleeping();
+		Test.InitChain(2, FVec3(0, 0, -1));
+
+		// Set up the joint as a critically damped spring with a known resting extension under gravity
+		const FReal Extension = 10;
+		const FReal Stiffness = Gravity / Extension;
+		const FReal Damping = 2 * FMath::Sqrt(Stiffness);
+		Test.JointSettings[0].LinearMotionTypes = { EJointMotionType::Free, EJointMotionType::Free, EJointMotionType::Free };
+		Test.JointSettings[0].AngularMotionTypes = { EJointMotionType::Free, EJointMotionType::Free, EJointMotionType::Free };
+		Test.JointSettings[0].bLinearPositionDriveEnabled = { false, false, true };
+		Test.JointSettings[0].bLinearVelocityDriveEnabled = { false, false, true };
+		Test.JointSettings[0].LinearDriveForceMode = EJointForceMode::Acceleration;
+		Test.JointSettings[0].LinearDriveStiffness = FVec3(0, 0, Stiffness);
+		Test.JointSettings[0].LinearDriveDamping = FVec3(0, 0, Damping);
+
+		Test.Create();
+
+		FGenericParticleHandle P0 = Test.GetParticle(0);
+		FConstGenericParticleHandle P1 = Test.GetParticle(1);
+		FPBDJointConstraintHandle* J01 = Test.GetJoint(0);
+
+		// Run the sim till the dynamic sleeps (should be less than 1s of sim time)
+		const int32 MaxSteps = 200;
+		for (int32 Step = 0; Step < MaxSteps; ++Step)
+		{
+			Test.Advance(Dt);
+
+			if (P1->IsSleeping())
+			{
+				break;
+			}
+		}
+
+		EXPECT_TRUE(P1->IsSleeping());
+		if (P1->IsSleeping())
+		{
+			// Verify that the joint we created is still in the constraint graph
+			// This was the source of a wake bug in the past
+			EXPECT_TRUE(J01->IsInConstraintGraph());
+
+			// Move the kinematic and tick
+			const FKinematicTarget KinematicTarget = FKinematicTarget::MakePositionTarget(FRigidTransform3(FVec3(50, 0, 0), FRotation3::FromIdentity()));
+			Test.Evolution.SetParticleKinematicTarget(P0->CastToKinematicParticle(), KinematicTarget);
+
+			Test.Advance(Dt);
+
+			EXPECT_FALSE(P1->IsSleeping());
+		}
+	}
+
+	// Verify that a joint between 2 kinematics is removed from the graph anmd re-added
+	// if one of the particles becomes dynamic again.
+	GTEST_TEST(JointTests, TestKinematicKinematicGraph)
+	{
+		const int32 NumSolverIterations = 10;
+		const FReal Gravity = 1000;
+		const FReal Dt = 0.01f;
+
+		FJointChainTest<FPBDRigidsEvolutionGBF> Test(NumSolverIterations, Gravity);
+		Test.EnableSleeping();
+		Test.InitChain(2, FVec3(0, 0, -1));
+
+		// Set up the joint as a critically damped spring with a known resting extension under gravity
+		const FReal Extension = 10;
+		const FReal Stiffness = Gravity / Extension;
+		const FReal Damping = 2 * FMath::Sqrt(Stiffness);
+		Test.JointSettings[0].LinearMotionTypes = { EJointMotionType::Free, EJointMotionType::Free, EJointMotionType::Free };
+		Test.JointSettings[0].AngularMotionTypes = { EJointMotionType::Free, EJointMotionType::Free, EJointMotionType::Free };
+		Test.JointSettings[0].bLinearPositionDriveEnabled = { false, false, true };
+		Test.JointSettings[0].bLinearVelocityDriveEnabled = { false, false, true };
+		Test.JointSettings[0].LinearDriveForceMode = EJointForceMode::Acceleration;
+		Test.JointSettings[0].LinearDriveStiffness = FVec3(0, 0, Stiffness);
+		Test.JointSettings[0].LinearDriveDamping = FVec3(0, 0, Damping);
+
+		Test.Create();
+
+		FGenericParticleHandle P0 = Test.GetParticle(0);
+		FGenericParticleHandle P1 = Test.GetParticle(1);
+		FPBDJointConstraintHandle* J01 = Test.GetJoint(0);
+
+		// Run the sim till the dynamic sleeps (should be less than 1s of sim time)
+		const int32 MaxSteps = 200;
+		for (int32 Step = 0; Step < MaxSteps; ++Step)
+		{
+			Test.Advance(Dt);
+
+			if (P1->IsSleeping())
+			{
+				break;
+			}
+		}
+
+		// Turn the dynamic into a kinematic and tick.
+		// The joint should still be removed from the graph.
+		// NOTE: This is different to the behaviour of collisions where we keep kinematic-kinematic collisions
+		// in the graph. That is because collisions are transient, but we want the collision to remain allocated
+		// for when the particle becomes dynamic again. We don't need this for Joints because the joints are
+		// persistent and we can just re-add them when the particle becomes dynamic again.
+		Test.Evolution.SetParticleObjectState(P1->CastToRigidParticle(), EObjectStateType::Kinematic);
+		Test.Advance(Dt);
+
+		EXPECT_TRUE(P1->IsKinematic());
+		EXPECT_FALSE(J01->IsInConstraintGraph());
+
+		// Now make the particle dynamic again. 
+		// The joint should get re-added to the graph.
+		Test.Evolution.SetParticleObjectState(P1->CastToRigidParticle(), EObjectStateType::Dynamic);
+		Test.Advance(Dt);
+
+		EXPECT_TRUE(P1->IsDynamic());
+		EXPECT_FALSE(P1->IsSleeping());
+		EXPECT_TRUE(J01->IsInConstraintGraph());
+	}
+
+}
\ No newline at end of file
diff --git a/Engine/Source/Programs/HeadlessChaos/Private/HeadlessChaosTestJoints.cpp b/Engine/Source/Programs/HeadlessChaos/Private/HeadlessChaosTestJoints.cpp
--- a/Engine/Source/Programs/HeadlessChaos/Private/HeadlessChaosTestJoints.cpp
+++ b/Engine/Source/Programs/HeadlessChaos/Private/HeadlessChaosTestJoints.cpp
@@ -42,12 +42,15 @@ namespace ChaosTest {
 		FJointConstraintsTest(const int32 NumIterations, const FReal Gravity)
 			: Base(NumIterations, Gravity)
 		{
-			Evolution.AddConstraintContainer(Joints);
 		}
 
 		FPBDJointConstraintHandle* AddJoint(const TVec2<FGeometryParticleHandle*>& InConstrainedParticles, const FVec3& InLocation)
 		{
-			FPBDJointConstraintHandle* Constraint = Joints.AddConstraint(InConstrainedParticles, FRigidTransform3(InLocation, FRotation3::FromIdentity()));
+			FPBDJointConstraintHandle* Constraint = Evolution.GetJointConstraints().AddConstraint(InConstrainedParticles, FRigidTransform3(InLocation, FRotation3::FromIdentity()));
+
+			FPBDJointSettings Settings = Constraint->GetSettings();
+			Settings.LinearProjection = 1.0;
+			Constraint->SetSettings(Settings);
 
 			// @todo(chaos): this should be automatic, but it's handled by the proxy. Fix this
 			InConstrainedParticles[0]->ParticleConstraints().Add(Constraint);
@@ -79,9 +82,6 @@ namespace ChaosTest {
 		// Initial joints setup
 		TArray<FVec3> JointPositions;
 		TArray<TVec2<int32>> JointParticleIndices;
-
-		// Solver state
-		FPBDJointConstraints Joints;
 	};
 
 	/**
@@ -147,7 +147,7 @@ namespace ChaosTest {
 	template <typename TEvolution>
 	void JointConstraint_SingleMoveRoot()
 	{
-		const int32 NumIterations = 5;
+		const int32 NumIterations = 10;
 		const FReal Gravity = 0;
 		const FReal BoxSize = 1;
 		const FReal BoxMass = 1;
@@ -203,7 +203,7 @@ namespace ChaosTest {
 
 		// Move the kinematic body
 		const FVec3 RootPosition = Test.ParticlePositions[0] + RootDelta;
-		Test.GetParticle(Box1Id)->X() = RootPosition;
+		Test.Evolution.SetParticleKinematicTarget(Test.GetParticle(Box1Id)->CastToKinematicParticle(), FKinematicTarget::MakePositionTarget(RootPosition, Test.GetParticle(Box1Id)->R()));
 
 		for (int32 i = 0; i < 1000; ++i)
 		{
@@ -214,13 +214,14 @@ namespace ChaosTest {
 			EXPECT_LT((Test.GetParticle(Box1Id)->X() - RootPosition).Size(), (FReal)0.1 * BoxSize) << "Post-move instability on frame " << i;
 
 			// Particles should remain fixed distance apart (joint point is at Box1 location)
-			const FVec3 Delta = Test.GetParticle(Box2Id)->CastToRigidParticle()->P() - Test.GetParticle(Box1Id)->X();
+			// NOTE: when using linear joints the error can be moderately large
+			const FVec3 Delta = Test.GetParticle(Box2Id)->X() - Test.GetParticle(Box1Id)->X();
 			const FReal Distance = Delta.Size();
-			EXPECT_NEAR(Distance, ExpectedDistance, (FReal)0.1 * BoxSize) << "Post-move instability on frame " << i;
+			EXPECT_NEAR(Distance, ExpectedDistance, (FReal)0.15 * BoxSize) << "Post-move instability on frame " << i;
 
 			// Joint position calculted from pose and local-space joint pos
 			const FVec3 Box2WorldSpaceJointPosition = Test.GetParticle(Box2Id)->R().RotateVector(Box2LocalSpaceJointPosition) + Test.GetParticle(Box2Id)->X();
-			EXPECT_LT((Box2WorldSpaceJointPosition - RootPosition).Size(), (FReal)0.1 * BoxSize) << "Post-move instability on frame " << i;
+			EXPECT_LT((Box2WorldSpaceJointPosition - RootPosition).Size(), (FReal)0.15 * BoxSize) << "Post-move instability on frame " << i;
 		}
 	}
 
@@ -279,7 +280,7 @@ namespace ChaosTest {
 			const FVec3 RootOffset = FMath::Sin((FReal)2 * PI * Time / AnimPeriod) * AnimDelta;
 			const FVec3 RootPosition = Test.ParticlePositions[0] + RootOffset;
 
-			Test.GetParticle(Box1Id)->X() = RootPosition;
+			Test.Evolution.SetParticleKinematicTarget(Test.GetParticle(Box1Id)->CastToKinematicParticle(), FKinematicTarget::MakePositionTarget(RootPosition, Test.GetParticle(Box1Id)->R()));
 
 			Test.Evolution.AdvanceOneTimeStep(Dt);
 			Test.Evolution.EndFrame(Dt);
@@ -357,7 +358,7 @@ namespace ChaosTest {
 			const FVec3 RootOffset = FMath::Sin((FReal)2 * PI * Time / AnimPeriod) * AnimDelta;
 			const FVec3 RootPosition = Test.ParticlePositions[0] + RootOffset;
 
-			Test.GetParticle(0)->X() = RootPosition;
+			Test.Evolution.SetParticleKinematicTarget(Test.GetParticle(0)->CastToKinematicParticle(), FKinematicTarget::MakePositionTarget(RootPosition, Test.GetParticle(0)->R()));
 
 			Test.Evolution.AdvanceOneTimeStep(Dt);
 			Test.Evolution.EndFrame(Dt);
@@ -443,7 +444,7 @@ namespace ChaosTest {
 			const FVec3 RootOffset = FMath::Sin((FReal)2 * PI * Time / AnimPeriod) * AnimDelta;
 			const FVec3 RootPosition = Test.ParticlePositions[0] + RootOffset;
 
-			Test.GetParticle(0)->X() = RootPosition;
+			Test.Evolution.SetParticleKinematicTarget(Test.GetParticle(0)->CastToKinematicParticle(), FKinematicTarget::MakePositionTarget(RootPosition, Test.GetParticle(0)->R()));
 
 			Test.Evolution.AdvanceOneTimeStep(Dt);
 			Test.Evolution.EndFrame(Dt);
diff --git a/Engine/Source/Programs/HeadlessChaos/Public/HeadlessChaosTestConstraints.h b/Engine/Source/Programs/HeadlessChaos/Public/HeadlessChaosTestConstraints.h
--- a/Engine/Source/Programs/HeadlessChaos/Public/HeadlessChaosTestConstraints.h
+++ b/Engine/Source/Programs/HeadlessChaos/Public/HeadlessChaosTestConstraints.h
@@ -31,6 +31,7 @@ namespace ChaosTest
 			PhysicalMaterial->SleepingAngularThreshold = 0;
 			PhysicalMaterial->DisabledLinearThreshold = 0;
 			PhysicalMaterial->DisabledAngularThreshold = 0;
+			PhysicalMaterial->SleepCounterThreshold = std::numeric_limits<int32>::max();
 
 			Evolution.SetNumPositionIterations(NumIterations);
 			Evolution.SetNumVelocityIterations(1);
@@ -42,38 +43,51 @@ namespace ChaosTest
 		{
 		}
 
-		auto AddParticleBox(const FVec3& Position, const FRotation3& Rotation, const FVec3& Size, FReal Mass)
+		// By default the tests are set up to disable sleeping. Undo this by giving the physical material a non-zero sleep threshold
+		void EnableSleeping()
 		{
-			FGeometryParticleHandle& Particle = Mass > SMALL_NUMBER ? *AppendDynamicParticleBox(SOAs, Size) : *AppendStaticParticleBox(SOAs, Size);
+			// These are the defaults from PhysicalMaterial.cpp
+			PhysicalMaterial->SleepingLinearThreshold = 1;
+			PhysicalMaterial->SleepingAngularThreshold = 0.05;
+			PhysicalMaterial->SleepCounterThreshold = 4;
+		}
+
+		FGeometryParticleHandle* AddParticleBox(const FVec3& Position, const FRotation3& Rotation, const FVec3& Size, FReal Mass)
+		{
+			FPBDRigidParticleHandle* Particle = AppendDynamicParticleBox(SOAs, Size);
+
+			if (Mass == 0)
+			{
+				Evolution.SetParticleObjectState(Particle, EObjectStateType::Kinematic);
+			}
 
-			ResetParticle(&Particle, Position, Rotation, FVec3(0), FVec3(0));
+			ResetParticle(Particle, Position, Rotation, FVec3(0), FVec3(0));
 
-			auto PBDParticlePtr = Particle.CastToRigidParticle();
-			if(PBDParticlePtr && PBDParticlePtr->ObjectState() == EObjectStateType::Dynamic)
+			if(Particle->ObjectState() == EObjectStateType::Dynamic)
 			{
-				auto& PBDParticle = *PBDParticlePtr;
-				PBDParticle.M() = PBDParticle.M() * Mass;
-				PBDParticle.I() = PBDParticle.I() * Mass;
-				PBDParticle.InvM() = PBDParticle.InvM() * ((FReal)1 / Mass);
-				PBDParticle.InvI() = PBDParticle.InvI() * ((FReal)1 / Mass);
+				// NOTE: particle was created with unit mass
+				Particle->M() = Particle->M() * Mass;
+				Particle->I() = Particle->I() * Mass;
+				Particle->InvM() = Particle->InvM() * ((FReal)1 / Mass);
+				Particle->InvI() = Particle->InvI() * ((FReal)1 / Mass);
 			}
-			Evolution.SetPhysicsMaterial(&Particle, MakeSerializable(PhysicalMaterial));
+			Evolution.SetPhysicsMaterial(Particle, MakeSerializable(PhysicalMaterial));
 
-			Evolution.EnableParticle(&Particle);
+			Evolution.EnableParticle(Particle);
 
-			return &Particle;
+			return Particle;
 		}
 
 		void ResetParticle(FGeometryParticleHandle* Particle, const FVec3& Position, const FRotation3& Rotation, const FVec3& Velocity, const FVec3& AngularVelocity)
 		{
 			Particle->X() = Position;
 			Particle->R() = Rotation;
-			if (auto KinParticle = Particle->CastToKinematicParticle())
+			if (FKinematicGeometryParticleHandle* KinParticle = Particle->CastToKinematicParticle())
 			{
 				KinParticle->V() = Velocity;
 				KinParticle->W() = AngularVelocity;
 			}
-			if (auto PBDParticle = Particle->CastToRigidParticle())
+			if (FPBDRigidParticleHandle* PBDParticle = Particle->CastToRigidParticle())
 			{
 				PBDParticle->P() = Position;
 				PBDParticle->Q() = Rotation;
diff --git a/Engine/Source/Runtime/Experimental/Chaos/Private/Chaos/PBDJointConstraints.cpp b/Engine/Source/Runtime/Experimental/Chaos/Private/Chaos/PBDJointConstraints.cpp
--- a/Engine/Source/Runtime/Experimental/Chaos/Private/Chaos/PBDJointConstraints.cpp
+++ b/Engine/Source/Runtime/Experimental/Chaos/Private/Chaos/PBDJointConstraints.cpp
@@ -834,7 +834,7 @@ namespace Chaos
 			check(ConstraintHandle != nullptr);
 
 			const bool bIsInGraph = ConstraintHandle->IsInConstraintGraph();
-			const bool bShouldBeInGraph = IsConstraintEnabled(ConstraintIndex) && CanEvaluate(ConstraintIndex);
+			const bool bShouldBeInGraph = ShouldBeInGraph(ConstraintIndex);
 
 			if (bShouldBeInGraph && !bIsInGraph)
 			{
@@ -863,33 +863,37 @@ namespace Chaos
 		}
 	}
 
-	bool FPBDJointConstraints::CanEvaluate(const int32 ConstraintIndex) const
+	bool FPBDJointConstraints::ShouldBeInGraph(const int32 ConstraintIndex) const
 	{
 		if (!IsConstraintEnabled(ConstraintIndex))
 		{
 			return false;
 		}
 
-		int32 Index0, Index1;
-		GetConstrainedParticleIndices(ConstraintIndex, Index0, Index1);
-		const FGenericParticleHandle Particle0 = FGenericParticleHandle(ConstraintParticles[ConstraintIndex][Index0]);
-		const FGenericParticleHandle Particle1 = FGenericParticleHandle(ConstraintParticles[ConstraintIndex][Index1]);
+		const FGenericParticleHandle Particle0 = FGenericParticleHandle(ConstraintParticles[ConstraintIndex][0]);
+		const FGenericParticleHandle Particle1 = FGenericParticleHandle(ConstraintParticles[ConstraintIndex][1]);
 
-		// check for valid and enabled particles
-		if (Particle0->Handle() == nullptr || Particle0->Disabled()
-			|| Particle1->Handle() == nullptr || Particle1->Disabled())
+		// Must have 2 particles
+		if (!Particle0.IsValid() || !Particle1.IsValid())
 		{
 			return false;
 		}
 
-		if ((Particle0->Sleeping() && Particle1->Sleeping())
-			|| (Particle0->IsKinematic() && Particle1->Sleeping())
-			|| (Particle0->Sleeping() && Particle1->IsKinematic())
-			|| (Particle0->IsKinematic() && Particle1->IsKinematic()))
+		// Both particles must be active
+		if (Particle0->Disabled() || Particle1->Disabled())
 		{
 			return false;
 		}
 
+		// Must have at least one dynamic particle
+		if (!Particle0->IsDynamic() && !Particle1->IsDynamic())
+		{
+			return false;
+		}
+
+		// NOTE: Joints between sleeping particles (or a kinematic and sleeper) must stay in the graph 
+		// so that when a kinematic particle moves it will wake the attached particle/island
+
 		return true;
 	}
 
diff --git a/Engine/Source/Runtime/Experimental/Chaos/Public/Chaos/KinematicTargets.h b/Engine/Source/Runtime/Experimental/Chaos/Public/Chaos/KinematicTargets.h
--- a/Engine/Source/Runtime/Experimental/Chaos/Public/Chaos/KinematicTargets.h
+++ b/Engine/Source/Runtime/Experimental/Chaos/Public/Chaos/KinematicTargets.h
@@ -34,6 +34,11 @@ namespace Chaos
 			return TKinematicTarget<T, d>(InTransform);
 		}
 
+		static TKinematicTarget<T, d> MakePositionTarget(const TVector<T, d>& InPosition, const TRotation<T, d>& InRotation)
+		{
+			return TKinematicTarget<T, d>(InPosition, InRotation);
+		}
+
 		TKinematicTarget()
 			: Position(0)
 			, Rotation(TRotation<T, d>::FromIdentity())
@@ -146,6 +151,13 @@ namespace Chaos
 		{
 		}
 
+		TKinematicTarget(const TVector<T, d>& InPosition, const TRotation<T, d>& InRotation)
+			: Position(InPosition)
+			, Rotation(InRotation)
+			, Mode(EKinematicTargetMode::Position)
+		{
+		}
+
 		TVector<T, d> Position;
 		TRotation<T, d> Rotation;
 		EKinematicTargetMode Mode;
diff --git a/Engine/Source/Runtime/Experimental/Chaos/Public/Chaos/PBDJointConstraints.h b/Engine/Source/Runtime/Experimental/Chaos/Public/Chaos/PBDJointConstraints.h
--- a/Engine/Source/Runtime/Experimental/Chaos/Public/Chaos/PBDJointConstraints.h
+++ b/Engine/Source/Runtime/Experimental/Chaos/Public/Chaos/PBDJointConstraints.h
@@ -356,7 +356,7 @@ namespace Chaos
 		void ColorConstraints();
 		void SortConstraints();
 
-		bool CanEvaluate(const int32 ConstraintIndex) const;
+		bool ShouldBeInGraph(const int32 ConstraintIndex) const;
 
 		void ApplyPlasticityLimits(const int32 ConstraintIndex, const FSolverBody& SolverBody0, const FSolverBody& SolverBody1);
 
