Chaos - joint drive force limits now working

diff --git a/Engine/Source/Programs/HeadlessChaos/Private/HeadlessChaosTestJointForces.cpp b/Engine/Source/Programs/HeadlessChaos/Private/HeadlessChaosTestJointForces.cpp
--- a/Engine/Source/Programs/HeadlessChaos/Private/HeadlessChaosTestJointForces.cpp
+++ b/Engine/Source/Programs/HeadlessChaos/Private/HeadlessChaosTestJointForces.cpp
@@ -169,4 +169,304 @@ namespace ChaosTest {
 	{
 		JointForces_Angular<FPBDRigidsEvolutionGBF>();
 	}
+
+	// Check that a joint drive linear stiffness calculates the correct force F=-K.X assuming implicit integration
+	// NOTE: Uses Joint Force mode which isn't currently selectable from the editor
+	//
+	GTEST_TEST(JointForceTests, TestLinearDriveForceMode_Force)
+	{
+		const int32 NumSolverIterations = 20;
+		const FReal Gravity = 0;
+		const FReal Dt = 0.01;
+		const int32 NumBodies = 2;
+
+		const FReal Extension = 10;
+		const FReal Stiffness = 10000;
+		const FReal Damping = 0;
+
+		FJointChainTest<FPBDRigidsEvolutionGBF> Test(NumSolverIterations, Gravity);
+		Test.InitChain(NumBodies, FVec3(0, 0, -1));
+
+		// Disable all limits
+		Test.JointSettings[0].LinearMotionTypes = { EJointMotionType::Free, EJointMotionType::Free, EJointMotionType::Free };
+		Test.JointSettings[0].AngularMotionTypes = { EJointMotionType::Free, EJointMotionType::Free, EJointMotionType::Free };
+
+		// Set up the drive in force mode
+		Test.JointSettings[0].bLinearPositionDriveEnabled = { false, false, true };
+		Test.JointSettings[0].bLinearVelocityDriveEnabled = { false, false, false };
+		Test.JointSettings[0].LinearDriveForceMode = EJointForceMode::Force;
+		Test.JointSettings[0].LinearDriveStiffness = FVec3(0, 0, Stiffness);
+		Test.JointSettings[0].LinearDriveDamping = FVec3(0, 0, Damping);
+
+		Test.Create();
+
+		FGenericParticleHandle P1 = Test.GetParticle(1);
+
+		// Reposition the particle to have some extension in the spring
+		P1->InitTransform(P1->X() + FVec3(0,0,-Extension), P1->R());
+
+		// Run the sim
+		Test.Evolution.AdvanceOneTimeStep(Dt);
+		Test.Evolution.EndFrame(Dt);
+		
+		// Calculate expected force from F = -K.X with implicit integration
+		const FReal M = Test.ParticleMasses[1];
+		FReal ExpectedForceZ = 0;
+		FReal DP = 0;
+		for (int32 It = 0; It < 10; ++It)
+		{
+			const FReal X = -Extension + DP;
+			const FReal F = -Stiffness * X;
+			const FReal DV = ((F - ExpectedForceZ) / M) * Dt;
+			DP += DV * Dt;
+			ExpectedForceZ = F;
+		}
+
+		// Check the joint forces agree
+		const FReal ForceZ = -Test.Evolution.GetJointConstraints().GetConstraintLinearImpulse(0).Z / Dt;
+		EXPECT_NEAR(ForceZ, ExpectedForceZ, 0.01);
+	}
+
+	// Check that a joint drive linear damping calculates the correct force F=-D.V assuming implicit integration
+	// NOTE: Uses Joint Force mode which isn't currently selectable from the editor
+	//
+	GTEST_TEST(JointForceTests, TestLinearDriveForceMode_Damping)
+	{
+		const int32 NumSolverIterations = 20;
+		const FReal Gravity = 0;
+		const FReal Dt = 0.01;
+		const int32 NumBodies = 2;
+
+		const FReal Velocity = 100;
+		const FReal Stiffness = 0;
+		const FReal Damping = 2000;
+
+		FJointChainTest<FPBDRigidsEvolutionGBF> Test(NumSolverIterations, Gravity);
+		Test.InitChain(NumBodies, FVec3(0, 0, -1));
+
+		// Disable all limits
+		Test.JointSettings[0].LinearMotionTypes = { EJointMotionType::Free, EJointMotionType::Free, EJointMotionType::Free };
+		Test.JointSettings[0].AngularMotionTypes = { EJointMotionType::Free, EJointMotionType::Free, EJointMotionType::Free };
+
+		// Set up the drive in force mode
+		Test.JointSettings[0].bLinearPositionDriveEnabled = { false, false, false };
+		Test.JointSettings[0].bLinearVelocityDriveEnabled = { false, false, true };
+		Test.JointSettings[0].LinearDriveForceMode = EJointForceMode::Force;
+		Test.JointSettings[0].LinearDriveStiffness = FVec3(0, 0, Stiffness);
+		Test.JointSettings[0].LinearDriveDamping = FVec3(0, 0, Damping);
+
+		Test.Create();
+
+		FGenericParticleHandle P1 = Test.GetParticle(1);
+
+		// Give the particle velocity so that joint damping has some work to do
+		P1->SetV(FVec3(0,0,-Velocity));
+
+		// Run the sim
+		Test.Evolution.AdvanceOneTimeStep(Dt);
+		Test.Evolution.EndFrame(Dt);
+
+		// Calculate expected force from F = -D.V with implicit integration
+		const FReal M = Test.ParticleMasses[1];
+		FReal ExpectedForceZ = 0;
+		FReal DP = 0;
+		for (int32 It = 0; It < 10; ++It)
+		{
+			const FReal V = -Velocity + DP / Dt;
+			const FReal F = -Damping * V;
+			const FReal DV = ((F - ExpectedForceZ) / M) * Dt;
+			DP += DV * Dt;
+			ExpectedForceZ = F;
+		}
+
+		// Check the joint forces agree
+		const FReal ForceZ = -Test.Evolution.GetJointConstraints().GetConstraintLinearImpulse(0).Z / Dt;
+		EXPECT_NEAR(ForceZ, ExpectedForceZ, 0.1);
+	}
+
+	// Check that a handing mass on a joint drive reaches the correct extension with the correct spring force when using Force mode.
+	// This is just a pre-test for TestLinearDriveForceMode_MaxForce to verify that everything works in the absense of a max force
+	//
+	GTEST_TEST(JointForceTests, TestLinearDriveForceMode_MaxForcePreTest)
+	{
+		const int32 NumSolverIterations = 20;
+		const FReal Gravity = 1000;
+		const FReal Dt = 0.01;
+		const int32 NumSteps = 100;
+		const int32 NumBodies = 2;
+		const FReal Extension = 10;
+
+		FJointChainTest<FPBDRigidsEvolutionGBF> Test(NumSolverIterations, Gravity);
+		Test.InitChain(NumBodies, FVec3(0, 0, -1));
+
+		// Disable all limits
+		Test.JointSettings[0].LinearMotionTypes = { EJointMotionType::Free, EJointMotionType::Free, EJointMotionType::Free };
+		Test.JointSettings[0].AngularMotionTypes = { EJointMotionType::Free, EJointMotionType::Free, EJointMotionType::Free };
+
+		// Set up the drive without a force limit
+		const FReal Stiffness = Test.ParticleMasses[1] * Gravity / Extension;
+		const FReal Damping = 2 * FMath::Sqrt(Stiffness * Test.ParticleMasses[1]);
+		Test.JointSettings[0].bLinearPositionDriveEnabled = { false, false, true };
+		Test.JointSettings[0].bLinearVelocityDriveEnabled = { false, false, true };
+		Test.JointSettings[0].LinearDriveForceMode = EJointForceMode::Force;
+		Test.JointSettings[0].LinearDriveStiffness = FVec3(0, 0, Stiffness);
+		Test.JointSettings[0].LinearDriveDamping = FVec3(0, 0, Damping);
+
+		Test.Create();
+
+		FConstGenericParticleHandle P1 = Test.GetParticle(1);
+
+		// Run the sim - the dangling box should reach a steady state
+		for (int32 i = 0; i < NumSteps; ++i)
+		{
+			Test.Evolution.AdvanceOneTimeStep(Dt);
+			Test.Evolution.EndFrame(Dt);
+		}
+
+		// We should be stationary at the desired extension
+		const FReal ExpectedVZ = 0;
+		const FReal ExpectedZ = Test.ParticlePositions[1].Z - Extension;
+		EXPECT_NEAR(P1->V().Z, ExpectedVZ, 0.1);
+		EXPECT_NEAR(P1->X().Z, ExpectedZ, 0.1);
+	}
+
+	// Check that the maximum drive force setting honored for linear drives.
+	// We repeat TestLinearDriveForceMode_MaxForcePreTest but with a max force which is less than K.X at the rest extension
+	//
+	GTEST_TEST(JointForceTests, TestLinearDriveForceMode_MaxForce)
+	{
+		const int32 NumSolverIterations = 20;
+		const FReal Gravity = 1000;
+		const FReal Dt = 0.01;
+		const int32 NumSteps = 100;
+		const int32 NumBodies = 2;
+		const FReal Extension = 10;
+
+		FJointChainTest<FPBDRigidsEvolutionGBF> Test(NumSolverIterations, Gravity);
+		Test.InitChain(NumBodies, FVec3(0, 0, -1));
+
+		// Disable all limits
+		Test.JointSettings[0].LinearMotionTypes = { EJointMotionType::Free, EJointMotionType::Free, EJointMotionType::Free };
+		Test.JointSettings[0].AngularMotionTypes = { EJointMotionType::Free, EJointMotionType::Free, EJointMotionType::Free };
+
+		// Set up the drive with a force limit
+		const FReal Stiffness = Test.ParticleMasses[1] * Gravity / Extension;
+		const FReal Damping = 2 * FMath::Sqrt(Stiffness * Test.ParticleMasses[1]);
+		Test.JointSettings[0].bLinearPositionDriveEnabled = { false, false, true };
+		Test.JointSettings[0].bLinearVelocityDriveEnabled = { false, false, true };
+		Test.JointSettings[0].LinearDriveForceMode = EJointForceMode::Force;
+		Test.JointSettings[0].LinearDriveStiffness = FVec3(0, 0, Stiffness);
+		Test.JointSettings[0].LinearDriveDamping = FVec3(0, 0, Damping);
+		Test.JointSettings[0].LinearDriveMaxForce = FVec3(0,0, 0.5 * Stiffness * Extension);
+
+		Test.Create();
+
+		FConstGenericParticleHandle P1 = Test.GetParticle(1);
+
+		// Run the sim - the dangling box should reach a steady state
+		for (int32 i = 0; i < NumSteps; ++i)
+		{
+			Test.Evolution.AdvanceOneTimeStep(Dt);
+			Test.Evolution.EndFrame(Dt);
+
+			const FReal ForceZ = -Test.Evolution.GetJointConstraints().GetConstraintLinearImpulse(0).Z / Dt;
+			EXPECT_LT(ForceZ, Test.JointSettings[0].LinearDriveMaxForce.Z + 0.1);
+		}
+	}
+
+	// Check that a handing mass on a joint drive reaches the correct extension with the correct spring force when using Acceleration mode.
+	// This is just a pre-test for TestLinearDriveForceMode_MaxForce to verify that everything works in the absense of a max force
+	// NOTE: Using Accleration mode on the drive but with adjusted stiffness and damping. 
+	// Should yield same results as TestLinearDriveForceMode_MaxForcePreTest.
+	//
+	GTEST_TEST(JointForceTests, TestLinearDriveAccMode_MaxForcePreTest)
+	{
+		const int32 NumSolverIterations = 20;
+		const FReal Gravity = 1000;
+		const FReal Dt = 0.01;
+		const int32 NumSteps = 100;
+		const int32 NumBodies = 2;
+		const FReal Extension = 10;
+
+		FJointChainTest<FPBDRigidsEvolutionGBF> Test(NumSolverIterations, Gravity);
+		Test.InitChain(NumBodies, FVec3(0, 0, -1));
+
+		// Disable all limits
+		Test.JointSettings[0].LinearMotionTypes = { EJointMotionType::Free, EJointMotionType::Free, EJointMotionType::Free };
+		Test.JointSettings[0].AngularMotionTypes = { EJointMotionType::Free, EJointMotionType::Free, EJointMotionType::Free };
+
+		// Set up the drive without a force limit
+		// NOTE: Acceleration mode - no masses in expressions
+		const FReal Stiffness = Gravity / Extension;
+		const FReal Damping = 2 * FMath::Sqrt(Stiffness);
+		Test.JointSettings[0].bLinearPositionDriveEnabled = { false, false, true };
+		Test.JointSettings[0].bLinearVelocityDriveEnabled = { false, false, true };
+		Test.JointSettings[0].LinearDriveForceMode = EJointForceMode::Acceleration;
+		Test.JointSettings[0].LinearDriveStiffness = FVec3(0, 0, Stiffness);
+		Test.JointSettings[0].LinearDriveDamping = FVec3(0, 0, Damping);
+
+		Test.Create();
+
+		FConstGenericParticleHandle P1 = Test.GetParticle(1);
+
+		// Run the sim - the dangling box should reach a steady state
+		for (int32 i = 0; i < NumSteps; ++i)
+		{
+			Test.Evolution.AdvanceOneTimeStep(Dt);
+			Test.Evolution.EndFrame(Dt);
+		}
+
+		// We should be stationary at the desired extension
+		const FReal ExpectedVZ = 0;
+		const FReal ExpectedZ = Test.ParticlePositions[1].Z - Extension;
+		EXPECT_NEAR(P1->V().Z, ExpectedVZ, 0.1);
+		EXPECT_NEAR(P1->X().Z, ExpectedZ, 0.1);
+	}
+
+
+	// Check that the maximum drive force setting honored for linear drives.
+	// We repeat TestLinearDriveForceMode_MaxForcePreTest but with a max force which is less than K.X at the rest extension
+	// NOTE: Using Accleration mode on the drive but with adjusted stiffness and damping. 
+	//
+	GTEST_TEST(JointForceTests, TestLinearDriveAccMode_MaxForce)
+	{
+		const int32 NumSolverIterations = 20;
+		const FReal Gravity = 1000;
+		const FReal Dt = 0.01;
+		const int32 NumSteps = 100;
+		const int32 NumBodies = 2;
+		const FReal Extension = 10;
+
+		FJointChainTest<FPBDRigidsEvolutionGBF> Test(NumSolverIterations, Gravity);
+		Test.InitChain(NumBodies, FVec3(0, 0, -1));
+
+		// Disable all limits
+		Test.JointSettings[0].LinearMotionTypes = { EJointMotionType::Free, EJointMotionType::Free, EJointMotionType::Free };
+		Test.JointSettings[0].AngularMotionTypes = { EJointMotionType::Free, EJointMotionType::Free, EJointMotionType::Free };
+
+		// Set up the drive with a force limit
+		// NOTE: Acceleration mode - no masses in expressions
+		const FReal Stiffness = Gravity / Extension;
+		const FReal Damping = 2 * FMath::Sqrt(Stiffness);
+		Test.JointSettings[0].bLinearPositionDriveEnabled = { false, false, true };
+		Test.JointSettings[0].bLinearVelocityDriveEnabled = { false, false, true };
+		Test.JointSettings[0].LinearDriveForceMode = EJointForceMode::Acceleration;
+		Test.JointSettings[0].LinearDriveStiffness = FVec3(0, 0, Stiffness);
+		Test.JointSettings[0].LinearDriveDamping = FVec3(0, 0, Damping);
+		Test.JointSettings[0].LinearDriveMaxForce = FVec3(0, 0, 0.5 * Stiffness * Extension);
+
+		Test.Create();
+
+		FConstGenericParticleHandle P1 = Test.GetParticle(1);
+
+		// Run the sim - the dangling box should reach a steady state
+		for (int32 i = 0; i < NumSteps; ++i)
+		{
+			Test.Evolution.AdvanceOneTimeStep(Dt);
+			Test.Evolution.EndFrame(Dt);
+
+			const FReal ForceZ = -Test.Evolution.GetJointConstraints().GetConstraintLinearImpulse(0).Z / Dt;
+			EXPECT_LT(ForceZ, Test.JointSettings[0].LinearDriveMaxForce.Z * Test.ParticleMasses[1] + 0.1);
+		}
+	}
 }
\ No newline at end of file
diff --git a/Engine/Source/Runtime/Engine/Private/PhysicsEngine/Experimental/PhysInterface_Chaos.cpp b/Engine/Source/Runtime/Engine/Private/PhysicsEngine/Experimental/PhysInterface_Chaos.cpp
--- a/Engine/Source/Runtime/Engine/Private/PhysicsEngine/Experimental/PhysInterface_Chaos.cpp
+++ b/Engine/Source/Runtime/Engine/Private/PhysicsEngine/Experimental/PhysInterface_Chaos.cpp
@@ -384,6 +384,7 @@ void FPhysInterface_Chaos::UpdateLinearDrive_AssumesLocked(const FPhysicsConstra
 			Constraint->SetLinearDriveForceMode(Chaos::EJointForceMode::Acceleration);
 			Constraint->SetLinearDriveStiffness(Chaos::ConstraintSettings::LinearDriveStiffnessScale() * Chaos::FVec3(InDriveParams.XDrive.Stiffness, InDriveParams.YDrive.Stiffness, InDriveParams.ZDrive.Stiffness));
 			Constraint->SetLinearDriveDamping(Chaos::ConstraintSettings::LinearDriveDampingScale() * Chaos::FVec3(InDriveParams.XDrive.Damping, InDriveParams.YDrive.Damping, InDriveParams.ZDrive.Damping));
+			Constraint->SetLinearDriveMaxForce(Chaos::FVec3(InDriveParams.XDrive.MaxForce, InDriveParams.YDrive.MaxForce, InDriveParams.ZDrive.MaxForce));
 		}
 	}
 }
@@ -454,11 +455,13 @@ void FPhysInterface_Chaos::UpdateAngularDrive_AssumesLocked(const FPhysicsConstr
 			{
 				Constraint->SetAngularDriveStiffness(Chaos::ConstraintSettings::AngularDriveStiffnessScale() * Chaos::FVec3(InDriveParams.TwistDrive.Stiffness, InDriveParams.SwingDrive.Stiffness, InDriveParams.SwingDrive.Stiffness));
 				Constraint->SetAngularDriveDamping(Chaos::ConstraintSettings::AngularDriveDampingScale() * Chaos::FVec3(InDriveParams.TwistDrive.Damping, InDriveParams.SwingDrive.Damping, InDriveParams.SwingDrive.Damping));
+				Constraint->SetAngularDriveMaxTorque(Chaos::FVec3(InDriveParams.TwistDrive.MaxForce, InDriveParams.SwingDrive.MaxForce, InDriveParams.SwingDrive.MaxForce));
 			}
 			else
 			{
 				Constraint->SetAngularDriveStiffness(Chaos::ConstraintSettings::AngularDriveStiffnessScale() * Chaos::FVec3(InDriveParams.SlerpDrive.Stiffness));
 				Constraint->SetAngularDriveDamping(Chaos::ConstraintSettings::AngularDriveDampingScale() * Chaos::FVec3(InDriveParams.SlerpDrive.Damping));
+				Constraint->SetAngularDriveMaxTorque(Chaos::FVec3(InDriveParams.TwistDrive.MaxForce, InDriveParams.SwingDrive.MaxForce, InDriveParams.SwingDrive.MaxForce));
 			}
 		}
 	}
diff --git a/Engine/Source/Runtime/Experimental/Chaos/Private/Chaos/Joint/JointSolverConstraints.cpp b/Engine/Source/Runtime/Experimental/Chaos/Private/Chaos/Joint/JointSolverConstraints.cpp
--- a/Engine/Source/Runtime/Experimental/Chaos/Private/Chaos/Joint/JointSolverConstraints.cpp
+++ b/Engine/Source/Runtime/Experimental/Chaos/Private/Chaos/Joint/JointSolverConstraints.cpp
@@ -579,7 +579,7 @@ namespace Chaos
 		RowData.Limit = 0.0f;
 		RowData.Stiffness = FPBDJointUtilities::GetLinearDriveStiffness(SolverSettings, JointSettings, AxisIndex);
 		RowData.Damping = FPBDJointUtilities::GetLinearDriveDamping(SolverSettings, JointSettings, AxisIndex);
-		RowData.bIsAccelerationMode = FPBDJointUtilities::GetDriveAccelerationMode(SolverSettings, JointSettings);
+		RowData.bIsAccelerationMode = FPBDJointUtilities::GetLinearDriveAccelerationMode(SolverSettings, JointSettings);
 		RowData.bIsSoft = true;
 	}
 
@@ -665,7 +665,7 @@ namespace Chaos
 
 		RowData.Stiffness = FPBDJointUtilities::GetAngularTwistDriveStiffness(SolverSettings, JointSettings);
 		RowData.Damping = FPBDJointUtilities::GetAngularTwistDriveDamping(SolverSettings, JointSettings);
-		RowData.bIsAccelerationMode = FPBDJointUtilities::GetDriveAccelerationMode(SolverSettings, JointSettings);
+		RowData.bIsAccelerationMode = FPBDJointUtilities::GetAngularDriveAccelerationMode(SolverSettings, JointSettings);
 		RowData.bIsSoft = true;
 	}
 
@@ -688,7 +688,7 @@ namespace Chaos
 
 		RowData.Stiffness = FPBDJointUtilities::GetAngularSwingDriveStiffness(SolverSettings, JointSettings);
 		RowData.Damping = FPBDJointUtilities::GetAngularSwingDriveDamping(SolverSettings, JointSettings);
-		RowData.bIsAccelerationMode = FPBDJointUtilities::GetDriveAccelerationMode(SolverSettings, JointSettings);
+		RowData.bIsAccelerationMode = FPBDJointUtilities::GetAngularDriveAccelerationMode(SolverSettings, JointSettings);
 		RowData.bIsSoft = true;
 	}
 
@@ -708,7 +708,7 @@ namespace Chaos
 
 		RowData.Stiffness = FPBDJointUtilities::GetAngularSLerpDriveStiffness(SolverSettings, JointSettings);
 		RowData.Damping = FPBDJointUtilities::GetAngularSLerpDriveDamping(SolverSettings, JointSettings);
-		RowData.bIsAccelerationMode = FPBDJointUtilities::GetDriveAccelerationMode(SolverSettings, JointSettings);
+		RowData.bIsAccelerationMode = FPBDJointUtilities::GetAngularDriveAccelerationMode(SolverSettings, JointSettings);
 		RowData.bIsSoft = true;
 	}
 
diff --git a/Engine/Source/Runtime/Experimental/Chaos/Private/Chaos/Joint/PBDJointCachedSolverGaussSeidel.cpp b/Engine/Source/Runtime/Experimental/Chaos/Private/Chaos/Joint/PBDJointCachedSolverGaussSeidel.cpp
--- a/Engine/Source/Runtime/Experimental/Chaos/Private/Chaos/Joint/PBDJointCachedSolverGaussSeidel.cpp
+++ b/Engine/Source/Runtime/Experimental/Chaos/Private/Chaos/Joint/PBDJointCachedSolverGaussSeidel.cpp
@@ -819,17 +819,17 @@ void FPBDJointCachedSolver::InitRotationConstraints(
 	const int32 S2 = (int32)EJointAngularConstraintIndex::Swing2;
 
 	RotationConstraints.InitDatas(TW,FPBDJointUtilities::GetSoftTwistLimitEnabled(SolverSettings, JointSettings) && !bAngularLocked[TW],
-FPBDJointUtilities::GetSoftTwistStiffness(SolverSettings, JointSettings),
-							   FPBDJointUtilities::GetSoftTwistDamping(SolverSettings, JointSettings),
+					FPBDJointUtilities::GetSoftTwistStiffness(SolverSettings, JointSettings),
+					FPBDJointUtilities::GetSoftTwistDamping(SolverSettings, JointSettings),
 					FPBDJointUtilities::GetTwistStiffness(SolverSettings, JointSettings));
 
 	RotationConstraints.InitDatas(S1,FPBDJointUtilities::GetSoftSwingLimitEnabled(SolverSettings, JointSettings) && !bAngularLocked[S1],
-	FPBDJointUtilities::GetSoftSwingStiffness(SolverSettings, JointSettings),
+					FPBDJointUtilities::GetSoftSwingStiffness(SolverSettings, JointSettings),
 					FPBDJointUtilities::GetSoftSwingDamping(SolverSettings, JointSettings),
 					FPBDJointUtilities::GetSwingStiffness(SolverSettings, JointSettings));
 
 	RotationConstraints.InitDatas(S2, FPBDJointUtilities::GetSoftSwingLimitEnabled(SolverSettings, JointSettings) && !bAngularLocked[S2],
-	FPBDJointUtilities::GetSoftSwingStiffness(SolverSettings, JointSettings),
+					FPBDJointUtilities::GetSoftSwingStiffness(SolverSettings, JointSettings),
 					FPBDJointUtilities::GetSoftSwingDamping(SolverSettings, JointSettings),
 					FPBDJointUtilities::GetSwingStiffness(SolverSettings, JointSettings));
 
@@ -1263,7 +1263,7 @@ void FPBDJointCachedSolver::InitPositionDrives(
 			(JointSettings.bLinearPositionDriveEnabled[2] || JointSettings.bLinearVelocityDriveEnabled[2]) && (JointSettings.LinearMotionTypes[2] != EJointMotionType::Locked),
 		};
 	
-		PositionDrives.bAccelerationMode = FPBDJointUtilities::GetDriveAccelerationMode(SolverSettings, JointSettings);
+		PositionDrives.bAccelerationMode = FPBDJointUtilities::GetLinearDriveAccelerationMode(SolverSettings, JointSettings);
 
 		// Rectangular position drives
 		if (bDriven[0] || bDriven[1] || bDriven[2])
@@ -1278,13 +1278,15 @@ void FPBDJointCachedSolver::InitPositionDrives(
 				if (bDriven[AxisIndex])
 				{
 					PositionDrives.InitDatas(AxisIndex, true, FPBDJointUtilities::GetLinearDriveStiffness(SolverSettings, JointSettings, AxisIndex),
-						FPBDJointUtilities::GetLinearDriveDamping(SolverSettings, JointSettings, AxisIndex), 0.0f);
+						FPBDJointUtilities::GetLinearDriveDamping(SolverSettings, JointSettings, AxisIndex), 0);
 					const FVec3 Axis = R0M.GetAxis(AxisIndex);
 				
 					if ((FMath::Abs(FVec3::DotProduct(CX,Axis)) > PositionTolerance) || (PositionDrives.ConstraintSoftDamping[AxisIndex] > 0.0f))
 					{
 						InitAxisPositionDrive(AxisIndex, Axis, CX, VTarget, Dt);
 					}
+
+					PositionDrives.SetMaxForce(AxisIndex, JointSettings.LinearDriveMaxForce[AxisIndex], Dt);
 				}
 			}
 		}
@@ -1553,15 +1555,19 @@ void FPBDJointCachedSolver::ApplyAxisPositionDrive(
 	FReal VelDt = 0;
 	if (PositionDrives.ConstraintSoftDamping[ConstraintIndex] > UE_KINDA_SMALL_NUMBER)
 	{
-		const FVec3 V0Dt = FVec3::CalculateVelocity(InitConnectorXs[0], ConnectorXs[0]+ Delta0, 1.0f);
-		const FVec3 V1Dt = FVec3::CalculateVelocity(InitConnectorXs[1], ConnectorXs[1]+ Delta1, 1.0f);
+		const FVec3 V0Dt = FVec3::CalculateVelocity(InitConnectorXs[0], ConnectorXs[0] + Delta0, 1.0f);
+		const FVec3 V1Dt = FVec3::CalculateVelocity(InitConnectorXs[1], ConnectorXs[1] + Delta1, 1.0f);
 		VelDt = PositionDrives.ConstraintVX[ConstraintIndex] * Dt + FVec3::DotProduct(V0Dt - V1Dt, PositionDrives.ConstraintAxis[ConstraintIndex] );
 	}
 
-	const FReal DeltaLambda = SolverStiffness * (PositionDrives.ConstraintSoftStiffness[ConstraintIndex] * DeltaPos -
+	FReal DeltaLambda = SolverStiffness * (PositionDrives.ConstraintSoftStiffness[ConstraintIndex] * DeltaPos -
 		PositionDrives.ConstraintSoftDamping[ConstraintIndex] * VelDt - PositionDrives.ConstraintLambda[ConstraintIndex]) /
 		PositionDrives.ConstraintSoftIM[ConstraintIndex];
-	PositionDrives.ConstraintLambda[ConstraintIndex] += DeltaLambda;
+	FReal Lambda = PositionDrives.ConstraintLambda[ConstraintIndex] + DeltaLambda;
+
+	PositionDrives.ApplyMaxLambda(ConstraintIndex, DeltaLambda, Lambda);
+
+	PositionDrives.ConstraintLambda[ConstraintIndex] = Lambda;
 
 	SolvePositionConstraintDelta(ConstraintIndex, DeltaLambda, PositionDrives);
 }
@@ -1659,13 +1665,14 @@ void FPBDJointCachedSolver::InitSwingTwistDrives(
 	RotationDrives.InitDatas(S2, true, FPBDJointUtilities::GetAngularSwingDriveStiffness(SolverSettings, JointSettings),
 		FPBDJointUtilities::GetAngularSwingDriveDamping(SolverSettings, JointSettings), 0.0);
 
-	RotationDrives.bAccelerationMode = FPBDJointUtilities::GetDriveAccelerationMode(SolverSettings, JointSettings);
+	RotationDrives.bAccelerationMode = FPBDJointUtilities::GetAngularDriveAccelerationMode(SolverSettings, JointSettings);
 
 	const bool bUseTwistDrive = bTwistDriveEnabled && (((FMath::Abs(DTwistAngle) > AngleTolerance) && (RotationDrives.ConstraintSoftStiffness[TW] > 0.0f)) || (RotationDrives.ConstraintSoftDamping[TW]  > 0.0f));
 	if (bUseTwistDrive)
 	{
 		InitRotationConstraintDrive(TW, ConnectorRs[1] * FJointConstants::TwistAxis(), Dt, DTwistAngle);
 		RotationDrives.ConstraintVX[TW] = JointSettings.AngularDriveVelocityTarget[TW];
+		RotationDrives.SetMaxForce(TW, JointSettings.AngularDriveMaxTorque[TW], Dt);
 	}
 
 	const bool bUseSwing1Drive = bSwing1DriveEnabled && (((FMath::Abs(DSwing1Angle) > AngleTolerance) && (RotationDrives.ConstraintSoftStiffness[S1] > 0.0f)) || (RotationDrives.ConstraintSoftDamping[S1] > 0.0f));
@@ -1673,6 +1680,7 @@ void FPBDJointCachedSolver::InitSwingTwistDrives(
 	{
 		InitRotationConstraintDrive(S1, ConnectorRs[1] * FJointConstants::Swing1Axis(),  Dt, DSwing1Angle);
 		RotationDrives.ConstraintVX[S1] = JointSettings.AngularDriveVelocityTarget[S1];
+		RotationDrives.SetMaxForce(S1, JointSettings.AngularDriveMaxTorque[S1], Dt);
 	}
 
 	const bool bUseSwing2Drive = bSwing2DriveEnabled && (((FMath::Abs(DSwing2Angle) > AngleTolerance) && (RotationDrives.ConstraintSoftStiffness[S2] > 0.0f)) || (RotationDrives.ConstraintSoftDamping[S2] > 0.0f));
@@ -1680,6 +1688,7 @@ void FPBDJointCachedSolver::InitSwingTwistDrives(
 	{
 		InitRotationConstraintDrive(S2, ConnectorRs[1] * FJointConstants::Swing2Axis(),  Dt, DSwing2Angle);
 		RotationDrives.ConstraintVX[S2] = JointSettings.AngularDriveVelocityTarget[S2];
+		RotationDrives.SetMaxForce(S2, JointSettings.AngularDriveMaxTorque[S2], Dt);
 	}
 }
 
@@ -1693,7 +1702,7 @@ void FPBDJointCachedSolver::InitSLerpDrive(
 		RotationDrives.InitDatas(AxisIndex, true, FPBDJointUtilities::GetAngularSLerpDriveStiffness(SolverSettings, JointSettings),
 						FPBDJointUtilities::GetAngularSLerpDriveDamping(SolverSettings, JointSettings), 0.0);
 	}
-	RotationDrives.bAccelerationMode = FPBDJointUtilities::GetDriveAccelerationMode(SolverSettings, JointSettings);
+	RotationDrives.bAccelerationMode = FPBDJointUtilities::GetAngularDriveAccelerationMode(SolverSettings, JointSettings);
 
 	const FRotation3 R01 = ConnectorRs[0].Inverse() * ConnectorRs[1];
 	FRotation3 TargetAngPos = JointSettings.AngularDrivePositionTarget;
@@ -1714,6 +1723,10 @@ void FPBDJointCachedSolver::InitSLerpDrive(
 	InitRotationConstraintDrive(1, Axes[1], Dt, AxisAngles[1]);
 	InitRotationConstraintDrive(2, Axes[2], Dt, AxisAngles[2]);
 
+	RotationDrives.SetMaxForce(0, JointSettings.AngularDriveMaxTorque[0], Dt);
+	RotationDrives.SetMaxForce(1, JointSettings.AngularDriveMaxTorque[1], Dt);
+	RotationDrives.SetMaxForce(2, JointSettings.AngularDriveMaxTorque[2], Dt);
+
 	// @todo(chaos): pass constraint target velocity into InitRotationConstraintDrive (it currently sets it ConstraintVX to 0)
 	if (!JointSettings.AngularDriveVelocityTarget.IsNearlyZero())
 	{
@@ -1754,10 +1767,14 @@ void FPBDJointCachedSolver::ApplyAxisRotationDrive(
 		AngVelDt = RotationDrives.ConstraintVX[ConstraintIndex] * Dt + FVec3::DotProduct(RotationDrives.ConstraintAxis[ConstraintIndex] , W0Dt - W1Dt);
 	}
 
-	const FReal DeltaLambda = SolverStiffness * (RotationDrives.ConstraintSoftStiffness[ConstraintIndex] * DeltaConstraint -
+	FReal DeltaLambda = SolverStiffness * (RotationDrives.ConstraintSoftStiffness[ConstraintIndex] * DeltaConstraint -
 		RotationDrives.ConstraintSoftDamping[ConstraintIndex] * AngVelDt - RotationDrives.ConstraintLambda[ConstraintIndex]) /
 		RotationDrives.ConstraintSoftIM[ConstraintIndex];
-	RotationDrives.ConstraintLambda[ConstraintIndex] += DeltaLambda;
+	FReal Lambda = RotationDrives.ConstraintLambda[ConstraintIndex] + DeltaLambda;
+
+	RotationDrives.ApplyMaxLambda(ConstraintIndex, DeltaLambda, Lambda);
+
+	RotationDrives.ConstraintLambda[ConstraintIndex] = Lambda;
 
 	SolveRotationConstraintDelta(ConstraintIndex, DeltaLambda,true, RotationDrives);
 }
@@ -1789,6 +1806,7 @@ void FAxisConstraintDatas::InitDatas(
 	ConstraintHardStiffness[ConstraintIndex] = HardStiffness;
 	ConstraintSoftStiffness[ConstraintIndex] = SoftStiffness;
 	ConstraintSoftDamping[ConstraintIndex] = SoftDamping;
+	ConstraintMaxLambda[ConstraintIndex] = 0;
 	SettingsSoftStiffness[ConstraintIndex] = SoftStiffness;
 	SettingsSoftDamping[ConstraintIndex] = SoftDamping;
 	bValidDatas[ConstraintIndex] = false;
@@ -1839,6 +1857,47 @@ void FAxisConstraintDatas::UpdateMass(
 		ConstraintSoftIM[ConstraintIndex] = (ConstraintSoftStiffness[ConstraintIndex] + ConstraintSoftDamping[ConstraintIndex]) * ConstraintHardIM[ConstraintIndex] + (FReal)1;
 	}
 }
+
+void FAxisConstraintDatas::SetMaxForce(
+	const int32 ConstraintIndex,
+	const FReal InMaxForce,
+	const FReal Dt)
+{
+	// We use 0 to disable max force clamping. See ApplyMaxLambda
+	ConstraintMaxLambda[ConstraintIndex] = 0;
+
+	if ((InMaxForce > 0) && (InMaxForce < UE_MAX_FLT))
+	{
+		// Convert from force/torque to position/angle impulse
+		FReal MaxLambda = InMaxForce * Dt * Dt;
+		if (bAccelerationMode)
+		{
+			MaxLambda /= ConstraintHardIM[ConstraintIndex];
+		}
+		ConstraintMaxLambda[ConstraintIndex] = MaxLambda;
+	}
+}
+
+void FAxisConstraintDatas::ApplyMaxLambda(
+	const int32 ConstraintIndex,
+	FReal& DeltaLambda,
+	FReal& Lambda)
+{
+	if (ConstraintMaxLambda[ConstraintIndex] > 0)
+	{
+		if (Lambda > ConstraintMaxLambda[ConstraintIndex])
+		{
+			DeltaLambda = ConstraintMaxLambda[ConstraintIndex] - ConstraintLambda[ConstraintIndex];
+			Lambda = ConstraintMaxLambda[ConstraintIndex];
+		}
+		else if (Lambda < -ConstraintMaxLambda[ConstraintIndex])
+		{
+			DeltaLambda = -ConstraintMaxLambda[ConstraintIndex] - ConstraintLambda[ConstraintIndex];
+			Lambda = -ConstraintMaxLambda[ConstraintIndex];
+		}
+	}
+}
+
 	
 }
 
diff --git a/Engine/Source/Runtime/Experimental/Chaos/Private/Chaos/Joint/PBDJointSolverGaussSeidel.cpp b/Engine/Source/Runtime/Experimental/Chaos/Private/Chaos/Joint/PBDJointSolverGaussSeidel.cpp
--- a/Engine/Source/Runtime/Experimental/Chaos/Private/Chaos/Joint/PBDJointSolverGaussSeidel.cpp
+++ b/Engine/Source/Runtime/Experimental/Chaos/Private/Chaos/Joint/PBDJointSolverGaussSeidel.cpp
@@ -1790,7 +1790,7 @@ namespace Chaos
 		const FReal AngularTwistDriveDamping = FPBDJointUtilities::GetAngularTwistDriveDamping(SolverSettings, JointSettings);
 		const FReal AngularSwingDriveStiffness = FPBDJointUtilities::GetAngularSwingDriveStiffness(SolverSettings, JointSettings);
 		const FReal AngularSwingDriveDamping = FPBDJointUtilities::GetAngularSwingDriveDamping(SolverSettings, JointSettings);
-		const bool bAccelerationMode = FPBDJointUtilities::GetDriveAccelerationMode(SolverSettings, JointSettings);
+		const bool bAccelerationMode = FPBDJointUtilities::GetAngularDriveAccelerationMode(SolverSettings, JointSettings);
 
 		const bool bUseTwistDrive = bTwistDriveEnabled && (((FMath::Abs(DTwistAngle) > AngleTolerance) && (AngularTwistDriveStiffness > 0.0f)) || (AngularTwistDriveDamping > 0.0f));
 		if (bUseTwistDrive)
@@ -1827,7 +1827,7 @@ namespace Chaos
 	{
 		const FReal AngularDriveStiffness = FPBDJointUtilities::GetAngularSLerpDriveStiffness(SolverSettings, JointSettings);
 		const FReal AngularDriveDamping = FPBDJointUtilities::GetAngularSLerpDriveDamping(SolverSettings, JointSettings);
-		const bool bAccelerationMode = FPBDJointUtilities::GetDriveAccelerationMode(SolverSettings, JointSettings);
+		const bool bAccelerationMode = FPBDJointUtilities::GetAngularDriveAccelerationMode(SolverSettings, JointSettings);
 
 		const FRotation3 R01 = ConnectorRs[0].Inverse() * ConnectorRs[1];
 		FRotation3 TargetAngPos = JointSettings.AngularDrivePositionTarget;
@@ -2094,7 +2094,7 @@ namespace Chaos
 	{
 		const FReal JointStiffness = FPBDJointUtilities::GetLinearDriveStiffness(SolverSettings, JointSettings, AxisIndex);
 		const FReal JointDamping = FPBDJointUtilities::GetLinearDriveDamping(SolverSettings, JointSettings, AxisIndex);
-		const bool bAccelerationMode = FPBDJointUtilities::GetDriveAccelerationMode(SolverSettings, JointSettings);
+		const bool bAccelerationMode = FPBDJointUtilities::GetLinearDriveAccelerationMode(SolverSettings, JointSettings);
 
 		if ((FMath::Abs(DeltaPos) > PositionTolerance) || (JointDamping > 0.0f))
 		{
diff --git a/Engine/Source/Runtime/Experimental/Chaos/Private/Chaos/PBDJointConstraintUtilities.cpp b/Engine/Source/Runtime/Experimental/Chaos/Private/Chaos/PBDJointConstraintUtilities.cpp
--- a/Engine/Source/Runtime/Experimental/Chaos/Private/Chaos/PBDJointConstraintUtilities.cpp
+++ b/Engine/Source/Runtime/Experimental/Chaos/Private/Chaos/PBDJointConstraintUtilities.cpp
@@ -594,7 +594,14 @@ namespace Chaos
 		return JointSettings.AngularSoftForceMode == EJointForceMode::Acceleration;
 	}
 
-	bool FPBDJointUtilities::GetDriveAccelerationMode(
+	bool FPBDJointUtilities::GetLinearDriveAccelerationMode(
+		const FPBDJointSolverSettings& SolverSettings,
+		const FPBDJointSettings& JointSettings)
+	{
+		return JointSettings.LinearDriveForceMode == EJointForceMode::Acceleration;
+	}
+
+	bool FPBDJointUtilities::GetAngularDriveAccelerationMode(
 		const FPBDJointSolverSettings& SolverSettings,
 		const FPBDJointSettings& JointSettings)
 	{
diff --git a/Engine/Source/Runtime/Experimental/Chaos/Private/Chaos/PBDJointConstraints.cpp b/Engine/Source/Runtime/Experimental/Chaos/Private/Chaos/PBDJointConstraints.cpp
--- a/Engine/Source/Runtime/Experimental/Chaos/Private/Chaos/PBDJointConstraints.cpp
+++ b/Engine/Source/Runtime/Experimental/Chaos/Private/Chaos/PBDJointConstraints.cpp
@@ -195,9 +195,9 @@ namespace Chaos
 		, bShockPropagationEnabled(false)
 		, bMassConditioningEnabled(true)
 		, LinearMotionTypes({ EJointMotionType::Locked, EJointMotionType::Locked, EJointMotionType::Locked })
-		, LinearLimit(FLT_MAX)
+		, LinearLimit(UE_MAX_FLT)
 		, AngularMotionTypes({ EJointMotionType::Free, EJointMotionType::Free, EJointMotionType::Free })
-		, AngularLimits(FVec3(FLT_MAX, FLT_MAX, FLT_MAX))
+		, AngularLimits(FVec3(UE_MAX_FLT, UE_MAX_FLT, UE_MAX_FLT))
 		, bSoftLinearLimitsEnabled(false)
 		, bSoftTwistLimitsEnabled(false)
 		, bSoftSwingLimitsEnabled(false)
@@ -222,6 +222,7 @@ namespace Chaos
 		, LinearDriveForceMode(EJointForceMode::Acceleration)
 		, LinearDriveStiffness(FVec3(0))
 		, LinearDriveDamping(FVec3(0))
+		, LinearDriveMaxForce(FVec3(UE_MAX_FLT))
 		, AngularDrivePositionTarget(FRotation3::FromIdentity())
 		, AngularDriveVelocityTarget(FVec3(0, 0, 0))
 		, bAngularSLerpPositionDriveEnabled(false)
@@ -233,12 +234,13 @@ namespace Chaos
 		, AngularDriveForceMode(EJointForceMode::Acceleration)
 		, AngularDriveStiffness(FVec3(0))
 		, AngularDriveDamping(FVec3(0))
-		, LinearBreakForce(FLT_MAX)
-		, LinearPlasticityLimit(FLT_MAX)
+		, AngularDriveMaxTorque(FVec3(UE_MAX_FLT))
+		, LinearBreakForce(UE_MAX_FLT)
+		, LinearPlasticityLimit(UE_MAX_FLT)
 		, LinearPlasticityType(EPlasticityType::Free)
-		, LinearPlasticityInitialDistanceSquared(FLT_MAX)
-		, AngularBreakTorque(FLT_MAX)
-		, AngularPlasticityLimit(FLT_MAX)
+		, LinearPlasticityInitialDistanceSquared(UE_MAX_FLT)
+		, AngularBreakTorque(UE_MAX_FLT)
+		, AngularPlasticityLimit(UE_MAX_FLT)
 		, ContactTransferScale(0.f)
 		, UserData(nullptr)
 	{
diff --git a/Engine/Source/Runtime/Experimental/Chaos/Public/Chaos/Joint/PBDJointCachedSolverGaussSeidel.h b/Engine/Source/Runtime/Experimental/Chaos/Public/Chaos/Joint/PBDJointCachedSolverGaussSeidel.h
--- a/Engine/Source/Runtime/Experimental/Chaos/Public/Chaos/Joint/PBDJointCachedSolverGaussSeidel.h
+++ b/Engine/Source/Runtime/Experimental/Chaos/Public/Chaos/Joint/PBDJointCachedSolverGaussSeidel.h
@@ -24,6 +24,12 @@ struct FAxisConstraintDatas
 		const FReal SoftDamping,
 		const FReal HardStiffness);
 
+	/** Initialize the max force ( force is modified for acceleration mode and converted to positional impulse limit) */
+	void SetMaxForce(
+		const int32 ConstraintIndex,
+		const FReal InMaxForce,
+		const FReal Dt);
+
 	/** Update the axis joint datas with axis, limits, arms... */
 	void UpdateDatas(
 		const int32 ConstraintIndex,
@@ -42,11 +48,18 @@ struct FAxisConstraintDatas
 		const FVec3& DatasIA1,
 		const FReal DatasIM,
 		const FReal Dt);
+
+	/** Apply the impulse limits to the impulse delta and net impulse */
+	void ApplyMaxLambda(
+		const int32 ConstraintIndex,
+		FReal& DeltaLambda, 
+		FReal& Lambda);
 	
 	FVec3 ConstraintHardStiffness;
 	FVec3 ConstraintSoftStiffness;
 	FVec3 ConstraintSoftDamping;
-	
+	FVec3 ConstraintMaxLambda;
+
 	FVec3 SettingsSoftDamping;
 	FVec3 SettingsSoftStiffness;
 	
diff --git a/Engine/Source/Runtime/Experimental/Chaos/Public/Chaos/JointProperties.inl b/Engine/Source/Runtime/Experimental/Chaos/Public/Chaos/JointProperties.inl
--- a/Engine/Source/Runtime/Experimental/Chaos/Public/Chaos/JointProperties.inl
+++ b/Engine/Source/Runtime/Experimental/Chaos/Public/Chaos/JointProperties.inl
@@ -32,6 +32,7 @@ CHAOS_INNER_JOINT_PROPERTY(JointSettings, LinearMotionTypesY, LinearMotionTypes[
 CHAOS_INNER_JOINT_PROPERTY(JointSettings, LinearMotionTypesZ, LinearMotionTypes[2], EJointMotionType)
 CHAOS_INNER_JOINT_PROPERTY(JointSettings, LinearDriveStiffness, LinearDriveStiffness, FVec3)
 CHAOS_INNER_JOINT_PROPERTY(JointSettings, LinearDriveDamping, LinearDriveDamping, FVec3)
+CHAOS_INNER_JOINT_PROPERTY(JointSettings, LinearDriveMaxForce, LinearDriveMaxForce, FVec3)
 CHAOS_INNER_JOINT_PROPERTY(JointSettings, ContactTransferScale, ContactTransferScale, FReal)
 CHAOS_INNER_JOINT_PROPERTY(JointSettings, AngularSLerpPositionDriveEnabled, bAngularSLerpPositionDriveEnabled, bool)
 CHAOS_INNER_JOINT_PROPERTY(JointSettings, AngularTwistPositionDriveEnabled, bAngularTwistPositionDriveEnabled, bool)
@@ -47,6 +48,7 @@ CHAOS_INNER_JOINT_PROPERTY(JointSettings, AngularMotionTypesY, AngularMotionType
 CHAOS_INNER_JOINT_PROPERTY(JointSettings, AngularMotionTypesZ, AngularMotionTypes[2], EJointMotionType)
 CHAOS_INNER_JOINT_PROPERTY(JointSettings, AngularDriveStiffness, AngularDriveStiffness, FVec3)
 CHAOS_INNER_JOINT_PROPERTY(JointSettings, AngularDriveDamping, AngularDriveDamping, FVec3)
+CHAOS_INNER_JOINT_PROPERTY(JointSettings, AngularDriveMaxTorque, AngularDriveMaxTorque, FVec3)
 CHAOS_INNER_JOINT_PROPERTY(JointSettings, Stiffness, Stiffness, FReal)
 CHAOS_INNER_JOINT_PROPERTY(JointSettings, SoftLinearLimitsEnabled, bSoftLinearLimitsEnabled, bool)
 CHAOS_INNER_JOINT_PROPERTY(JointSettings, SoftTwistLimitsEnabled, bSoftTwistLimitsEnabled, bool)
diff --git a/Engine/Source/Runtime/Experimental/Chaos/Public/Chaos/PBDJointConstraintTypes.h b/Engine/Source/Runtime/Experimental/Chaos/Public/Chaos/PBDJointConstraintTypes.h
--- a/Engine/Source/Runtime/Experimental/Chaos/Public/Chaos/PBDJointConstraintTypes.h
+++ b/Engine/Source/Runtime/Experimental/Chaos/Public/Chaos/PBDJointConstraintTypes.h
@@ -174,6 +174,7 @@ namespace Chaos
 		EJointForceMode LinearDriveForceMode;
 		FVec3 LinearDriveStiffness;
 		FVec3 LinearDriveDamping;
+		FVec3 LinearDriveMaxForce;
 
 		FRotation3 AngularDrivePositionTarget;
 		FVec3 AngularDriveVelocityTarget;
@@ -187,6 +188,7 @@ namespace Chaos
 		EJointForceMode AngularDriveForceMode;
 		FVec3 AngularDriveStiffness;
 		FVec3 AngularDriveDamping;
+		FVec3 AngularDriveMaxTorque;
 
 		FReal LinearBreakForce;
 		FReal LinearPlasticityLimit;
diff --git a/Engine/Source/Runtime/Experimental/Chaos/Public/Chaos/PBDJointConstraintUtilities.h b/Engine/Source/Runtime/Experimental/Chaos/Public/Chaos/PBDJointConstraintUtilities.h
--- a/Engine/Source/Runtime/Experimental/Chaos/Public/Chaos/PBDJointConstraintUtilities.h
+++ b/Engine/Source/Runtime/Experimental/Chaos/Public/Chaos/PBDJointConstraintUtilities.h
@@ -263,7 +263,11 @@ namespace Chaos
 			const FPBDJointSolverSettings& SolverSettings,
 			const FPBDJointSettings& JointSettings);
 
-		static bool GetDriveAccelerationMode(
+		static bool GetLinearDriveAccelerationMode(
+			const FPBDJointSolverSettings& SolverSettings,
+			const FPBDJointSettings& JointSettings);
+
+		static bool GetAngularDriveAccelerationMode(
 			const FPBDJointSolverSettings& SolverSettings,
 			const FPBDJointSettings& JointSettings);
 
