Chaos - Add character ground constraint

diff --git a/Engine/Source/Programs/HeadlessChaos/Private/Character/TestCharacterGroundConstraint.cpp b/Engine/Source/Programs/HeadlessChaos/Private/Character/TestCharacterGroundConstraint.cpp
new file mode 100644
--- /dev/null
+++ b/Engine/Source/Programs/HeadlessChaos/Private/Character/TestCharacterGroundConstraint.cpp
@@ -0,0 +1,125 @@
+// Copyright Epic Games, Inc. All Rights Reserved.
+
+#include "HeadlessChaos.h"
+#include "HeadlessChaosTestUtility.h"
+
+#include "Chaos/Character/CharacterGroundConstraint.h"
+
+namespace ChaosTest
+{
+	using namespace Chaos;
+
+	// Only function is a == operator so test that
+	TEST(CharacterGroundConstraintTests, TestSettingsAndData)
+	{
+		FCharacterGroundConstraintSettings Settings;
+		FCharacterGroundConstraintSettings SettingsCopy = Settings;
+		EXPECT_EQ(Settings, SettingsCopy);
+
+		Settings.RadialForceLimit = 2000.0f;
+		EXPECT_FALSE(Settings == SettingsCopy);
+
+		FCharacterGroundConstraintDynamicData Data;
+		FCharacterGroundConstraintDynamicData DataCopy = Data;
+		EXPECT_EQ(Data, DataCopy);
+
+		Data.GroundDistance = 11.0f;
+		EXPECT_FALSE(Data == DataCopy);
+	}
+
+	TEST(CharacterGroundConstraintTests, TestConstraintInitialization)
+	{
+		FCharacterGroundConstraint Constraint;
+
+		// Not set proxy so constraint should be invalid
+		EXPECT_FALSE(Constraint.IsValid());
+
+		// Type should be set
+		EXPECT_EQ(Constraint.GetType(), EConstraintType::CharacterGroundConstraintType);
+		EXPECT_TRUE(Constraint.IsType(EConstraintType::CharacterGroundConstraintType));
+
+		// Force and torque should be initialized to zero
+		EXPECT_VECTOR_FLOAT_EQ(Constraint.GetSolverAppliedForce(), FVector(0));
+		EXPECT_VECTOR_FLOAT_EQ(Constraint.GetSolverAppliedTorque(), FVector(0));
+	}
+
+	TEST(CharacterGroundConstraintTests, TestConstraintDirtyFlags)
+	{
+		FCharacterGroundConstraint Constraint;
+
+		EXPECT_FALSE(Constraint.IsDirty());
+
+		Constraint.SetGroundDistance(10.0);
+		EXPECT_EQ(Constraint.GetGroundDistance(), 10.0);
+
+		EXPECT_TRUE(Constraint.IsDirty());
+		EXPECT_TRUE(Constraint.IsDirty(EChaosPropertyFlags::CharacterGroundConstraintDynamicData));
+		EXPECT_FALSE(Constraint.IsDirty(EChaosPropertyFlags::CharacterGroundConstraintSettings));
+		EXPECT_FALSE(Constraint.IsDirty(EChaosPropertyFlags::CharacterParticleProxy));
+		EXPECT_FALSE(Constraint.IsDirty(EChaosPropertyFlags::GroundParticleProxy));
+
+		Constraint.ClearDirtyFlags();
+		EXPECT_FALSE(Constraint.IsDirty());
+
+		Constraint.SetTargetHeight(20.0);
+		EXPECT_EQ(Constraint.GetTargetHeight(), 20.0);
+
+		EXPECT_TRUE(Constraint.IsDirty());
+		EXPECT_FALSE(Constraint.IsDirty(EChaosPropertyFlags::CharacterGroundConstraintDynamicData));
+		EXPECT_TRUE(Constraint.IsDirty(EChaosPropertyFlags::CharacterGroundConstraintSettings));
+		EXPECT_FALSE(Constraint.IsDirty(EChaosPropertyFlags::CharacterParticleProxy));
+		EXPECT_FALSE(Constraint.IsDirty(EChaosPropertyFlags::GroundParticleProxy));
+
+		Constraint.ClearDirtyFlags();
+
+		FSingleParticlePhysicsProxy* CharacterProxy = FSingleParticlePhysicsProxy::Create(FGeometryParticle::CreateParticle());
+		FSingleParticlePhysicsProxy* GroundProxy = FSingleParticlePhysicsProxy::Create(FGeometryParticle::CreateParticle());
+
+		Constraint.Init(CharacterProxy);
+		EXPECT_EQ(Constraint.GetCharacterParticleProxy(), CharacterProxy);
+
+		EXPECT_TRUE(Constraint.IsDirty(EChaosPropertyFlags::CharacterGroundConstraintDynamicData));
+		EXPECT_TRUE(Constraint.IsDirty(EChaosPropertyFlags::CharacterGroundConstraintSettings));
+		EXPECT_TRUE(Constraint.IsDirty(EChaosPropertyFlags::CharacterParticleProxy));
+		EXPECT_FALSE(Constraint.IsDirty(EChaosPropertyFlags::GroundParticleProxy));
+
+		Constraint.SetGroundParticleProxy(GroundProxy);
+		EXPECT_EQ(Constraint.GetGroundParticleProxy(), GroundProxy);
+		EXPECT_TRUE(Constraint.IsDirty(EChaosPropertyFlags::GroundParticleProxy));
+
+		Constraint.ClearDirtyFlags();
+
+		FDirtyProxiesBucketInfo BucketInfo;
+		BucketInfo.Num[(uint32)EPhysicsProxyType::CharacterGroundConstraintType] = 1;
+		FDirtyPropertiesManager Manager;
+		Manager.PrepareBuckets(BucketInfo);
+
+		FDirtyChaosProperties RemoteData;
+
+		EXPECT_FALSE(RemoteData.HasCharacterGroundConstraintSettings());
+
+		Constraint.SetTargetHeight(15.0);
+		Constraint.SyncRemoteData(Manager, 0, RemoteData);
+
+		ASSERT_TRUE(RemoteData.HasCharacterGroundConstraintSettings());
+		EXPECT_FALSE(RemoteData.HasCharacterGroundConstraintDynamicData());
+		const FCharacterGroundConstraintSettings& Settings = RemoteData.GetCharacterGroundConstraintSettings(Manager, 0);
+
+		EXPECT_EQ(Settings.TargetHeight, 15.0);
+
+		Constraint.ClearDirtyFlags();
+		RemoteData.Clear(Manager, 0);
+
+		EXPECT_FALSE(RemoteData.HasCharacterGroundConstraintDynamicData());
+		EXPECT_FALSE(RemoteData.HasCharacterGroundConstraintSettings());
+
+		Constraint.SetGroundDistance(34.5);
+		Constraint.SyncRemoteData(Manager, 0, RemoteData);
+
+		ASSERT_TRUE(RemoteData.HasCharacterGroundConstraintDynamicData());
+		EXPECT_FALSE(RemoteData.HasCharacterGroundConstraintSettings());
+		const FCharacterGroundConstraintDynamicData& Data = RemoteData.GetCharacterGroundConstraintDynamicData(Manager, 0);
+
+		EXPECT_EQ(Data.GroundDistance, 34.5);
+	}
+} // namespace ChaosTest
\ No newline at end of file
diff --git a/Engine/Source/Programs/HeadlessChaos/Private/Character/TestCharacterGroundConstraintContainer.cpp b/Engine/Source/Programs/HeadlessChaos/Private/Character/TestCharacterGroundConstraintContainer.cpp
new file mode 100644
--- /dev/null
+++ b/Engine/Source/Programs/HeadlessChaos/Private/Character/TestCharacterGroundConstraintContainer.cpp
@@ -0,0 +1,370 @@
+// Copyright Epic Games, Inc. All Rights Reserved.
+
+#include "HeadlessChaos.h"
+#include "HeadlessChaosTestUtility.h"
+
+#include "Chaos/Character/CharacterGroundConstraintContainer.h"
+#include "Chaos/Island/IslandManager.h"
+#include "Chaos/PBDRigidsSOAs.h"
+
+namespace ChaosTest
+{
+	using namespace Chaos;
+
+	class CharacterGroundConstraintContainerTest : public ::testing::Test
+	{
+	protected:
+		CharacterGroundConstraintContainerTest()
+			: SOAs(UniqueIndices)
+		{
+		}
+
+		void SetSolverForce(FCharacterGroundConstraintHandle* Constraint, const FVec3& Value)
+		{
+			Constraint->SolverAppliedForce = Value;
+		}
+
+		FCharacterGroundConstraintContainer Container;
+		TArray<FCharacterGroundConstraintSettings> Settings;
+		TArray<FCharacterGroundConstraintDynamicData> Data;
+		FParticleUniqueIndicesMultithreaded UniqueIndices;
+		FPBDRigidsSOAs SOAs;
+	};
+
+	TEST_F(CharacterGroundConstraintContainerTest, TestInitialization)
+	{
+		EXPECT_EQ(Container.NumConstraints(), 0);
+		EXPECT_EQ(Container.GetConstraints().Num(), 0);
+
+		TArray<FPBDRigidParticleHandle*> Particles = SOAs.CreateDynamicParticles(2);
+		Settings.SetNum(1);
+		Data.SetNum(1);
+
+		Container.AddConstraint(Settings[0], Data[0], Particles[0], Particles[1]);
+
+		FCharacterGroundConstraintHandle* Constraint = Container.GetConstraint(0);
+		ASSERT_NE(Constraint, nullptr);
+		EXPECT_VECTOR_FLOAT_EQ(Constraint->GetSolverAppliedForce(), FVec3::ZeroVector);
+		EXPECT_VECTOR_FLOAT_EQ(Constraint->GetSolverAppliedTorque(), FVec3::ZeroVector);
+		EXPECT_EQ(Constraint->GetSettings(), Settings[0]);
+		EXPECT_EQ(Constraint->GetData(), Data[0]);
+		EXPECT_TRUE(Constraint->IsEnabled());
+		EXPECT_FALSE(Constraint->IsSleeping());
+		EXPECT_TRUE(Constraint->IsValid());
+
+		Container.RemoveConstraint(Constraint);
+
+		Particles[0]->SetDisabled(true);
+		Container.AddConstraint(Settings[0], Data[0], Particles[0], Particles[1]);
+		Constraint = Container.GetConstraint(0);
+		EXPECT_FALSE(Constraint->IsEnabled());
+	}
+
+	TEST_F(CharacterGroundConstraintContainerTest, TestAddRemoveConstraints)
+	{
+		const int NumParticles = 6;
+		const int NumConstraints = 4;
+		TArray<FPBDRigidParticleHandle*> Particles = SOAs.CreateDynamicParticles(NumParticles);
+		Settings.SetNum(NumConstraints);
+		Data.SetNum(NumConstraints);
+
+		FCharacterGroundConstraintHandle* Constraints[4];
+
+		Constraints[0] = Container.AddConstraint(Settings[0], Data[0], Particles[0], Particles[1]);
+		Constraints[1] = Container.AddConstraint(Settings[1], Data[1], Particles[2]);
+		Constraints[2] = Container.AddConstraint(Settings[2], Data[2], Particles[3], Particles[1]);
+		Constraints[3] = Container.AddConstraint(Settings[3], Data[3], Particles[4], Particles[5]);
+
+		ASSERT_EQ(Container.NumConstraints(), NumConstraints);
+		EXPECT_EQ(Container.GetNumConstraints(), NumConstraints);
+
+		for (int Idx = 0; Idx < NumConstraints; ++Idx)
+		{
+			FCharacterGroundConstraintHandle* Constraint = Container.GetConstraint(Idx);
+			ASSERT_NE(Constraint, nullptr);
+		
+			EXPECT_EQ(Constraint->GetSettings(), Settings[Idx]);
+			EXPECT_EQ(Constraint->GetData(), Data[Idx]);
+		}
+
+		auto ContainerConstraints = Container.GetConstConstraints();
+		ASSERT_EQ(ContainerConstraints.Num(), NumConstraints);
+		for (int Idx = 0; Idx < NumConstraints; ++Idx)
+		{
+			EXPECT_EQ(ContainerConstraints[Idx], Constraints[Idx]);
+		}
+
+		Container.RemoveConstraint(Constraints[0]);
+		Container.RemoveConstraint(Constraints[2]);
+		ContainerConstraints = Container.GetConstConstraints();
+		ASSERT_EQ(ContainerConstraints.Num(), 2);
+		// Doesn't matter what order they're in but the two left should be 1 and 3
+		if (ContainerConstraints[0] == Constraints[1])
+		{
+			EXPECT_EQ(ContainerConstraints[1], Constraints[3]);
+		}
+		else if (ContainerConstraints[0] == Constraints[3])
+		{
+			EXPECT_EQ(ContainerConstraints[1], Constraints[1]);
+		}
+		else
+		{
+			EXPECT_TRUE(false);
+		}
+	}
+
+	TEST_F(CharacterGroundConstraintContainerTest, TestDisconnectConstraints)
+	{
+		const int NumParticles = 6;
+		const int NumConstraints = 5;
+		TArray<FPBDRigidParticleHandle*> Particles = SOAs.CreateDynamicParticles(NumParticles);
+		Settings.SetNum(NumConstraints);
+		Data.SetNum(NumConstraints);
+
+		TArray<FCharacterGroundConstraintHandle*> Constraints;
+		Constraints.Add(Container.AddConstraint(Settings[0], Data[0], Particles[0], Particles[1]));
+		Constraints.Add(Container.AddConstraint(Settings[1], Data[1], Particles[2]));
+		Constraints.Add(Container.AddConstraint(Settings[2], Data[2], Particles[3], Particles[1]));
+		Constraints.Add(Container.AddConstraint(Settings[3], Data[3], Particles[4], Particles[5]));
+		Constraints.Add(Container.AddConstraint(Settings[4], Data[4], Particles[5]));
+
+		for (auto Constraint : Constraints)
+		{
+			EXPECT_TRUE(Constraint->IsEnabled());
+		}
+
+		// Remove particle 5
+		// Constraint 4 should be disabled as particle 5 is its character particle
+		// Constraint 3 should have the ground particle set to null
+		Container.DisconnectConstraints({ Particles[5] });
+		EXPECT_FALSE(Constraints[4]->IsEnabled());
+		EXPECT_EQ(Constraints[3]->GetGroundParticle(), nullptr);
+	}
+
+	TEST_F(CharacterGroundConstraintContainerTest, TestAddToGraph)
+	{
+		const int NumDynamicParticles = 15;
+		const int NumStaticParticles = 2;
+
+		TArray<FPBDRigidParticleHandle*> DynamicParticles = SOAs.CreateDynamicParticles(NumDynamicParticles);
+		TArray<FGeometryParticleHandle*> StaticParticles = SOAs.CreateStaticParticles(NumStaticParticles);
+
+		TArray<TVec2<FGeometryParticleHandle*>> ParticlePairs =
+		{
+			// Two character particles with the same dynamic ground particle
+			{ DynamicParticles[0], DynamicParticles[9] },
+			{ DynamicParticles[1], DynamicParticles[9] },
+			// One character particle with no ground particle
+			{ DynamicParticles[2], nullptr },
+			// One character particle with single dynamic ground particle
+			{ DynamicParticles[3], DynamicParticles[10] },
+			// One character particle with two constraints to different dynamic bodies
+			{ DynamicParticles[4], DynamicParticles[11] },
+			{ DynamicParticles[4], DynamicParticles[12] },
+			// One character particle with two constraints to the same dynamic bodies
+			{ DynamicParticles[5], DynamicParticles[13] },
+			{ DynamicParticles[5], DynamicParticles[13] },
+			// One character particle with two constraints to a dynamic and static body
+			{ DynamicParticles[6], DynamicParticles[14] },
+			{ DynamicParticles[6], StaticParticles[0] },
+			// Two character particle with constraints to the same static body
+			{ DynamicParticles[7], StaticParticles[1] },
+			{ DynamicParticles[8], StaticParticles[1] },
+		};
+
+		const int NumConstraints = ParticlePairs.Num();
+		Settings.SetNum(NumConstraints);
+		Data.SetNum(NumConstraints);
+		for (int Idx = 0; Idx < NumConstraints; ++Idx)
+		{
+			Container.AddConstraint(Settings[Idx], Data[Idx], ParticlePairs[Idx][0], ParticlePairs[Idx][1]);
+		}
+
+		// Initialize graph
+		Private::FPBDIslandManager Graph;
+		Container.SetContainerId(0); // Usually set by the evolution the container is registered with
+		Graph.AddConstraintContainer(Container);
+
+		for (auto Particle : DynamicParticles)
+		{
+			Graph.AddParticle(Particle);
+		}
+		for (auto Particle : StaticParticles)
+		{
+			Graph.AddParticle(Particle);
+		}
+		Graph.InitializeGraph(SOAs.GetNonDisabledDynamicView());
+
+		// Add constraints and update graph
+		Container.AddConstraintsToGraph(Graph);
+		Graph.UpdateIslands(SOAs);
+
+		auto Constraints = Container.GetConstraints();
+		for (FCharacterGroundConstraintHandle* Constraint : Constraints)
+		{
+			EXPECT_TRUE(Constraint->IsInConstraintGraph());
+		}
+
+		TArray<TSet<FGeometryParticleHandle*>> ExpectedParticlesInIsland =
+		{
+			{ DynamicParticles[0], DynamicParticles[1], DynamicParticles[9] },
+			{ DynamicParticles[2] },
+			{ DynamicParticles[3], DynamicParticles[10] },
+			{ DynamicParticles[4], DynamicParticles[11], DynamicParticles[12] },
+			{ DynamicParticles[5], DynamicParticles[13] },
+			{ DynamicParticles[6], DynamicParticles[14] },
+			{ DynamicParticles[7] },
+			{ DynamicParticles[8] },
+		};
+		TArray<TSet<FCharacterGroundConstraintHandle*>> ExpectedConstraintsInIsland =
+		{
+			{ Constraints[0], Constraints[1] },
+			{ Constraints[2] },
+			{ Constraints[3] },
+			{ Constraints[4], Constraints[5] },
+			{ Constraints[6], Constraints[7] },
+			{ Constraints[8], Constraints[9] },
+			{ Constraints[10] },
+			{ Constraints[11] },
+		};
+		const int ExpectedNumIslands = ExpectedParticlesInIsland.Num();
+		TArray<int32> IslandIndices =
+		{
+			DynamicParticles[0]->CastToRigidParticle()->IslandIndex(),
+			DynamicParticles[2]->CastToRigidParticle()->IslandIndex(),
+			DynamicParticles[3]->CastToRigidParticle()->IslandIndex(),
+			DynamicParticles[4]->CastToRigidParticle()->IslandIndex(),
+			DynamicParticles[5]->CastToRigidParticle()->IslandIndex(),
+			DynamicParticles[6]->CastToRigidParticle()->IslandIndex(),
+			DynamicParticles[7]->CastToRigidParticle()->IslandIndex(),
+			DynamicParticles[8]->CastToRigidParticle()->IslandIndex(),
+		};
+
+		for (int32 IslandIdx = 0; IslandIdx < ExpectedNumIslands; ++IslandIdx)
+		{
+			Private::FPBDIsland* Island = Graph.GetIsland(IslandIdx);
+			TArrayView<Private::FPBDIslandParticle> IslandParticles = Island->GetParticles();
+			EXPECT_EQ(IslandParticles.Num(), ExpectedParticlesInIsland[IslandIdx].Num());
+			for (Private::FPBDIslandParticle& IslandParticle : IslandParticles)
+			{
+				EXPECT_TRUE(ExpectedParticlesInIsland[IslandIdx].Contains(IslandParticle.GetParticle()));
+			}
+
+			TArrayView<Private::FPBDIslandConstraint> IslandConstraints = Island->GetConstraints(Container.GetContainerId());
+			EXPECT_EQ(IslandConstraints.Num(), ExpectedConstraintsInIsland[IslandIdx].Num());
+			for (Private::FPBDIslandConstraint& IslandConstraint : IslandConstraints)
+			{
+				EXPECT_TRUE(ExpectedConstraintsInIsland[IslandIdx].Contains(IslandConstraint.GetConstraint()->As<FCharacterGroundConstraintHandle>()));
+			}
+		}
+
+		Graph.Reset();
+	}
+
+	TEST_F(CharacterGroundConstraintContainerTest, TestContainerSolver)
+	{
+		// Create solver
+		int Priority = 2;
+		TUniquePtr<FConstraintContainerSolver> Solver = Container.CreateSceneSolver(Priority);
+		EXPECT_EQ(Solver->GetPriority(), Priority);
+
+		// Create constraints in two islands
+		TArray<FPBDRigidParticleHandle*> Particles = SOAs.CreateDynamicParticles(4);
+		Settings.SetNum(1);
+		Data.SetNum(1);
+
+		// Initialize the data so that each constraint has work to do
+		Particles[0]->SetX(FVec3(0.0, 0.0, 10.0));
+		Particles[1]->SetX(FVec3(0.0, 0.0, 10.0));
+		Particles[2]->SetX(FVec3(0.0, 0.0, 10.0));
+		Particles[0]->SetP(FVec3(0.0, 0.0, 10.0));
+		Particles[1]->SetP(FVec3(0.0, 0.0, 10.0));
+		Particles[2]->SetP(FVec3(0.0, 0.0, 10.0));
+		Settings[0].TargetHeight = 20.0f;
+		Data[0].GroundDistance = 10.0f;
+
+		TArray<FCharacterGroundConstraintHandle*> Constraints;
+		Constraints.Add(Container.AddConstraint(Settings[0], Data[0], Particles[0], Particles[3]));
+		Constraints.Add(Container.AddConstraint(Settings[0], Data[0], Particles[1], Particles[3]));
+		Constraints.Add(Container.AddConstraint(Settings[0], Data[0], Particles[2]));
+
+		Private::FPBDIslandManager Graph;
+		Container.SetContainerId(0);
+		Graph.AddConstraintContainer(Container);
+		for (auto Particle : Particles)
+		{
+			Graph.AddParticle(Particle);
+		}
+		Graph.InitializeGraph(SOAs.GetNonDisabledDynamicView());
+		Container.AddConstraintsToGraph(Graph);
+		Graph.UpdateIslands(SOAs);
+
+		// First use non-island version
+		Solver->AddConstraints();
+		EXPECT_EQ(Solver->GetNumConstraints(), Constraints.Num());
+
+		FSolverBodyContainer SolverBodies;
+		SolverBodies.Reset(Particles.Num());
+
+		Solver->AddBodies(SolverBodies);
+		EXPECT_EQ(SolverBodies.Num(), Particles.Num());
+
+		const FReal Dt = 1.0 / 30.0;
+		SolverBodies.GatherInput(Dt, 0, SolverBodies.Num());
+		Solver->GatherInput(Dt);
+
+		for (FCharacterGroundConstraintHandle* Constraint : Constraints)
+		{
+			EXPECT_VECTOR_FLOAT_EQ(Constraint->GetSolverAppliedForce(), FVec3::ZeroVector);
+		}
+
+		Solver->ApplyPositionConstraints(Dt, 0, 1);
+		Solver->ScatterOutput(Dt);
+
+		for (FCharacterGroundConstraintHandle* Constraint : Constraints)
+		{
+			EXPECT_NE(Constraint->GetSolverAppliedForce(), FVec3::ZeroVector);
+		}
+
+		// Reset the constraint forces to something else
+		for (FCharacterGroundConstraintHandle* Constraint : Constraints)
+		{
+			SetSolverForce(Constraint, FVec3(1, 1, 1));
+		}
+
+		// Now use island API
+		Solver->Reset(Constraints.Num());
+
+		for (auto Particle : Particles)
+		{
+			Particle->SetSolverBodyIndex(INDEX_NONE);
+		}
+
+		int IslandIdx = 0;
+		Private::FPBDIsland* Island = Graph.GetIsland(IslandIdx);
+		Solver->AddConstraints(Island->GetConstraints(0));
+		EXPECT_EQ(Solver->GetNumConstraints(), 2);
+
+		SolverBodies.Reset(Particles.Num());
+
+		Solver->AddBodies(SolverBodies);
+		EXPECT_EQ(SolverBodies.Num(), 3);
+
+
+		SolverBodies.GatherInput(Dt, 0, SolverBodies.Num());
+		Solver->GatherInput(Dt);
+		Solver->ApplyPositionConstraints(Dt, 0, 1);
+		Solver->ScatterOutput(Dt);
+
+		// Expect the constraints in the island to be updated, but not
+		// the constraint in the other island
+		EXPECT_NE(Constraints[0]->GetSolverAppliedForce(), FVec3::ZeroVector);
+		EXPECT_NE(Constraints[0]->GetSolverAppliedForce(), FVec3(1, 1, 1));
+		EXPECT_NE(Constraints[1]->GetSolverAppliedForce(), FVec3::ZeroVector);
+		EXPECT_NE(Constraints[1]->GetSolverAppliedForce(), FVec3(1, 1, 1));
+
+		EXPECT_VECTOR_FLOAT_EQ(Constraints[2]->GetSolverAppliedForce(), FVec3(1, 1, 1));
+
+		Graph.Reset();
+	}
+}
\ No newline at end of file
diff --git a/Engine/Source/Programs/HeadlessChaos/Private/Character/TestCharacterGroundConstraintProxy.cpp b/Engine/Source/Programs/HeadlessChaos/Private/Character/TestCharacterGroundConstraintProxy.cpp
new file mode 100644
--- /dev/null
+++ b/Engine/Source/Programs/HeadlessChaos/Private/Character/TestCharacterGroundConstraintProxy.cpp
@@ -0,0 +1,113 @@
+// Copyright Epic Games, Inc. All Rights Reserved.
+
+#include "HeadlessChaos.h"
+#include "HeadlessChaosTestUtility.h"
+
+#include "Chaos/Character/CharacterGroundConstraint.h"
+#include "Chaos/Character/CharacterGroundConstraintContainer.h"
+#include "Chaos/PBDRigidsEvolutionGBF.h"
+#include "PBDRigidsSolver.h"
+#include "PhysicsProxy/CharacterGroundConstraintProxy.h"
+
+namespace ChaosTest
+{
+	using namespace Chaos;
+
+	class CharacterGroundConstraintsProxyTest : public ::testing::Test
+	{
+	protected:
+		void SetUp() override
+		{
+			FChaosSolversModule* Module = FChaosSolversModule::GetModule();
+			Solver = Module->CreateSolver(/*Owner*/nullptr, /*AsyncDt=*/-1);
+			Solver->SetThreadingMode_External(EThreadingModeTemp::SingleThread);
+		}
+
+		void TearDown() override
+		{
+			FChaosSolversModule* Module = FChaosSolversModule::GetModule();
+			Module->DestroySolver(Solver);
+		}
+
+		FSingleParticlePhysicsProxy* CreateParticle()
+		{
+			auto Geom = TSharedPtr<FImplicitObject, ESPMode::ThreadSafe>(new TBox<FReal, 3>(FVec3(-100), FVec3(100)));
+			FSingleParticlePhysicsProxy* Proxy = FSingleParticlePhysicsProxy::Create(FPBDRigidParticle::CreateParticle());
+			Proxy->SetHandle(Solver->GetEvolution()->CreateDynamicParticles(1)[0]);
+			Proxy->GetGameThreadAPI().SetGeometry(Geom);
+			Proxy->GetGameThreadAPI().SetM(1.0);
+			Proxy->GetGameThreadAPI().SetInvI(FVec3(1.0, 1.0, 1.0));
+			return Proxy;
+		}
+
+		FPBDRigidsSolver* Solver;
+	};
+
+	// This tests initialization of a character ground constraint and pushing
+	// of data from the game thread representation to the physics thread representation
+	TEST_F(CharacterGroundConstraintsProxyTest, TestProxyDataSyncing)
+	{
+		// Create a game thread constraint and register it with the solver
+		FCharacterGroundConstraint* Constraint_GT = new FCharacterGroundConstraint();
+		FSingleParticlePhysicsProxy* CharacterProxy = CreateParticle();
+		FSingleParticlePhysicsProxy* GroundProxy = CreateParticle();
+		Constraint_GT->Init(CharacterProxy);
+		Constraint_GT->SetGroundParticleProxy(GroundProxy);
+
+		Solver->RegisterObject(CharacterProxy);
+		Solver->RegisterObject(GroundProxy);
+		Solver->RegisterObject(Constraint_GT);
+
+		// Get the proxy from the game thread constraint
+		// Should have been set when registering the constraint with the solver
+		FCharacterGroundConstraintProxy* Proxy = Constraint_GT->GetProxy<FCharacterGroundConstraintProxy>();
+		ASSERT_TRUE(Proxy != nullptr);
+
+		EXPECT_FALSE(Proxy->IsInitialized());
+		EXPECT_EQ(Proxy->GetGameThreadAPI(), Constraint_GT);
+		EXPECT_EQ(Proxy->GetPhysicsThreadAPI(), nullptr);
+
+		// Create a local dirty properties manager and properties to use to transfer data from
+		// the game thread constraint to the physics thread constraint and initialize with the
+		// game thread constraint data
+
+		const int32 DataIdx = 0;
+		FDirtyProxiesBucketInfo BucketInfo;
+		BucketInfo.Num[(uint32)EPhysicsProxyType::CharacterGroundConstraintType] = 1;
+		FDirtyPropertiesManager Manager;
+		Manager.PrepareBuckets(BucketInfo);
+
+		FDirtyChaosProperties RemoteData;
+
+		Proxy->PushStateOnGameThread(Manager, 0, RemoteData);
+
+		// Initialize the physics thread data on the proxy
+		Proxy->InitializeOnPhysicsThread(Solver, Manager, DataIdx, RemoteData);
+
+		FCharacterGroundConstraintHandle* Constraint_PT = Proxy->GetPhysicsThreadAPI();
+		EXPECT_TRUE(Constraint_PT != nullptr);
+		EXPECT_EQ(Constraint_PT->GetCharacterParticle(), CharacterProxy->GetHandle_LowLevel());
+		EXPECT_EQ(Constraint_PT->GetGroundParticle(), GroundProxy->GetHandle_LowLevel());
+
+		EXPECT_EQ(Solver->GetCharacterGroundConstraints().GetNumConstraints(), 1);
+
+		// Change something on the game side and check that it is synced on the physics side
+		RemoteData.Clear(Manager, DataIdx);
+
+		Constraint_GT->SetGroundDistance(1234.5f);
+
+		Proxy->PushStateOnGameThread(Manager, DataIdx, RemoteData);
+		Proxy->PushStateOnPhysicsThread(Solver, Manager, DataIdx, RemoteData);
+
+		EXPECT_FLOAT_EQ(Constraint_PT->GetData().GroundDistance, 1234.5f);
+
+		// Remove the constraint and particles from the solver
+		Solver->UnregisterObject(Constraint_GT);
+		Solver->UnregisterObject(CharacterProxy);
+		Solver->UnregisterObject(GroundProxy);
+
+		EXPECT_EQ(Proxy->GetGameThreadAPI(), nullptr);
+
+		// Physics thread constraint is deleted in a callback set on the marshalling manager
+	}
+}
\ No newline at end of file
diff --git a/Engine/Source/Runtime/Experimental/Chaos/Private/Chaos/ChaosDebugDraw.cpp b/Engine/Source/Runtime/Experimental/Chaos/Private/Chaos/ChaosDebugDraw.cpp
--- a/Engine/Source/Runtime/Experimental/Chaos/Private/Chaos/ChaosDebugDraw.cpp
+++ b/Engine/Source/Runtime/Experimental/Chaos/Private/Chaos/ChaosDebugDraw.cpp
@@ -3,6 +3,7 @@
 #include "Chaos/Box.h"
 #include "Chaos/Capsule.h"
 #include "Chaos/CCDUtilities.h"
+#include "Chaos/Character/CharacterGroundConstraintContainer.h"
 #include "Chaos/Collision/CollisionConstraintAllocator.h"
 #include "Chaos/Collision/ParticlePairMidPhase.h"
 #include "Chaos/Convex.h"
@@ -1250,6 +1251,32 @@ namespace Chaos
 			}
 		}
 
+		void DrawCharacterGroundConstraintImpl(const FRigidTransform3& SpaceTransform, const FCharacterGroundConstraintHandle* Constraint, const FChaosDebugDrawSettings& Settings)
+		{
+			if (!Constraint->IsEnabled())
+			{
+				return;
+			}
+
+			const FRealSingle GroundDistance = FRealSingle(Constraint->GetData().GroundDistance);
+			const FRealSingle TargetHeight = FRealSingle(Constraint->GetSettings().TargetHeight);
+
+			if (GroundDistance < TargetHeight)
+			{
+				const FVec3 Pa = FParticleUtilities::GetActorWorldTransform(FConstGenericParticleHandle(Constraint->GetCharacterParticle())).GetTranslation();
+				const FVec3 Start = SpaceTransform.TransformPosition(Pa);
+				const FVec3 Up = SpaceTransform.TransformVector(Constraint->GetSettings().VerticalAxis);
+
+				const FRealSingle ArrowLength = Settings.DrawScale * Settings.BodyAxisLen;
+				const FRealSingle ArrowSize = Settings.DrawScale * Settings.ArrowSize;
+
+				Chaos::FDebugDrawQueue& DrawQueue = FDebugDrawQueue::GetInstance();
+				DrawQueue.DrawDebugLine(Pa, Pa - GroundDistance * Up, FColor::Emerald);
+				DrawQueue.DrawDebugLine(Pa - GroundDistance * Up, Pa - TargetHeight * Up, FColor::Red);
+				DrawQueue.DrawDebugDirectionalArrow(Pa - GroundDistance * Up, Pa - (GroundDistance - ArrowLength) * Up, ArrowSize, FColor::Cyan);
+			}
+		}
+
 		void DrawSimulationSpaceImpl(const FSimulationSpace& SimSpace, const FChaosDebugDrawSettings& Settings)
 		{
 			const FVec3 Pos = SimSpace.Transform.GetLocation();
@@ -1802,6 +1829,17 @@ namespace Chaos
 			}
 		}
 
+		void DrawCharacterGroundConstraints(const FRigidTransform3& SpaceTransform, const FCharacterGroundConstraintContainer& Constraints, const FChaosDebugDrawSettings* Settings)
+		{
+			if (FDebugDrawQueue::IsDebugDrawingEnabled())
+			{
+				for (int32 ConstraintIndex = 0; ConstraintIndex < Constraints.NumConstraints(); ++ConstraintIndex)
+				{
+					DrawCharacterGroundConstraintImpl(SpaceTransform, Constraints.GetConstraint(ConstraintIndex), GetChaosDebugDrawSettings(Settings));
+				}
+			}
+		}
+
 		void DrawCCDCollisionImpulse(const FRigidTransform3& SpaceTransform, const FCCDConstraint& CCDConstraint, const int32 ManifoldPointIndex, const FVec3& Impulse, const FChaosDebugDrawSettings* Settings)
 		{
 			if (FDebugDrawQueue::IsDebugDrawingEnabled())
diff --git a/Engine/Source/Runtime/Experimental/Chaos/Private/Chaos/Character/CharacterGroundConstraint.cpp b/Engine/Source/Runtime/Experimental/Chaos/Private/Chaos/Character/CharacterGroundConstraint.cpp
new file mode 100644
--- /dev/null
+++ b/Engine/Source/Runtime/Experimental/Chaos/Private/Chaos/Character/CharacterGroundConstraint.cpp
@@ -0,0 +1,37 @@
+// Copyright Epic Games, Inc. All Rights Reserved.
+
+#include "Chaos/Character/CharacterGroundConstraint.h"
+
+namespace Chaos
+{
+	FCharacterGroundConstraint::FCharacterGroundConstraint()
+		: FConstraintBase(EConstraintType::CharacterGroundConstraintType)
+		, SolverAppliedForce(FVector(0))
+		, SolverAppliedTorque(FVector(0))
+	{
+	}
+
+	void FCharacterGroundConstraint::Init(FSingleParticlePhysicsProxy* InCharacterProxy)
+	{
+		ensure(InCharacterProxy);
+
+		CharacterProxy.Modify(/*bInvalidate=*/true, DirtyFlags, Proxy, [InCharacterProxy](FParticleProxyProperty& Data)
+			{
+				Data.ParticleProxy = InCharacterProxy;
+			});
+
+		// Mark the settings and data as dirty as this is required to be pushed to the
+		// physics thread to initialize the data there and there is a chance the user
+		// does not set any settings or data before that happens
+		DirtyFlags.MarkDirty(EChaosPropertyFlags::CharacterGroundConstraintSettings);
+		DirtyFlags.MarkDirty(EChaosPropertyFlags::CharacterGroundConstraintDynamicData);
+	}
+
+	void FCharacterGroundConstraint::SyncRemoteDataImp(FDirtyPropertiesManager& Manager, int32 DataIdx, FDirtyChaosProperties& RemoteData)
+	{
+		CharacterProxy.SyncRemote(Manager, DataIdx, RemoteData);
+		GroundProxy.SyncRemote(Manager, DataIdx, RemoteData);
+		ConstraintSettings.SyncRemote(Manager, DataIdx, RemoteData);
+		ConstraintData.SyncRemote(Manager, DataIdx, RemoteData);
+	}
+}
\ No newline at end of file
diff --git a/Engine/Source/Runtime/Experimental/Chaos/Private/Chaos/Character/CharacterGroundConstraintContainer.cpp b/Engine/Source/Runtime/Experimental/Chaos/Private/Chaos/Character/CharacterGroundConstraintContainer.cpp
new file mode 100644
--- /dev/null
+++ b/Engine/Source/Runtime/Experimental/Chaos/Private/Chaos/Character/CharacterGroundConstraintContainer.cpp
@@ -0,0 +1,172 @@
+// Copyright Epic Games, Inc. All Rights Reserved.
+
+#include "Chaos/Character/CharacterGroundConstraintContainer.h"
+
+#include "Chaos/Island/IslandManager.h"
+#include "CharacterGroundConstraintContainerSolver.h"
+
+namespace Chaos
+{
+	FCharacterGroundConstraintContainer::FCharacterGroundConstraintContainer()
+		: Base(FCharacterGroundConstraintHandle::StaticType())
+		, ConstraintPool(16, 0)
+	{
+	}
+
+	FCharacterGroundConstraintContainer::~FCharacterGroundConstraintContainer()
+	{
+	}
+
+	void FCharacterGroundConstraintContainer::SetConstraintEnabled(int32 ConstraintIndex, bool bEnabled)
+	{
+		if (ConstraintIndex < Constraints.Num())
+		{
+			FCharacterGroundConstraintHandle* Constraint = Constraints[ConstraintIndex];
+			const FGenericParticleHandle CharacterParticle = FGenericParticleHandle(Constraint->GetCharacterParticle());
+			if (bEnabled)
+			{
+				if (CharacterParticle->Handle() != nullptr && !CharacterParticle->Disabled())
+				{
+					Constraint->SetEnabled(true);
+				}
+			}
+			else
+			{
+				// desirable to allow disabling no matter what state the endpoints
+				Constraint->SetEnabled(false);
+			}
+		}
+	}
+
+	FCharacterGroundConstraintHandle* FCharacterGroundConstraintContainer::AddConstraint(
+		const FCharacterGroundConstraintSettings& InConstraintSettings,
+		const FCharacterGroundConstraintDynamicData& InConstraintData,
+		FGeometryParticleHandle* CharacterParticle,
+		FGeometryParticleHandle* GroundParticle)
+	{
+		FCharacterGroundConstraintHandle* Constraint = ConstraintPool.Alloc();
+		Constraints.Add(Constraint);
+
+		check(CharacterParticle);
+		CharacterParticle->AddConstraintHandle(Constraint);
+		if (GroundParticle != nullptr)
+		{
+			GroundParticle->AddConstraintHandle(Constraint);
+		}
+
+		Constraint->CharacterParticle = CharacterParticle;
+		Constraint->GroundParticle = GroundParticle;
+		Constraint->Settings = InConstraintSettings;
+		Constraint->Data = InConstraintData;
+		Constraint->SetContainer(this);
+
+		// If our character particle is disabled, so is the constraint for now.
+		const bool bStartDisabled = FConstGenericParticleHandle(CharacterParticle)->Disabled();
+		Constraint->bDisabled = bStartDisabled;
+
+		return Constraint;
+	}
+
+	void FCharacterGroundConstraintContainer::RemoveConstraint(FCharacterGroundConstraintHandle* Constraint)
+	{
+		if (Constraint != nullptr)
+		{
+			Constraint->CharacterParticle->RemoveConstraintHandle(Constraint);
+			if (Constraint->GroundParticle != nullptr)
+			{
+				Constraint->GroundParticle->RemoveConstraintHandle(Constraint);
+			}
+
+			ConstraintPool.Free(Constraint);
+			Constraints.Remove(Constraint);
+		}
+	}
+
+	TUniquePtr<FConstraintContainerSolver> FCharacterGroundConstraintContainer::CreateSceneSolver(const int32 Priority)
+	{
+		return MakeUnique<Private::FCharacterGroundConstraintContainerSolver>(*this, Priority);
+	}
+
+	TUniquePtr<FConstraintContainerSolver> FCharacterGroundConstraintContainer::CreateGroupSolver(const int32 Priority)
+	{
+		return MakeUnique<Private::FCharacterGroundConstraintContainerSolver>(*this, Priority);
+	}
+
+	void FCharacterGroundConstraintContainer::AddConstraintsToGraph(Private::FPBDIslandManager& IslandManager)
+	{
+		for (FCharacterGroundConstraintHandle* Constraint : GetConstraints())
+		{
+			const bool bIsInGraph = Constraint->IsInConstraintGraph();
+			const bool bShouldBeInGraph = Constraint->IsEnabled() && CanEvaluate(Constraint);
+
+			if (bShouldBeInGraph && !bIsInGraph)
+			{
+				IslandManager.AddConstraint(ContainerId, Constraint, Constraint->GetConstrainedParticles());
+			}
+			else if (bIsInGraph && !bShouldBeInGraph)
+			{
+				IslandManager.RemoveConstraint(ContainerId, Constraint);
+			}
+		}
+	}
+
+	void FCharacterGroundConstraintContainer::PrepareTick()
+	{
+	}
+
+	void FCharacterGroundConstraintContainer::UnprepareTick()
+	{
+	}
+
+	void FCharacterGroundConstraintContainer::DisconnectConstraints(const TSet<TGeometryParticleHandle<FReal, 3>*>& RemovedParticles)
+	{
+		for (TGeometryParticleHandle<FReal, 3>*RemovedParticle : RemovedParticles)
+		{
+			for (FConstraintHandle* ConstraintHandle : RemovedParticle->ParticleConstraints())
+			{
+				if (FCharacterGroundConstraintHandle* Constraint = ConstraintHandle->As<FCharacterGroundConstraintHandle>())
+				{
+					if (Constraint->IsValid())
+					{
+						// If it's the character particle remove the constraint
+						// If it's the ground particle keep the constraint and set the ground particle to null
+						if (Constraint->CharacterParticle == RemovedParticle)
+						{
+							Constraint->SetEnabled(false); // constraint lifespan is managed by the proxy
+							Constraint->CharacterParticle = nullptr;
+						}
+						else if (Constraint->GroundParticle == RemovedParticle)
+						{
+							Constraint->GroundParticle = nullptr;
+						}
+					}
+					else
+					{
+						Constraint->SetEnabled(false);
+					}
+				}
+			}
+
+			RemovedParticle->ParticleConstraints().Reset();
+		}
+	}
+
+	bool FCharacterGroundConstraintContainer::CanEvaluate(const FCharacterGroundConstraintHandle* Constraint) const
+	{
+		if (!Constraint->IsEnabled())
+		{
+			return false;
+		}
+
+		const FGenericParticleHandle CharacterParticle = FGenericParticleHandle(Constraint->GetCharacterParticle());
+
+		// check for valid and enabled character particle
+		if (CharacterParticle->Handle() == nullptr || CharacterParticle->Disabled() || CharacterParticle->Sleeping() || !(CharacterParticle->InvM() > 0.0))
+		{
+			return false;
+		}
+
+		return true;
+	}
+
+}
\ No newline at end of file
diff --git a/Engine/Source/Runtime/Experimental/Chaos/Private/Chaos/Character/CharacterGroundConstraintContainerSolver.cpp b/Engine/Source/Runtime/Experimental/Chaos/Private/Chaos/Character/CharacterGroundConstraintContainerSolver.cpp
new file mode 100644
--- /dev/null
+++ b/Engine/Source/Runtime/Experimental/Chaos/Private/Chaos/Character/CharacterGroundConstraintContainerSolver.cpp
@@ -0,0 +1,135 @@
+// Copyright Epic Games, Inc. All Rights Reserved.
+
+#include "CharacterGroundConstraintContainerSolver.h"
+
+namespace Chaos
+{
+	namespace Private
+	{
+		FCharacterGroundConstraintContainerSolver::FCharacterGroundConstraintContainerSolver(FCharacterGroundConstraintContainer& InConstraintContainer, const int32 InPriority)
+			: FConstraintContainerSolver(InPriority)
+			, ConstraintContainer(InConstraintContainer)
+		{
+		}
+
+		FCharacterGroundConstraintContainerSolver::~FCharacterGroundConstraintContainerSolver()
+		{
+		}
+
+		void FCharacterGroundConstraintContainerSolver::Reset(const int32 InMaxConstraints)
+		{
+			SolvedConstraints.Reset(InMaxConstraints);
+
+			Solvers.Empty();
+			Solvers.SetNum(InMaxConstraints);
+		}
+
+		void FCharacterGroundConstraintContainerSolver::AddConstraints()
+		{
+			Reset(ConstraintContainer.GetNumConstraints());
+
+			for (FCharacterGroundConstraintHandle* Constraint : ConstraintContainer.GetConstraints())
+			{
+				AddConstraint(Constraint);
+			}
+		}
+
+		void FCharacterGroundConstraintContainerSolver::AddConstraints(const TArrayView<Private::FPBDIslandConstraint>& IslandConstraints)
+		{
+			for (Private::FPBDIslandConstraint& IslandConstraint : IslandConstraints)
+			{
+				// We will only ever be given constraints from our container (asserts in non-shipping)
+				FCharacterGroundConstraintHandle* Constraint = IslandConstraint.GetConstraint()->AsUnsafe<FCharacterGroundConstraintHandle>();
+
+				AddConstraint(Constraint);
+			}
+		}
+
+		void FCharacterGroundConstraintContainerSolver::AddConstraint(FCharacterGroundConstraintHandle* Constraint)
+		{
+			// If this triggers, Reset was called with the wrong constraint count
+			check(SolvedConstraints.Num() < SolvedConstraints.Max());
+
+			SolvedConstraints.Add(Constraint);
+		}
+
+		void FCharacterGroundConstraintContainerSolver::AddBodies(FSolverBodyContainer& SolverBodyContainer)
+		{
+			for (int32 SolverConstraintIndex = 0, SolverConstraintEndIndex = SolvedConstraints.Num(); SolverConstraintIndex < SolverConstraintEndIndex; ++SolverConstraintIndex)
+			{
+				FCharacterGroundConstraintHandle* Constraint = SolvedConstraints[SolverConstraintIndex];
+
+				FGenericParticleHandle CharacterParticle = FGenericParticleHandle(Constraint->GetCharacterParticle());
+				FGenericParticleHandle GroundParticle = FGenericParticleHandle(Constraint->GetGroundParticle());
+
+				FSolverBody* CharacterSolverBody = SolverBodyContainer.FindOrAdd(CharacterParticle);
+
+				if (GroundParticle.IsValid())
+				{
+					FSolverBody* GroundSolverBody = SolverBodyContainer.FindOrAdd(GroundParticle);
+					Solvers[SolverConstraintIndex].SetBodies(CharacterSolverBody, GroundSolverBody);
+				}
+				else
+				{
+					Solvers[SolverConstraintIndex].SetBodies(CharacterSolverBody);
+				}
+			}
+		}
+
+		void FCharacterGroundConstraintContainerSolver::GatherInput(const FReal Dt)
+		{
+			GatherInput(Dt, 0, SolvedConstraints.Num());
+		}
+
+		void FCharacterGroundConstraintContainerSolver::GatherInput(const FReal Dt, const int32 BeginIndex, const int32 EndIndex)
+		{
+			// We may have conservatively allocated the solver arrays. If so, reduce their size now
+			const int32 NumConstraints = GetNumConstraints();
+			check(Solvers.Num() >= NumConstraints);
+			if (Solvers.Num() > NumConstraints)
+			{
+				Solvers.SetNum(GetNumConstraints());
+			}
+
+			for (int32 SolverConstraintIndex = BeginIndex; SolverConstraintIndex < EndIndex; ++SolverConstraintIndex)
+			{
+				FCharacterGroundConstraintHandle* Constraint = SolvedConstraints[SolverConstraintIndex];
+				Solvers[SolverConstraintIndex].GatherInput(Dt, Constraint->GetSettings(), Constraint->GetData());
+			}
+		}
+
+		void FCharacterGroundConstraintContainerSolver::ScatterOutput(const FReal Dt)
+		{
+			ScatterOutput(Dt, 0, SolvedConstraints.Num());
+		}
+
+		void FCharacterGroundConstraintContainerSolver::ScatterOutput(const FReal Dt, const int32 BeginIndex, const int32 EndIndex)
+		{
+			for (int32 SolverConstraintIndex = BeginIndex; SolverConstraintIndex < EndIndex; ++SolverConstraintIndex)
+			{
+				FVec3 SolverAppliedForce, SolverAppliedTorque;
+				Solvers[SolverConstraintIndex].ScatterOutput(Dt, SolverAppliedForce, SolverAppliedTorque);
+				FCharacterGroundConstraintHandle* Constraint = SolvedConstraints[SolverConstraintIndex];
+				Constraint->SolverAppliedForce = SolverAppliedForce;
+				Constraint->SolverAppliedTorque = SolverAppliedTorque;
+			}
+		}
+
+		void FCharacterGroundConstraintContainerSolver::ApplyPositionConstraints(const FReal Dt, const int32 It, const int32 NumIts)
+		{
+			for (int32 SolverConstraintIndex = 0; SolverConstraintIndex < Solvers.Num(); ++SolverConstraintIndex)
+			{
+				Solvers[SolverConstraintIndex].SolvePosition();
+			}
+		}
+
+		void FCharacterGroundConstraintContainerSolver::ApplyVelocityConstraints(const FReal Dt, const int32 It, const int32 NumIts)
+		{
+		}
+
+		void FCharacterGroundConstraintContainerSolver::ApplyProjectionConstraints(const FReal Dt, const int32 It, const int32 NumIts)
+		{
+		}
+
+	} // namespace Private
+} // namespace Chaos
\ No newline at end of file
diff --git a/Engine/Source/Runtime/Experimental/Chaos/Private/Chaos/Character/CharacterGroundConstraintContainerSolver.h b/Engine/Source/Runtime/Experimental/Chaos/Private/Chaos/Character/CharacterGroundConstraintContainerSolver.h
new file mode 100644
--- /dev/null
+++ b/Engine/Source/Runtime/Experimental/Chaos/Private/Chaos/Character/CharacterGroundConstraintContainerSolver.h
@@ -0,0 +1,65 @@
+// Copyright Epic Games, Inc. All Rights Reserved.
+#pragma once
+
+#include "Chaos/Evolution/SolverConstraintContainer.h"
+#include "Chaos/Character/CharacterGroundConstraintContainer.h"
+#include "CharacterGroundConstraintSolver.h"
+
+namespace Chaos
+{
+	namespace Private
+	{
+		class FCharacterGroundConstraintContainerSolver : public FConstraintContainerSolver
+		{
+		public:
+			FCharacterGroundConstraintContainerSolver(FCharacterGroundConstraintContainer& InConstraintContainer, const int32 InPriority);
+			~FCharacterGroundConstraintContainerSolver();
+
+			//////////////////////////////////////////////////////////////////////////
+			// FConstraintContainerSolver impl
+
+			virtual int32 GetNumConstraints() const override final { return SolvedConstraints.Num(); }
+
+			virtual void Reset(const int32 InMaxConstraints) override final;
+
+			/// Add all constraints from the constraint container in the order that they appear there
+			virtual void AddConstraints() override final;
+
+			/// Add all constraints from an island (guaranteed to be from the same constraint container)
+			virtual void AddConstraints(const TArrayView<Private::FPBDIslandConstraint>& IslandConstraints) override final;
+
+			/// Set the bodies on the constraint from the constraint data solver
+			virtual void AddBodies(FSolverBodyContainer& SolverBodyContainer) override final;
+
+			/// Initialize the solver from the constraint data
+			virtual void GatherInput(const FReal Dt) override final;
+			virtual void GatherInput(const FReal Dt, const int32 BeginIndex, const int32 EndIndex) override final;
+
+			/// Get the solver output that will be stored per constraint - applied force and torque
+			virtual void ScatterOutput(const FReal Dt) override final;
+			virtual void ScatterOutput(const FReal Dt, const int32 BeginIndex, const int32 EndIndex) override final;
+
+			/// Apply a single position solve iteration for all constraints
+			virtual void ApplyPositionConstraints(const FReal Dt, const int32 It, const int32 NumIts) override final;
+
+			/// Apply a single velocity solve iteration for all constraints
+			/// Note: Currently not used for character ground constraints
+			virtual void ApplyVelocityConstraints(const FReal Dt, const int32 It, const int32 NumIts) override final;
+
+			/// Apply a single projection solve iteration for all constraints
+			/// Note: Currently not used for character ground constraints
+			virtual void ApplyProjectionConstraints(const FReal Dt, const int32 It, const int32 NumIts) override final;
+
+		private:
+			void AddConstraint(FCharacterGroundConstraintHandle* Constraint);
+
+			FCharacterGroundConstraintContainer& ConstraintContainer;
+
+			TArray<FCharacterGroundConstraintHandle*> SolvedConstraints;
+
+			// Array of single constraint solvers
+			TArray<FCharacterGroundConstraintSolver> Solvers;
+		};
+
+	} // namespace Private
+} // namespace Chaos
\ No newline at end of file
diff --git a/Engine/Source/Runtime/Experimental/Chaos/Private/Chaos/PBDRigidsEvolutionGBF.cpp b/Engine/Source/Runtime/Experimental/Chaos/Private/Chaos/PBDRigidsEvolutionGBF.cpp
--- a/Engine/Source/Runtime/Experimental/Chaos/Private/Chaos/PBDRigidsEvolutionGBF.cpp
+++ b/Engine/Source/Runtime/Experimental/Chaos/Private/Chaos/PBDRigidsEvolutionGBF.cpp
@@ -68,6 +68,9 @@ namespace Chaos
 		int32 ChaosSolverSuspensionPriority = 0;
 		FAutoConsoleVariableRef CVarChaosSolverSuspensionPriority(TEXT("p.Chaos.Solver.Suspension.Priority"), ChaosSolverSuspensionPriority, TEXT("Set constraint priority. Larger values are evaluated later [def:0]"));
 
+		int32 ChaosSolverCharacterGroundConstraintPriority = 0;
+		FAutoConsoleVariableRef CVarChaosSolverChaosCharacterGroundConstraintPriority(TEXT("p.Chaos.Solver.CharacterGroundConstraint.Priority"), ChaosSolverCharacterGroundConstraintPriority, TEXT("Set constraint priority. Larger values are evaluated later [def:0]"));
+
 		bool DoTransferJointConstraintCollisions = true;
 		FAutoConsoleVariableRef CVarDoTransferJointConstraintCollisions(TEXT("p.Chaos.Solver.Joint.TransferCollisions"), DoTransferJointConstraintCollisions, TEXT("Allows joints to apply collisions to the parent from the child when the Joints TransferCollisionScale is not 0 [def:true]"));
 
@@ -679,6 +682,7 @@ FPBDRigidsEvolutionGBF::FPBDRigidsEvolutionGBF(FPBDRigidsSOAs& InParticles, THan
 	AddConstraintContainer(SuspensionConstraints, ChaosSolverSuspensionPriority);
 	AddConstraintContainer(CollisionConstraints, ChaosSolverCollisionPriority);
 	AddConstraintContainer(JointConstraints, ChaosSolverJointPriority);
+	AddConstraintContainer(CharacterGroundConstraints, ChaosSolverCharacterGroundConstraintPriority);
 
 	SetInternalParticleInitilizationFunction([](const FGeometryParticleHandle*, const FGeometryParticleHandle*) {});
 }
diff --git a/Engine/Source/Runtime/Experimental/Chaos/Private/PBDRigidsSolver.cpp b/Engine/Source/Runtime/Experimental/Chaos/Private/PBDRigidsSolver.cpp
--- a/Engine/Source/Runtime/Experimental/Chaos/Private/PBDRigidsSolver.cpp
+++ b/Engine/Source/Runtime/Experimental/Chaos/Private/PBDRigidsSolver.cpp
@@ -4,6 +4,7 @@
 
 #include "Async/AsyncWork.h"
 #include "Chaos/ChaosArchive.h"
+#include "Chaos/Character/CharacterGroundConstraint.h"
 #include "Chaos/PBDCollisionConstraintsUtil.h"
 #include "Chaos/Utilities.h"
 #include "Chaos/ChaosDebugDraw.h"
@@ -13,6 +14,7 @@
 #include "ChaosVisualDebugger/ChaosVisualDebuggerTrace.h"
 #include "HAL/FileManager.h"
 #include "Misc/ScopeLock.h"
+#include "PhysicsProxy/CharacterGroundConstraintProxy.h"
 #include "PhysicsProxy/SingleParticlePhysicsProxy.h"
 #include "PhysicsProxy/SkeletalMeshPhysicsProxy.h"
 #include "PhysicsProxy/StaticMeshPhysicsProxy.h"
@@ -51,6 +53,7 @@ DECLARE_DWORD_COUNTER_STAT(TEXT("NumActiveManifoldPoints"), STAT_ChaosCounter_Nu
 DECLARE_DWORD_COUNTER_STAT(TEXT("NumRestoredManifoldPoints"), STAT_ChaosCounter_NumRestoredManifoldPoints, STATGROUP_ChaosCounters);
 DECLARE_DWORD_COUNTER_STAT(TEXT("NumUpdatedManifoldPoints"), STAT_ChaosCounter_NumUpdatedManifoldPoints, STATGROUP_ChaosCounters);
 DECLARE_DWORD_COUNTER_STAT(TEXT("NumJoints"), STAT_ChaosCounter_NumJoints, STATGROUP_ChaosCounters);
+DECLARE_DWORD_COUNTER_STAT(TEXT("NumCharacterGroundConstraints"), STAT_ChaosCounter_NumCharacterGroundConstraints, STATGROUP_ChaosCounters);
 
 // Stat Iteration counters
 DECLARE_DWORD_COUNTER_STAT(TEXT("NumPositionIterations"), STAT_ChaosCounter_NumPositionIterations, STATGROUP_ChaosCounters);
@@ -80,6 +83,7 @@ namespace Chaos
 		int32 ChaosSolverDrawShapesShowKinematic = 1;
 		int32 ChaosSolverDrawShapesShowDynamic = 1;
 		int32 ChaosSolverDrawJoints = 0;
+		int32 ChaosSolverDrawCharacterGroundConstraints = 0;
 		int32 ChaosSolverDebugDrawSpatialAccelerationStructure = 0;
 		int32 ChaosSolverDebugDrawSpatialAccelerationStructureShowLeaves = 0;
 		int32 ChaosSolverDebugDrawSpatialAccelerationStructureShowNodes = 0;
@@ -100,6 +104,7 @@ namespace Chaos
 		FAutoConsoleVariableRef CVarChaosSolverDrawShapesShapesKinematic(TEXT("p.Chaos.Solver.DebugDraw.ShowKinematics"), ChaosSolverDrawShapesShowKinematic, TEXT("If DebugDrawShapes is enabled, whether to show kinematic objects"));
 		FAutoConsoleVariableRef CVarChaosSolverDrawShapesShapesDynamic(TEXT("p.Chaos.Solver.DebugDraw.ShowDynamics"), ChaosSolverDrawShapesShowDynamic, TEXT("If DebugDrawShapes is enabled, whether to show dynamic objects"));
 		FAutoConsoleVariableRef CVarChaosSolverDrawJoints(TEXT("p.Chaos.Solver.DebugDrawJoints"), ChaosSolverDrawJoints, TEXT("Draw joints"));
+		FAutoConsoleVariableRef CVarChaosSolverDrawCharacterGroundConstraints(TEXT("p.Chaos.Solver.DebugDrawCharacterGroundConstraints"), ChaosSolverDrawCharacterGroundConstraints, TEXT("Draw character ground constraints"));
 		FAutoConsoleVariableRef CVarChaosSolverDebugDrawSpatialAccelerationStructure(TEXT("p.Chaos.Solver.DebugDrawSpatialAccelerationStructure"), ChaosSolverDebugDrawSpatialAccelerationStructure, TEXT("Draw spatial acceleration structure"));
 		FAutoConsoleVariableRef CVarChaosSolverDebugDrawSpatialAccelerationStructureShowLeaves(TEXT("p.Chaos.Solver.DebugDrawSpatialAccelerationStructure.ShowLeaves"), ChaosSolverDebugDrawSpatialAccelerationStructureShowLeaves, TEXT("Show spatial acceleration structure leaves when its debug draw is enabled"));
 		FAutoConsoleVariableRef CVarChaosSolverDebugDrawSpatialAccelerationStructureShowNodes(TEXT("p.Chaos.Solver.DebugDrawSpatialAccelerationStructure.ShowNodes"), ChaosSolverDebugDrawSpatialAccelerationStructureShowNodes, TEXT("Show spatial acceleration structure nodes when its debug draw is enabled"));
@@ -712,6 +717,43 @@ namespace Chaos
 		});
 	}
 
+	void FPBDRigidsSolver::RegisterObject(Chaos::FCharacterGroundConstraint* GTConstraint)
+	{
+		LLM_SCOPE(ELLMTag::ChaosConstraint);
+		FCharacterGroundConstraintProxy* ConstraintProxy = new FCharacterGroundConstraintProxy(GTConstraint);
+		ConstraintProxy->SetSolver(this);
+
+		AddDirtyProxy(ConstraintProxy);
+	}
+
+	void FPBDRigidsSolver::UnregisterObject(Chaos::FCharacterGroundConstraint* GTConstraint)
+	{
+		FCharacterGroundConstraintProxy* ConstraintProxy = GTConstraint->GetProxy<FCharacterGroundConstraintProxy>();
+		check(ConstraintProxy);
+
+		RemoveDirtyProxy(ConstraintProxy);
+
+		// mark proxy timestamp so we avoid trying to pull from sim after deletion
+		GTConstraint->GetProxy()->MarkDeleted();
+		GTConstraint->SetProxy(static_cast<FCharacterGroundConstraintProxy*>(nullptr));
+
+		ConstraintProxy->DestroyOnGameThread();	//destroy the game thread portion of the proxy
+
+		// Finish de-registration on the physics thread...
+		EnqueueCommandImmediate([ConstraintProxy, this]()
+			{
+				// TODO: Add character ground constraint to rewind data
+				//if (FRewindData* RewindData = GetRewindData())
+				//{
+				//	RewindData->RemoveObject(ConstraintProxy->GetPhysicsThreadAPI());
+				//}
+
+				ConstraintProxy->DestroyOnPhysicsThread(this);
+				CharacterGroundConstraintProxies_Internal.RemoveSingle(ConstraintProxy);
+				delete ConstraintProxy;
+			});
+	}
+
 	void FPBDRigidsSolver::RegisterObject(Chaos::FSuspensionConstraint* GTConstraint)
 	{
 		LLM_SCOPE(ELLMTag::ChaosConstraint);
@@ -1032,6 +1074,13 @@ namespace Chaos
 				Proxy->ResetDirtyIdx();
 				break;
 			}
+			case EPhysicsProxyType::CharacterGroundConstraintType:
+			{
+				auto Proxy = static_cast<FCharacterGroundConstraintProxy*>(Dirty.Proxy);
+				Proxy->PushStateOnGameThread(*Manager, DataIdx, Dirty.PropertyData);
+				Proxy->ResetDirtyIdx();
+				break;
+			}
 
 			default:
 			ensure(0 && TEXT("Unknown proxy type in physics solver."));
@@ -1042,9 +1091,10 @@ namespace Chaos
 		{
 			int32 NumJoints = 0;
 			int32 NumSuspension = 0;
+			int32 NumCharacterGroundConstraints = 0;
 			int32 NumParticles = 0;
 			UE_LOG(LogChaos, Warning, TEXT("LogDirtyParticles:"));
-			DirtyProxiesData->ForEachProxy([&NumJoints, &NumSuspension, &NumParticles](int32 DataIdx, FDirtyProxy& Dirty)
+			DirtyProxiesData->ForEachProxy([&NumJoints, &NumSuspension, &NumCharacterGroundConstraints, &NumParticles](int32 DataIdx, FDirtyProxy& Dirty)
 			{
 				switch (Dirty.Proxy->GetType())
 				{
@@ -1068,12 +1118,17 @@ namespace Chaos
 						++NumSuspension;
 						break;
 					}
+					case EPhysicsProxyType::CharacterGroundConstraintType:
+					{
+						++NumCharacterGroundConstraints;
+						break;
+					}
 
 					default: break;
 				}
 			});
 
-			UE_LOG(LogChaos, Warning, TEXT("Num Particles:%d Num Shapes:%d Num Joints:%d Num Suspensions:%d"), NumParticles, DirtyProxiesData->NumDirtyShapes(), NumJoints, NumSuspension);
+			UE_LOG(LogChaos, Warning, TEXT("Num Particles:%d Num Shapes:%d Num Joints:%d Num Suspensions:%d Num CharGround:%d"), NumParticles, DirtyProxiesData->NumDirtyShapes(), NumJoints, NumSuspension, NumCharacterGroundConstraints);
 		}
 
 		GetEvolution()->GetBroadPhase().GetIgnoreCollisionManager().PushProducerStorageData_External(MarshallingManager.GetExternalTimestamp_External());
@@ -1193,6 +1248,7 @@ namespace Chaos
 					}
 					case EPhysicsProxyType::JointConstraintType:
 					case EPhysicsProxyType::SuspensionConstraintType:
+					case EPhysicsProxyType::CharacterGroundConstraintType:
 					{
 						// Pass until after all bodies are created. 
 						break;
@@ -1250,6 +1306,28 @@ namespace Chaos
 					break;
 				}
 
+				case EPhysicsProxyType::CharacterGroundConstraintType:
+				{
+					auto ConstraintProxy = static_cast<FCharacterGroundConstraintProxy*>(Dirty.Proxy);
+					const bool bIsNew = !ConstraintProxy->IsInitialized();
+					if (bIsNew)
+					{
+						CharacterGroundConstraintProxies_Internal.Add(ConstraintProxy);
+						ConstraintProxy->InitializeOnPhysicsThread(this, *Manager, DataIdx, Dirty.PropertyData);
+						ConstraintProxy->SetInitialized(GetCurrentFrame());
+					}
+
+					//TODO: Support rewind for character ground constraints
+					//if (RewindData)
+					//{
+					//	RewindData->PushGTDirtyData(*Manager, DataIdx, Dirty, nullptr);
+					//}
+
+					ConstraintProxy->PushStateOnPhysicsThread(this, *Manager, DataIdx, Dirty.PropertyData);
+					Dirty.Proxy->ResetDirtyIdx();
+					break;
+				}
+
 				}
 			}
 		});
@@ -1495,9 +1573,18 @@ namespace Chaos
 				JointConstraintPhysicsProxies_Internal[Idx]->BufferPhysicsResults(PullData->DirtyJointConstraints.Last());
 			}
 		}
-		
 
-		
+		{
+			ensure(PullData->DirtyCharacterGroundConstraints.Num() == 0);	//we only fill this once per frame
+			PullData->DirtyCharacterGroundConstraints.Reserve(CharacterGroundConstraintProxies_Internal.Num());
+
+			for (int32 Idx = 0; Idx < CharacterGroundConstraintProxies_Internal.Num(); ++Idx)
+			{
+				PullData->DirtyCharacterGroundConstraints.AddDefaulted();
+				CharacterGroundConstraintProxies_Internal[Idx]->BufferPhysicsResults(PullData->DirtyCharacterGroundConstraints.Last());
+			}
+		}
+
 		// Now that results have been buffered we have completed a solve step so we can broadcast that event
 		EventPostSolve.Broadcast(MLastDt);
 
@@ -1554,6 +1641,7 @@ CSV_CUSTOM_STAT(PhysicsCounters, Name, Value, ECsvCustomStatOp::Set);
 		CHAOS_COUNTER_STAT(NumIslandGroups, GetEvolution()->GetIslandGroupManager().GetNumActiveGroups());
 		CHAOS_COUNTER_STAT(NumContacts, NumCollisionConstraints());
 		CHAOS_COUNTER_STAT(NumJoints, NumJointConstraints());
+		CHAOS_COUNTER_STAT(NumCharacterGroundConstraints, GetEvolution()->GetCharacterGroundConstraints().GetNumConstraints());
 
 #if !UE_BUILD_SHIPPING && !UE_BUILD_TEST
 		UpdateExpensiveStatCounters();
@@ -1693,6 +1781,10 @@ CSV_CUSTOM_STAT(PhysicsCounters, Name, Value, ECsvCustomStatOp::Set);
 		{
 			DebugDraw::DrawJointConstraints(FRigidTransform3(), MEvolution->GetJointConstraints(), 1.0f, ChaosSolverDrawJointFeatures, &ChaosSolverDebugDebugDrawSettings);
 		}
+		if (ChaosSolverDrawCharacterGroundConstraints == 1)
+		{
+			DebugDraw::DrawCharacterGroundConstraints(FRigidTransform3(), MEvolution->GetCharacterGroundConstraints(), &ChaosSolverDebugDebugDrawSettings);
+		}
 		if (ChaosSolverDebugDrawSpatialAccelerationStructure)
 		{
 			if (const auto SpatialAccelerationStructure = GetEvolution()->GetSpatialAcceleration())
diff --git a/Engine/Source/Runtime/Experimental/Chaos/Private/PhysicsProxy/CharacterGroundConstraintProxy.cpp b/Engine/Source/Runtime/Experimental/Chaos/Private/PhysicsProxy/CharacterGroundConstraintProxy.cpp
new file mode 100644
--- /dev/null
+++ b/Engine/Source/Runtime/Experimental/Chaos/Private/PhysicsProxy/CharacterGroundConstraintProxy.cpp
@@ -0,0 +1,165 @@
+// Copyright Epic Games, Inc. All Rights Reserved.
+
+#include "PhysicsProxy/CharacterGroundConstraintProxy.h"
+
+#include "Chaos/Character/CharacterGroundConstraint.h"
+#include "Chaos/Character/CharacterGroundConstraintContainer.h"
+#include "Chaos/ParticleDirtyFlags.h"
+#include "Chaos/PullPhysicsDataImp.h"
+#include "PBDRigidsSolver.h"
+
+//////////////////////////////////////////////////////////////////////////
+/// Utility function
+
+namespace
+{
+	Chaos::FGeometryParticleHandle* GetParticleHandleFromProxy(IPhysicsProxyBase* ProxyBase)
+	{
+		if (ProxyBase)
+		{
+			if (ProxyBase->GetType() == EPhysicsProxyType::SingleParticleProxy)
+			{
+				return ((Chaos::FSingleParticlePhysicsProxy*)ProxyBase)->GetHandle_LowLevel();
+			}
+		}
+		return nullptr;
+	}
+}
+
+//////////////////////////////////////////////////////////////////////////
+
+namespace Chaos
+{
+	FCharacterGroundConstraintProxy::FCharacterGroundConstraintProxy(FCharacterGroundConstraint* InConstraintGT, FCharacterGroundConstraintHandle* InConstraintPT, UObject* InOwner)
+		: Base(EPhysicsProxyType::CharacterGroundConstraintType, InOwner, MakeShared<FProxyTimestampBase>())
+		, Constraint_GT(InConstraintGT)
+		, Constraint_PT(InConstraintPT)
+	{
+		check(Constraint_GT != nullptr);
+		Constraint_GT->SetProxy(this);
+	}
+
+	void FCharacterGroundConstraintProxy::InitializeOnPhysicsThread(FPBDRigidsSolver* InSolver, FDirtyPropertiesManager& Manager, int32 DataIdx, FDirtyChaosProperties& RemoteData)
+	{
+		TGeometryParticleHandles<FReal,3>& Handles = InSolver->GetParticles().GetParticleHandles();
+		if (Handles.Size() > 0)
+		{
+			check(InSolver);
+			FCharacterGroundConstraintContainer& ConstraintContainer = InSolver->GetCharacterGroundConstraints();
+
+			FGeometryParticleHandle* CharacterHandle = nullptr;
+			if (const FParticleProxyProperty* CharacterProxy = RemoteData.FindCharacterParticleProxy(Manager, DataIdx))
+			{
+				CharacterHandle = GetParticleHandleFromProxy(CharacterProxy->ParticleProxy);
+			}
+
+			FGeometryParticleHandle* GroundHandle = nullptr;
+			if (const FParticleProxyProperty* GroundProxy = RemoteData.FindGroundParticleProxy(Manager, DataIdx))
+			{
+				GroundHandle = GetParticleHandleFromProxy(GroundProxy->ParticleProxy);
+			}
+
+			// Constraint only requires that the character particle be set
+			if (CharacterHandle)
+			{
+				const FCharacterGroundConstraintSettings* ConstraintSettings = RemoteData.FindCharacterGroundConstraintSettings(Manager, DataIdx);
+				const FCharacterGroundConstraintDynamicData* ConstraintData = RemoteData.FindCharacterGroundConstraintDynamicData(Manager, DataIdx);
+
+				if (ConstraintSettings && ConstraintData)
+				{
+					Constraint_PT = ConstraintContainer.AddConstraint(*ConstraintSettings, *ConstraintData, CharacterHandle, GroundHandle);
+
+					CharacterHandle->AddConstraintHandle(Constraint_PT);
+
+					if (GroundHandle)
+					{
+						GroundHandle->AddConstraintHandle(Constraint_PT);
+					}
+				}
+			}
+		}
+	}
+
+	void FCharacterGroundConstraintProxy::PushStateOnPhysicsThread(FPBDRigidsSolver* InSolver, const FDirtyPropertiesManager& Manager, int32 DataIdx, const FDirtyChaosProperties& RemoteData)
+	{
+		if (Constraint_PT && Constraint_PT->IsValid())
+		{
+			if (const FCharacterGroundConstraintSettings* Settings = RemoteData.FindCharacterGroundConstraintSettings(Manager, DataIdx))
+			{
+				Constraint_PT->Settings = *Settings;
+			}
+
+			if (const FCharacterGroundConstraintDynamicData* Data = RemoteData.FindCharacterGroundConstraintDynamicData(Manager, DataIdx))
+			{
+				Constraint_PT->Data = *Data;
+			}
+
+			if (const FParticleProxyProperty* GroundParticleProxy = RemoteData.FindGroundParticleProxy(Manager, DataIdx))
+			{
+				if (FGeometryParticleHandle* Handle = GetParticleHandleFromProxy(GroundParticleProxy->ParticleProxy))
+				{
+					Constraint_PT->GroundParticle = Handle;
+				}
+			}
+		}
+	}
+
+	void FCharacterGroundConstraintProxy::PushStateOnGameThread(FDirtyPropertiesManager& Manager, int32 DataIdx, FDirtyChaosProperties& RemoteData)
+	{
+		if (Constraint_GT && Constraint_GT->IsValid())
+		{
+			Constraint_GT->SyncRemoteData(Manager, DataIdx, RemoteData);
+		}
+	}
+
+	void FCharacterGroundConstraintProxy::DestroyOnGameThread()
+	{
+		if (Constraint_GT)
+		{
+			delete Constraint_GT;
+			Constraint_GT = nullptr;
+		}
+	}
+
+	void FCharacterGroundConstraintProxy::DestroyOnPhysicsThread(FPBDRigidsSolver* InSolver)
+	{
+		if (Constraint_PT)
+		{
+			if (FGeometryParticleHandle* CharacterParticle = Constraint_PT->GetCharacterParticle())
+			{
+				CharacterParticle->RemoveConstraintHandle(Constraint_PT);
+			}
+			if (FGeometryParticleHandle* GroundParticle = Constraint_PT->GetGroundParticle())
+			{
+				GroundParticle->RemoveConstraintHandle(Constraint_PT);
+			}
+
+			InSolver->GetEvolution()->RemoveConstraintFromConstraintGraph(Constraint_PT);
+			FCharacterGroundConstraintContainer& Constraints = InSolver->GetCharacterGroundConstraints();
+			Constraints.RemoveConstraint(Constraint_PT);
+			Constraint_PT = nullptr;
+		}
+	}
+
+	void FCharacterGroundConstraintProxy::BufferPhysicsResults(FDirtyCharacterGroundConstraintData& Buffer)
+	{
+		Buffer.SetProxy(*this);
+		if (Constraint_PT != nullptr && Constraint_PT->IsValid())
+		{
+			Buffer.Force = Constraint_PT->GetSolverAppliedForce();
+			Buffer.Torque = Constraint_PT->GetSolverAppliedTorque();
+		}
+	}
+
+	bool FCharacterGroundConstraintProxy::PullFromPhysicsState(const FDirtyCharacterGroundConstraintData& Buffer, const int32 SolverSyncTimestamp)
+	{
+		if (Constraint_GT != nullptr && Constraint_GT->IsValid())
+		{
+			Constraint_GT->SolverAppliedForce = Buffer.Force;
+			Constraint_GT->SolverAppliedTorque = Buffer.Torque;
+		}
+
+		return true;
+	}
+
+} // namespace Chaos
\ No newline at end of file
diff --git a/Engine/Source/Runtime/Experimental/Chaos/Public/Chaos/ChaosDebugDraw.h b/Engine/Source/Runtime/Experimental/Chaos/Public/Chaos/ChaosDebugDraw.h
--- a/Engine/Source/Runtime/Experimental/Chaos/Public/Chaos/ChaosDebugDraw.h
+++ b/Engine/Source/Runtime/Experimental/Chaos/Public/Chaos/ChaosDebugDraw.h
@@ -217,6 +217,7 @@ namespace Chaos
 		CHAOS_API void DrawCollisions(const FRigidTransform3& SpaceTransform, const Private::FCollisionConstraintAllocator& CollisionAllocator, FRealSingle ColorScale, const FChaosDebugDrawSettings* Settings = nullptr);
 		CHAOS_API void DrawJointConstraints(const FRigidTransform3& SpaceTransform, const TArray<FPBDJointConstraintHandle*>& ConstraintHandles, FRealSingle ColorScale, const FChaosDebugDrawJointFeatures& FeatureMask = FChaosDebugDrawJointFeatures::MakeDefault(), const FChaosDebugDrawSettings* Settings = nullptr);
 		CHAOS_API void DrawJointConstraints(const FRigidTransform3& SpaceTransform, const FPBDJointConstraints& Constraints, FRealSingle ColorScale, const FChaosDebugDrawJointFeatures& FeatureMask = FChaosDebugDrawJointFeatures::MakeDefault(), const FChaosDebugDrawSettings* Settings = nullptr);
+		CHAOS_API void DrawCharacterGroundConstraints(const FRigidTransform3& SpaceTransform, const FCharacterGroundConstraintContainer& Constraints, const FChaosDebugDrawSettings* Settings = nullptr);
 		CHAOS_API void DrawSimulationSpace(const FSimulationSpace& SimSpace, const FChaosDebugDrawSettings* Settings = nullptr);
 		CHAOS_API void DrawShape(const FRigidTransform3& ShapeTransform, const FImplicitObject* Implicit, const FShapeOrShapesArray& Shapes, const FColor& Color, const FChaosDebugDrawSettings* Settings = nullptr);
 		CHAOS_API void DrawConstraintGraph(const FRigidTransform3& ShapeTransform, const Private::FPBDIslandManager& Graph, const FChaosDebugDrawSettings* Settings = nullptr);
diff --git a/Engine/Source/Runtime/Experimental/Chaos/Public/Chaos/ChaosDebugDrawDeclares.h b/Engine/Source/Runtime/Experimental/Chaos/Public/Chaos/ChaosDebugDrawDeclares.h
--- a/Engine/Source/Runtime/Experimental/Chaos/Public/Chaos/ChaosDebugDrawDeclares.h
+++ b/Engine/Source/Runtime/Experimental/Chaos/Public/Chaos/ChaosDebugDrawDeclares.h
@@ -26,4 +26,6 @@ namespace Chaos
 	class FPBDJointConstraints;
 
 	class FSimulationSpace;
+
+	class FCharacterGroundConstraintContainer;
 }
diff --git a/Engine/Source/Runtime/Experimental/Chaos/Public/Chaos/Character/CharacterGroundConstraint.h b/Engine/Source/Runtime/Experimental/Chaos/Public/Chaos/Character/CharacterGroundConstraint.h
new file mode 100644
--- /dev/null
+++ b/Engine/Source/Runtime/Experimental/Chaos/Public/Chaos/Character/CharacterGroundConstraint.h
@@ -0,0 +1,231 @@
+// Copyright Epic Games, Inc. All Rights Reserved.
+#pragma once
+
+#include "Chaos/Character/CharacterGroundConstraintSettings.h"
+#include "Chaos/PBDConstraintBaseData.h"
+#include "Chaos/ParticleDirtyFlags.h"
+#include "PhysicsProxy/SingleParticlePhysicsProxy.h"
+
+namespace Chaos
+{
+	/// Game thread side representation of a character ground constraint
+	/// Modified data from this class will be pushed automatically
+	/// to the physics thread and available through the
+	/// FCharacterGroundConstraintHandle class.
+	/// The FCharacterGroundConstraintProxy class, accessible via
+	/// GetProxy(), knows about both representations and marshalls
+	/// the data between them
+	class CHAOS_API FCharacterGroundConstraint : public FConstraintBase
+	{
+	public:
+		using Base = FConstraintBase;
+		friend class FCharacterGroundConstraintProxy; // For setting the output data
+
+		FCharacterGroundConstraint();
+		virtual ~FCharacterGroundConstraint() override {}
+
+		// The constraint requires a character body, which is fixed, but the ground body is
+		// optional and can be changed dynamically using SetGroundParticleProxy
+		void Init(FSingleParticlePhysicsProxy* InCharacterProxy);
+
+		//////////////////////////////////////////////////////////////////////////
+		// Particle proxy
+
+		FSingleParticlePhysicsProxy* GetCharacterParticleProxy() const
+		{
+			return CharacterProxy.Read().ParticleProxy;
+		}
+
+		FSingleParticlePhysicsProxy* GetGroundParticleProxy() const
+		{
+			return GroundProxy.Read().ParticleProxy;
+		}
+
+		void SetGroundParticleProxy(FSingleParticlePhysicsProxy* InGroundProxy)
+		{
+			GroundProxy.Modify(/*bInvalidate=*/true, DirtyFlags, Proxy, [InGroundProxy](FParticleProxyProperty& Data)
+				{
+					Data.ParticleProxy = InGroundProxy;
+				});
+		}
+
+		//////////////////////////////////////////////////////////////////////////
+		// Settings accessors
+
+		/// Unit normal to the ground plane
+		/// Input vector is normalized if it is not already a unit vector
+		void SetGroundNormal(const FVec3& Normal)
+		{
+			FVec3 Tmp = Normal;
+			FReal SizeSq = Tmp.SizeSquared();
+			if (!FMath::IsNearlyEqual(SizeSq, FReal(1.0)))
+			{
+				Tmp *= FMath::InvSqrt(SizeSq);
+			}
+			ConstraintData.Modify(true, DirtyFlags, Proxy, [&Tmp](FCharacterGroundConstraintDynamicData& Data) {
+				Data.GroundNormal = Tmp;
+				});
+		}
+		FVec3 GetGroundNormal() const { return ConstraintData.Read().GroundNormal; }
+
+		/// Unit vector defining the character up direction
+		/// Input vector is normalized if it is not already a unit vector
+		void SetVerticalAxis(const FVec3& Axis)
+		{
+			FVec3 Tmp = Axis;
+			FReal SizeSq = Tmp.SizeSquared();
+			if (!FMath::IsNearlyEqual(SizeSq, FReal(1.0)))
+			{
+				Tmp *= FMath::InvSqrt(SizeSq);
+			}
+			ConstraintSettings.Modify(true, DirtyFlags, Proxy, [&Tmp](FCharacterGroundConstraintSettings& Data) {
+				Data.VerticalAxis = Tmp;
+				});
+		}
+		FVec3 GetVerticalAxis() const { return ConstraintSettings.Read().VerticalAxis; }
+
+		/// Vertical distance from the character body to the ground
+		/// This should be set every frame as the character moves over the ground and is typically
+		/// measured by performing one or more physics traces
+		/// Input should be positive
+		void SetGroundDistance(const FReal& Value)
+		{
+			ConstraintData.Modify(true, DirtyFlags, Proxy, [&Value](FCharacterGroundConstraintDynamicData& Data) {
+				Data.GroundDistance = Value > FReal(0.0) ? Value : FReal(0.0);
+				});
+		}
+		FReal GetGroundDistance() const { return ConstraintData.Read().GroundDistance; }
+
+		/// Target vertical distance from the character body to the ground
+		/// If the ground distance is smaller than this number the constraint will
+		/// try to fix up the error to keep the character at the target height
+		/// Input should be positive
+		void SetTargetHeight(const FReal& Value)
+		{
+			ConstraintSettings.Modify(true, DirtyFlags, Proxy, [&Value](FCharacterGroundConstraintSettings& Data) {
+				Data.TargetHeight = Value > FReal(0.0) ? Value : FReal(0.0);
+				});
+		}
+		FReal GetTargetHeight() const { return ConstraintSettings.Read().TargetHeight; }
+
+		/// The damping factor is used to add softness to the vertical constraint between the character
+		/// and the ground. This can add springiness to landings and slow the rate of step ups.
+		/// Input should be positive. A value of 0 gives rigid behavior. A value of 1/Dt or above will
+		/// give maximum softness
+		/// Units are /T
+		void SetDampingFactor(const FReal& Value)
+		{
+			ConstraintSettings.Modify(true, DirtyFlags, Proxy, [&Value](FCharacterGroundConstraintSettings& Data) {
+				Data.DampingFactor = Value > FReal(0.0) ? Value : FReal(0.0);
+				});
+		}
+		FReal GetDampingFactor() const { return ConstraintSettings.Read().DampingFactor; }
+
+		/// This is the height below which the character is assumed to be on the ground and can apply
+		/// force to move. This is useful if the character is slightly off the ground but should still
+		/// be considered in a locomotion rather than falling state.
+		/// If this number is set to zero then the constraint can only apply force to move to the motion
+		/// target if the solver determines that there is a ground reaction force
+		/// Input should be positive
+		void SetAssumedOnGroundHeight(const FReal& Value)
+		{
+			ConstraintSettings.Modify(true, DirtyFlags, Proxy, [&Value](FCharacterGroundConstraintSettings& Data) {
+				Data.AssumedOnGroundHeight = Value > FReal(0.0) ? Value : FReal(0.0);
+				});
+		}
+		FReal GetAssumedOnGroundHeight() const { return ConstraintSettings.Read().AssumedOnGroundHeight; }
+
+		/// The maximum angle of slope that is considered walkable.
+		/// The slope angle is defined as the angle between the ground normal and vertical axis.
+		/// If the slope angle is greater than the maximum walkable slope angle the motion target
+		/// vector is adjusted to stop movement up the slope
+		void SetMaxWalkableSlopeAngle(const FReal& Value)
+		{
+			ConstraintSettings.Modify(true, DirtyFlags, Proxy, [&Value](FCharacterGroundConstraintSettings& Data) {
+				FReal Angle = FMath::Clamp(Value, 0.0, UE_HALF_PI);
+				Data.CosMaxWalkableSlopeAngle = FMath::Cos(Angle);
+				});
+		}
+		FReal GetMaxWalkableSlopeAngle() const { return FMath::Acos(ConstraintSettings.Read().CosMaxWalkableSlopeAngle); }
+		FReal GetCosMaxWalkableSlopeAngle() const { return ConstraintSettings.Read().CosMaxWalkableSlopeAngle; }
+
+		/// Set the target change in position for the character. This is used to move the character
+		/// along the ground, and this target will be clamped to the ground plane.
+		/// The solver will apply a force to move the character to the desired position as long as
+		/// the character is considered to be on the ground and the force required is less than or
+		/// equal to the radial force limit. For motion off of the ground, e.g. jumping, set a
+		/// velocity or impulse on the character body directly.
+		/// This should be set every frame for a locomoting character, and the value is typically
+		/// controlled either from a character movement model or from root motion extracted from
+		/// an animation sequence.
+		void SetTargetDeltaPosition(const FVec3& Value)
+		{
+			ConstraintData.Modify(true, DirtyFlags, Proxy, [&Value](FCharacterGroundConstraintDynamicData& Data) {
+				Data.TargetDeltaPosition = Value;
+				});
+		}
+		FVec3 GetTargetDeltaPosition() const { return ConstraintData.Read().TargetDeltaPosition; }
+
+		/// Set the target orientation change about the vertical axis. This is used to turn the
+		/// character when it is on the ground. The solver will do this by applying a torque if the
+		/// character is considered to be on the ground and the torque required is less than or equal
+		/// to the torque limit. For orientation change off the ground, e.g. in air steering, set an
+		/// angular velocity or impulse on the character body directly.
+		/// This should be set every frame for a locomoting character, and the value is typically
+		/// controlled either from a character movement model or from root motion extracted from
+		/// an animation sequence.
+		void SetTargetDeltaFacing(const FReal& Value)
+		{
+			ConstraintData.Modify(true, DirtyFlags, Proxy, [&Value](FCharacterGroundConstraintDynamicData& Data) {
+				Data.TargetDeltaFacing = Value;
+				});
+		}
+		FReal GetTargetDeltaFacing() const { return ConstraintData.Read().TargetDeltaFacing; }
+
+		/// The maximum force that the character can apply in the plane defined by the ground normal
+		/// to move the character to its target position (or hold the character in place if no target
+		/// delta position is set).
+		void SetRadialForceLimit(const FReal& Value)
+		{
+			ConstraintSettings.Modify(true, DirtyFlags, Proxy, [&Value](FCharacterGroundConstraintSettings& Data) {
+				Data.RadialForceLimit = Value > FReal(0.0) ? Value : FReal(0.0);
+				});
+		}
+		FReal GetRadialForceLimit() const { return ConstraintSettings.Read().RadialForceLimit; }
+
+		/// The maximum torque that the character can apply about the vertical axis to rotate the
+		/// character to its target facing direction (or hold the character in place if no target
+		/// delta facing is set).
+		void SetTorqueLimit(const FReal& Value)
+		{
+			ConstraintSettings.Modify(true, DirtyFlags, Proxy, [&Value](FCharacterGroundConstraintSettings& Data) {
+				Data.TorqueLimit = Value > FReal(0.0) ? Value : FReal(0.0);
+				});
+		}
+		FReal GetTorqueLimit() const { return ConstraintSettings.Read().TorqueLimit; }
+
+		//////////////////////////////////////////////////////////////////////////
+		// Output
+
+		// Get the force applied by the solver due to this constraint. Units are ML/T^2
+		FVector GetSolverAppliedForce() const { return SolverAppliedForce; }
+
+		// Get the torque applied by the solver due to this constraint. Units are ML^2/T^2
+		FVector GetSolverAppliedTorque() const { return SolverAppliedTorque; }
+
+	private:
+		//////////////////////////////////////////////////////////////////////////
+		// FConstraintBase implementation
+		virtual void SyncRemoteDataImp(FDirtyPropertiesManager& Manager, int32 DataIdx, FDirtyChaosProperties& RemoteData) override;
+
+		TChaosProperty<FParticleProxyProperty, EChaosProperty::CharacterParticleProxy> CharacterProxy;
+		TChaosProperty<FParticleProxyProperty, EChaosProperty::GroundParticleProxy> GroundProxy;
+		TChaosProperty<FCharacterGroundConstraintSettings, EChaosProperty::CharacterGroundConstraintSettings> ConstraintSettings;
+		TChaosProperty<FCharacterGroundConstraintDynamicData, EChaosProperty::CharacterGroundConstraintDynamicData> ConstraintData;
+
+		friend class FCharacterGroundConstraintProxy; // for setting the force and torque from the physics thread data
+		FVector SolverAppliedForce;
+		FVector SolverAppliedTorque;
+	};
+
+} // namespace Chaos
\ No newline at end of file
diff --git a/Engine/Source/Runtime/Experimental/Chaos/Public/Chaos/Character/CharacterGroundConstraintContainer.h b/Engine/Source/Runtime/Experimental/Chaos/Public/Chaos/Character/CharacterGroundConstraintContainer.h
new file mode 100644
--- /dev/null
+++ b/Engine/Source/Runtime/Experimental/Chaos/Public/Chaos/Character/CharacterGroundConstraintContainer.h
@@ -0,0 +1,138 @@
+// Copyright Epic Games, Inc. All Rights Reserved.
+#pragma once
+
+#include "Chaos/Core.h"
+
+#include "Chaos/Character/CharacterGroundConstraintSettings.h"
+#include "Chaos/ConstraintHandle.h"
+#include "Chaos/Evolution/IndexedConstraintContainer.h"
+#include "Chaos/GeometryParticles.h"
+#include "Chaos/PBDConstraintContainer.h"
+
+namespace ChaosTest
+{
+	class CharacterGroundConstraintContainerTest;
+}
+
+namespace Chaos
+{
+	class FCharacterGroundConstraintContainer;
+	namespace Private
+	{
+		class FCharacterGroundConstraintContainerSolver;
+	}
+
+	/// Physics thread side representation of a character ground constraint
+	/// Data accessed through this class is synced from the game thread
+	/// representation FCharacterGroundConstraint
+	class CHAOS_API FCharacterGroundConstraintHandle final : public TIntrusiveConstraintHandle<FCharacterGroundConstraintHandle>
+	{
+	public:
+		using Base = TIntrusiveConstraintHandle<FCharacterGroundConstraintHandle>;
+		using FConstraintContainer = FCharacterGroundConstraintContainer;
+
+		FCharacterGroundConstraintHandle()
+			: TIntrusiveConstraintHandle<FCharacterGroundConstraintHandle>()
+		{
+		}
+
+		static const FConstraintHandleTypeID& StaticType()
+		{
+			static FConstraintHandleTypeID STypeID(TEXT("FCharacterGroundConstraintHandle"), &FIndexedConstraintHandle::StaticType());
+			return STypeID;
+		}
+
+		ESyncState GetSyncState() const { return SyncState; }
+		void SetSyncState(ESyncState InSyncState) { SyncState = InSyncState; }
+
+		EResimType GetResimType() const { return ResimType; }
+		bool GetEnabledDuringResim() const { return bEnabledDuringResim; }
+		void SetEnabledDuringResim(bool bEnabled) { bEnabledDuringResim = bEnabled; }
+
+		virtual void SetEnabled(bool InEnabled) override { bDisabled = !InEnabled; }
+		virtual bool IsEnabled() const { return !bDisabled; }
+
+		/// Currently settings and data are only modifiable on the game thread
+		/// via FCharacterGroundConstraint
+		const FCharacterGroundConstraintSettings& GetSettings() const { return Settings; }
+		const FCharacterGroundConstraintDynamicData& GetData() const { return Data; }
+
+		// Declared final so that TPBDConstraintGraphRuleImpl::AddToGraph() does not need to hit vtable
+		virtual FParticlePair GetConstrainedParticles() const override final { return { CharacterParticle, GroundParticle }; }
+		FGeometryParticleHandle* GetCharacterParticle() const { return CharacterParticle; }
+		FGeometryParticleHandle* GetGroundParticle() const { return GroundParticle; }
+
+		// Get the force applied by the solver due to this constraint. Units are ML/T^2
+		FVec3 GetSolverAppliedForce() const { return SolverAppliedForce; }
+
+		// Get the torque applied by the solver due to this constraint. Units are ML^2/T^2
+		FVec3 GetSolverAppliedTorque() const { return SolverAppliedTorque; }
+
+	private:
+		friend class FCharacterGroundConstraintContainer; // For creation
+		friend class FCharacterGroundConstraintProxy; // For setting the data from the game thread constraint
+		friend class Private::FCharacterGroundConstraintContainerSolver; // For setting the solver force and torque
+		friend class ChaosTest::CharacterGroundConstraintContainerTest; // For testing internals
+
+		FCharacterGroundConstraintSettings Settings;
+		FCharacterGroundConstraintDynamicData Data;
+		FVec3 SolverAppliedForce = FVec3::ZeroVector;
+		FVec3 SolverAppliedTorque = FVec3::ZeroVector;
+		FGeometryParticleHandle* CharacterParticle;
+		FGeometryParticleHandle* GroundParticle;
+		bool bDisabled = false;
+		bool bEnabledDuringResim;
+		EResimType ResimType = EResimType::FullResim;
+		ESyncState SyncState = ESyncState::InSync;
+	};
+
+	/// Container class for all character ground constraints on the physics thread
+	class CHAOS_API FCharacterGroundConstraintContainer : public FPBDConstraintContainer
+	{
+	public:
+		using Base = FPBDConstraintContainer;
+		using FConstraints = TArrayView<FCharacterGroundConstraintHandle* const>;
+		using FConstConstraints = TArrayView<const FCharacterGroundConstraintHandle* const>;
+
+		FCharacterGroundConstraintContainer();
+		virtual ~FCharacterGroundConstraintContainer();
+
+		int32 NumConstraints() const { return Constraints.Num(); }
+
+		FCharacterGroundConstraintHandle* AddConstraint(
+			const FCharacterGroundConstraintSettings& InConstraintSettings,
+			const FCharacterGroundConstraintDynamicData& InConstraintData,
+			FGeometryParticleHandle* CharacterParticle,
+			FGeometryParticleHandle* GroundParticle = nullptr);
+
+		void RemoveConstraint(FCharacterGroundConstraintHandle* Constraint);
+
+		FConstraints GetConstraints() { return MakeArrayView(Constraints); }
+		FConstConstraints GetConstConstraints() const { return MakeArrayView(Constraints); }
+
+		FCharacterGroundConstraintHandle* GetConstraint(int32 ConstraintIndex) { check(ConstraintIndex < NumConstraints()); return Constraints[ConstraintIndex]; }
+		const FCharacterGroundConstraintHandle* GetConstraint(int32 ConstraintIndex) const { check(ConstraintIndex < NumConstraints()); return Constraints[ConstraintIndex]; }
+
+		//////////////////////////////////////////////////////////////////////////
+		// FConstraintContainer Implementation
+		virtual TUniquePtr<FConstraintContainerSolver> CreateSceneSolver(const int32 Priority) override final;
+		virtual TUniquePtr<FConstraintContainerSolver> CreateGroupSolver(const int32 Priority) override final;
+		virtual int32 GetNumConstraints() const override final { return Constraints.Num(); }
+		virtual void ResetConstraints() override final {}
+		virtual void AddConstraintsToGraph(Private::FPBDIslandManager& IslandManager) override final;
+		virtual void PrepareTick() override final;
+		virtual void UnprepareTick() override final;
+		virtual void DisconnectConstraints(const TSet<TGeometryParticleHandle<FReal, 3>*>& RemovedParticles) override final;
+
+		//////////////////////////////////////////////////////////////////////////
+		// Required API from FConstraintContainer
+		void SetConstraintEnabled(int32 ConstraintIndex, bool bEnabled);
+		bool IsConstraintEnabled(int32 ConstraintIndex) const { check(ConstraintIndex < NumConstraints()); return Constraints[ConstraintIndex]->IsEnabled(); }
+
+	private:
+		bool CanEvaluate(const FCharacterGroundConstraintHandle* Constraint) const;
+
+		TObjectPool<FCharacterGroundConstraintHandle> ConstraintPool;
+		TArray<FCharacterGroundConstraintHandle*> Constraints;
+	};
+}
\ No newline at end of file
diff --git a/Engine/Source/Runtime/Experimental/Chaos/Public/Chaos/Framework/PhysicsProxyBase.h b/Engine/Source/Runtime/Experimental/Chaos/Public/Chaos/Framework/PhysicsProxyBase.h
--- a/Engine/Source/Runtime/Experimental/Chaos/Public/Chaos/Framework/PhysicsProxyBase.h
+++ b/Engine/Source/Runtime/Experimental/Chaos/Public/Chaos/Framework/PhysicsProxyBase.h
@@ -15,6 +15,7 @@ enum class EPhysicsProxyType : uint32
 	SkeletalMeshType = 4,
 	JointConstraintType = 8,	//left gap when removed some types in case these numbers actually matter to someone, should remove
 	SuspensionConstraintType = 9,
+	CharacterGroundConstraintType = 10,
 	SingleParticleProxy,
 	Count
 };
diff --git a/Engine/Source/Runtime/Experimental/Chaos/Public/Chaos/PBDConstraintBaseData.h b/Engine/Source/Runtime/Experimental/Chaos/Public/Chaos/PBDConstraintBaseData.h
--- a/Engine/Source/Runtime/Experimental/Chaos/Public/Chaos/PBDConstraintBaseData.h
+++ b/Engine/Source/Runtime/Experimental/Chaos/Public/Chaos/PBDConstraintBaseData.h
@@ -20,7 +20,8 @@ enum class EConstraintType
 	NoneType = 0,
 	JointConstraintType = 1,
 	SpringConstraintType = 2,
-	SuspensionConstraintType = 3
+	SuspensionConstraintType = 3,
+	CharacterGroundConstraintType = 4,
 };
 
 #define CONSTRAINT_JOINT_PROPERPETY_IMPL(TYPE, FNAME, ENAME, VNAME)\
diff --git a/Engine/Source/Runtime/Experimental/Chaos/Public/Chaos/PBDRigidsEvolution.h b/Engine/Source/Runtime/Experimental/Chaos/Public/Chaos/PBDRigidsEvolution.h
--- a/Engine/Source/Runtime/Experimental/Chaos/Public/Chaos/PBDRigidsEvolution.h
+++ b/Engine/Source/Runtime/Experimental/Chaos/Public/Chaos/PBDRigidsEvolution.h
@@ -1,6 +1,7 @@
 // Copyright Epic Games, Inc. All Rights Reserved.
 #pragma once
 
+#include "Chaos/Character/CharacterGroundConstraintContainer.h"
 #include "Chaos/PBDCollisionConstraints.h"
 #include "Chaos/PBDConstraintGraph.h"
 #include "Chaos/PBDRigidClustering.h"
@@ -596,14 +597,18 @@ public:
 			{
 				RemoveConstraintFromConstraintGraph(ConstraintHandle);
 			}
+			else if (FCharacterGroundConstraintHandle* CharacterGroundConstraint = BaseConstraintHandle->As<FCharacterGroundConstraintHandle>())
+			{
+				RemoveConstraintFromConstraintGraph(CharacterGroundConstraint);
+			}
 		}
 	}
 
 	/** 
 	* Disconnect constraints from a set of particles to be destroyed. 
-	* this will set the constraints to Enbaled = false and set their respective bodies handles to nullptr.
+	* this will set the constraints to Enabled = false and set their respective bodies handles to nullptr.
 	* Once this is done, the constraints cannot be re-enabled.
-	* @note This only applies to persistent constraints (joints etc), not transient constraints (collisons)
+	* @note This only applies to persistent constraints (joints etc), not transient constraints (collisions)
 	* @see DestroyTransientConstraints()
 	*/
 	CHAOS_API void DisconnectConstraints(const TSet<FGeometryParticleHandle*>& RemovedParticles)
@@ -623,7 +628,7 @@ public:
 	* Disconnect constraints from a particle to be removed (or destroyed)
 	* this will set the constraints to Enabled = false, but leave connections to the particles to support
 	* re-enabling at a later time.
-	* @note This only applies to persistent constraints (joints etc), not transient constraints (collisons)
+	* @note This only applies to persistent constraints (joints etc), not transient constraints (collisions)
 	* @see DestroyTransientConstraints()
 	*/
 	CHAOS_API void DisableConstraints(FGeometryParticleHandle* ParticleHandle)
diff --git a/Engine/Source/Runtime/Experimental/Chaos/Public/Chaos/PBDRigidsEvolutionGBF.h b/Engine/Source/Runtime/Experimental/Chaos/Public/Chaos/PBDRigidsEvolutionGBF.h
--- a/Engine/Source/Runtime/Experimental/Chaos/Public/Chaos/PBDRigidsEvolutionGBF.h
+++ b/Engine/Source/Runtime/Experimental/Chaos/Public/Chaos/PBDRigidsEvolutionGBF.h
@@ -2,6 +2,7 @@
 #pragma once
 
 #include "Chaos/ChaosPerfTest.h"
+#include "Chaos/Character/CharacterGroundConstraintContainer.h"
 #include "Chaos/Collision/SpatialAccelerationBroadPhase.h"
 #include "Chaos/Collision/SpatialAccelerationCollisionDetector.h"
 #include "Chaos/Evolution/SolverBodyContainer.h"
@@ -119,6 +120,9 @@ namespace Chaos
 		FORCEINLINE FPBDSuspensionConstraints& GetSuspensionConstraints() { return SuspensionConstraints; }
 		FORCEINLINE const FPBDSuspensionConstraints& GetSuspensionConstraints() const { return SuspensionConstraints; }
 
+		FORCEINLINE FCharacterGroundConstraintContainer& GetCharacterGroundConstraints() { return CharacterGroundConstraints; }
+		FORCEINLINE const FCharacterGroundConstraintContainer& GetCharacterGroundConstraints() const { return CharacterGroundConstraints; }
+
 
 		//
 		// Particle API (most of the particle API is in the base class)
@@ -347,6 +351,7 @@ namespace Chaos
 
 		FPBDJointConstraints JointConstraints;
 		FPBDSuspensionConstraints SuspensionConstraints;
+		FCharacterGroundConstraintContainer CharacterGroundConstraints;
 
 		FGravityForces GravityForces;
 		FCollisionConstraints CollisionConstraints;
diff --git a/Engine/Source/Runtime/Experimental/Chaos/Public/Chaos/ParticleDirtyFlags.h b/Engine/Source/Runtime/Experimental/Chaos/Public/Chaos/ParticleDirtyFlags.h
--- a/Engine/Source/Runtime/Experimental/Chaos/Public/Chaos/ParticleDirtyFlags.h
+++ b/Engine/Source/Runtime/Experimental/Chaos/Public/Chaos/ParticleDirtyFlags.h
@@ -4,6 +4,7 @@
 
 #include "Chaos/Core.h"
 #include "Chaos/ChaosArchive.h"
+#include "Chaos/Character/CharacterGroundConstraintSettings.h"
 #include "Chaos/Box.h"
 #include "Chaos/Particles.h"
 #include "Chaos/PhysicalMaterials.h"
diff --git a/Engine/Source/Runtime/Experimental/Chaos/Public/Chaos/ParticleProperties.inl b/Engine/Source/Runtime/Experimental/Chaos/Public/Chaos/ParticleProperties.inl
--- a/Engine/Source/Runtime/Experimental/Chaos/Public/Chaos/ParticleProperties.inl
+++ b/Engine/Source/Runtime/Experimental/Chaos/Public/Chaos/ParticleProperties.inl
@@ -25,4 +25,9 @@ CHAOS_PROPERTY(SuspensionSettings, FPBDSuspensionSettings, EPhysicsProxyType::Su
 CHAOS_PROPERTY(SuspensionParticleProxy, FParticleProxyProperty, EPhysicsProxyType::SuspensionConstraintType)
 CHAOS_PROPERTY(SuspensionLocation, FSuspensionLocation, EPhysicsProxyType::SuspensionConstraintType)
 
+CHAOS_PROPERTY(CharacterGroundConstraintSettings, FCharacterGroundConstraintSettings, EPhysicsProxyType::CharacterGroundConstraintType)
+CHAOS_PROPERTY(CharacterGroundConstraintDynamicData, FCharacterGroundConstraintDynamicData, EPhysicsProxyType::CharacterGroundConstraintType)
+CHAOS_PROPERTY(CharacterParticleProxy, FParticleProxyProperty, EPhysicsProxyType::CharacterGroundConstraintType)
+CHAOS_PROPERTY(GroundParticleProxy, FParticleProxyProperty, EPhysicsProxyType::CharacterGroundConstraintType)
+
 #undef CHAOS_PROPERTY
\ No newline at end of file
diff --git a/Engine/Source/Runtime/Experimental/Chaos/Public/Chaos/PhysicsSolverBaseImpl.h b/Engine/Source/Runtime/Experimental/Chaos/Public/Chaos/PhysicsSolverBaseImpl.h
--- a/Engine/Source/Runtime/Experimental/Chaos/Public/Chaos/PhysicsSolverBaseImpl.h
+++ b/Engine/Source/Runtime/Experimental/Chaos/Public/Chaos/PhysicsSolverBaseImpl.h
@@ -8,6 +8,7 @@
 
 #include "Chaos/Framework/PhysicsSolverBase.h"
 #include "Chaos/PullPhysicsDataImp.h"
+#include "PhysicsProxy/CharacterGroundConstraintProxy.h"
 #include "PhysicsProxy/SingleParticlePhysicsProxy.h"
 #include "PhysicsProxy/GeometryCollectionPhysicsProxy.h"
 #include "PhysicsProxy/JointConstraintProxy.h"
@@ -129,6 +130,16 @@ namespace Chaos
 
 			//latest data may be used multiple times during interpolation, so for non interpolated joints we clear it
 			LatestData->DirtyJointConstraints.Reset();
+
+			for (const FDirtyCharacterGroundConstraintData& DirtyData : LatestData->DirtyCharacterGroundConstraints)
+			{
+				if (auto Proxy = DirtyData.GetProxy())
+				{
+					// todo: Do we need to add a callback?
+					Proxy->PullFromPhysicsState(DirtyData, SyncTimestamp);
+				}
+			}
+			LatestData->DirtyCharacterGroundConstraints.Reset();
 		}
 	}
 
diff --git a/Engine/Source/Runtime/Experimental/Chaos/Public/Chaos/PullPhysicsDataImp.h b/Engine/Source/Runtime/Experimental/Chaos/Public/Chaos/PullPhysicsDataImp.h
--- a/Engine/Source/Runtime/Experimental/Chaos/Public/Chaos/PullPhysicsDataImp.h
+++ b/Engine/Source/Runtime/Experimental/Chaos/Public/Chaos/PullPhysicsDataImp.h
@@ -12,6 +12,8 @@
 
 namespace Chaos
 {
+
+class FCharacterGroundConstraintProxy;
 class FJointConstraintPhysicsProxy;
 
 template <typename TProxy, typename TTimeStamp>
@@ -75,6 +77,12 @@ struct FDirtyJointConstraintData : public TBasePullData<FJointConstraintPhysicsP
 	FJointConstraintOutputData OutputData;
 };
 
+struct FDirtyCharacterGroundConstraintData : public TBasePullData<FCharacterGroundConstraintProxy, FProxyTimestampBase>
+{
+	FVector Force = FVector(0);
+	FVector Torque = FVector(0);
+};
+
 //A simulation frame's result of dirty particles. These are all the particles that were dirtied in this particular sim step
 class FPullPhysicsData
 {
@@ -82,6 +90,7 @@ public:
 	TArray<FDirtyRigidParticleData> DirtyRigids;
 	TArray<FDirtyGeometryCollectionData> DirtyGeometryCollections;
 	TArray<FDirtyJointConstraintData> DirtyJointConstraints;
+	TArray<FDirtyCharacterGroundConstraintData> DirtyCharacterGroundConstraints;
 
 	int32 SolverTimestamp;
 	FReal ExternalStartTime;	//The start time associated with this result. The time is synced using the external time
diff --git a/Engine/Source/Runtime/Experimental/Chaos/Public/PBDRigidsSolver.h b/Engine/Source/Runtime/Experimental/Chaos/Public/PBDRigidsSolver.h
--- a/Engine/Source/Runtime/Experimental/Chaos/Public/PBDRigidsSolver.h
+++ b/Engine/Source/Runtime/Experimental/Chaos/Public/PBDRigidsSolver.h
@@ -53,6 +53,8 @@ namespace Chaos
 {
 	class FPersistentPhysicsTask;
 	class FChaosArchive;
+	class FCharacterGroundConstraint;
+	class FCharacterGroundConstraintProxy;
 	class FSingleParticleProxy;
 	class FGeometryParticleBuffer;
 
@@ -135,6 +137,9 @@ namespace Chaos
 		void RegisterObject(Chaos::FSuspensionConstraint* GTConstraint);
 		void UnregisterObject(Chaos::FSuspensionConstraint* GTConstraint);
 
+		void RegisterObject(Chaos::FCharacterGroundConstraint* GTConstraint);
+		void UnregisterObject(Chaos::FCharacterGroundConstraint* GTConstraint);
+
 		//
 		//  Simulation API
 		//
@@ -191,12 +196,12 @@ namespace Chaos
 		/**
 		 * @brief True if the simulation is running in deterministic mode
 		 * This will be true if determinism is explicitly requested (via SetIsDeterministic()) or if required
-		 * by some otyher system like Rewind/Resim support.
+		 * by some other system like Rewind/Resim support.
 		*/
 		bool IsDetemerministic() const;
 
 		/**
-		 * @brief Request that the sim be determinitic (or not)
+		 * @brief Request that the sim be deterministic (or not)
 		 * @note Even if set to false, the sim may still be deterministic if some other feature is enabled and requires it.
 		 * @see IsDetemerministic()
 		*/
@@ -210,6 +215,9 @@ namespace Chaos
 		const FPBDSuspensionConstraints& GetSuspensionConstraints() const { return MEvolution->GetSuspensionConstraints(); }
 		void SetSuspensionTarget(Chaos::FSuspensionConstraint* GTConstraint, const FVector& TargetPos, const FVector& Normal, bool Enabled);
 
+		FCharacterGroundConstraintContainer& GetCharacterGroundConstraints() { return MEvolution->GetCharacterGroundConstraints(); }
+		const FCharacterGroundConstraintContainer& GetCharacterGroundConstraints() const { return MEvolution->GetCharacterGroundConstraints(); }
+
 		void EnableRewindCapture(int32 NumFrames, bool InUseCollisionResimCache, TUniquePtr<IRewindCallback>&& RewindCallback = TUniquePtr<IRewindCallback>());
 
 		/**/
@@ -347,6 +355,7 @@ namespace Chaos
 		TSparseArray< FSingleParticlePhysicsProxy* > SingleParticlePhysicsProxies_PT;
 		TArray< FGeometryCollectionPhysicsProxy* > GeometryCollectionPhysicsProxies_Internal; // PT
 		TArray< FJointConstraintPhysicsProxy* > JointConstraintPhysicsProxies_Internal; // PT
+		TArray< FCharacterGroundConstraintProxy* > CharacterGroundConstraintProxies_Internal;
 
 		TUniquePtr<FPerSolverFieldSystem> PerSolverField;
 
diff --git a/Engine/Source/Runtime/Experimental/Chaos/Public/PhysicsProxy/CharacterGroundConstraintProxy.h b/Engine/Source/Runtime/Experimental/Chaos/Public/PhysicsProxy/CharacterGroundConstraintProxy.h
new file mode 100644
--- /dev/null
+++ b/Engine/Source/Runtime/Experimental/Chaos/Public/PhysicsProxy/CharacterGroundConstraintProxy.h
@@ -0,0 +1,71 @@
+// Copyright Epic Games, Inc. All Rights Reserved.
+#pragma once
+
+#include "Chaos/Framework/PhysicsProxyBase.h"
+
+namespace Chaos
+{
+	class FCharacterGroundConstraint;
+	class FCharacterGroundConstraintHandle;
+	class FDirtyChaosProperties;
+	struct FDirtyCharacterGroundConstraintData;
+	class FDirtyPropertiesManager;
+	class FPBDRigidsSolver;
+
+	/// Proxy class to manage access to and syncing of character ground constraint data
+	/// between the game thread and physics thread
+	class CHAOS_API FCharacterGroundConstraintProxy : public IPhysicsProxyBase
+	{
+	public:
+		using Base = IPhysicsProxyBase;
+
+		FCharacterGroundConstraintProxy(FCharacterGroundConstraint* InConstraintGT, FCharacterGroundConstraintHandle* InConstraintPT = nullptr, UObject* InOwner = nullptr);
+	
+		//////////////////////////////////////////////////////////////////////////
+		/// Member Access
+
+		FCharacterGroundConstraint* GetGameThreadAPI() { return Constraint_GT; }
+		const FCharacterGroundConstraint* GetGameThreadAPI() const { return Constraint_GT; }
+
+		FCharacterGroundConstraintHandle* GetPhysicsThreadAPI() { return Constraint_PT; }
+		const FCharacterGroundConstraintHandle* GetPhysicsThreadAPI() const { return Constraint_PT; }
+
+		//////////////////////////////////////////////////////////////////////////
+		/// State Management
+
+		/// Gets the data from the game thread via the RemoteData and creates a constraint on the physics thread
+		void InitializeOnPhysicsThread(FPBDRigidsSolver* InSolver, FDirtyPropertiesManager& Manager, int32 DataIdx, FDirtyChaosProperties& RemoteData);
+
+		/// Pushes any changed data from the game thread constraint to the remote data
+		void PushStateOnGameThread(FDirtyPropertiesManager& Manager, int32 DataIdx, FDirtyChaosProperties& RemoteData);
+
+		/// Reads changed data from the game thread constraint into the physics thread constraint via the remote data
+		void PushStateOnPhysicsThread(FPBDRigidsSolver* InSolver, const FDirtyPropertiesManager& Manager, int32 DataIdx, const FDirtyChaosProperties& RemoteData);
+
+		/// Deletes the game thread constraint
+		void DestroyOnGameThread();
+
+		/// Removes references to the physics thread constraint and deletes it
+		void DestroyOnPhysicsThread(FPBDRigidsSolver* InSolver);
+
+		//////////////////////////////////////////////////////////////////////////
+		/// IPhysicsProxyBase Implementation
+
+		virtual void* GetHandleUnsafe() const override { return Constraint_PT; }
+
+		//////////////////////////////////////////////////////////////////////////
+		/// Manage Output Data
+
+		/// Write the output constraint data from the physics thread to the buffer
+		void BufferPhysicsResults(FDirtyCharacterGroundConstraintData& Buffer);
+
+		/// Write the output constraint data from the buffer to the game thread 
+		bool PullFromPhysicsState(const FDirtyCharacterGroundConstraintData& Buffer, const int32 SolverSyncTimestamp);
+
+	private:
+		FCharacterGroundConstraint* Constraint_GT;
+		FCharacterGroundConstraintHandle* Constraint_PT;
+		bool bInitialized = false;
+	};
+
+} // namespace Chaos
\ No newline at end of file
