Add a unit test for kinematic in graph

diff --git a/Engine/Source/Programs/HeadlessChaos/Private/HeadlessChaosTestGraphEvolution.cpp b/Engine/Source/Programs/HeadlessChaos/Private/HeadlessChaosTestGraphEvolution.cpp
--- a/Engine/Source/Programs/HeadlessChaos/Private/HeadlessChaosTestGraphEvolution.cpp
+++ b/Engine/Source/Programs/HeadlessChaos/Private/HeadlessChaosTestGraphEvolution.cpp
@@ -903,6 +903,33 @@ namespace ChaosTest
 		EXPECT_TRUE(Test.ConstraintHandles[0]->IsSleeping());
 	}
 
+	// Test dynamic particle to be removed from island if no constraints
+	GTEST_TEST(GraphEvolutionTests, TestConstraintGraph_KinematicRemoveFromGraph)
+	{
+		FGraphEvolutionTest Test(3);
+
+		Test.Advance();
+		EXPECT_EQ(Test.IslandManager->NumIslands(), 3);
+		EXPECT_EQ(Test.IslandManager->GetIslandGraph()->ItemNodes.Num(), 3);
+		EXPECT_TRUE(Test.IslandManager->GetIslandGraph()->GraphNodes.IsValidIndex(1));
+
+		Test.Evolution.SetParticleObjectState(Test.ParticleHandles[1], EObjectStateType::Kinematic);
+
+		Test.Advance();
+		EXPECT_EQ(Test.IslandManager->NumIslands(), 2);
+		EXPECT_EQ(Test.IslandManager->GetIslandGraph()->ItemNodes.Num(), 3);
+		// Kinematic is still in a graph, and will be updated
+		EXPECT_TRUE(Test.IslandManager->GetIslandGraph()->GraphNodes.IsValidIndex(1));
+		// The two dynamics body are also there
+		EXPECT_TRUE(Test.IslandManager->GetIslandGraph()->GraphNodes.IsValidIndex(0));
+		EXPECT_TRUE(Test.IslandManager->GetIslandGraph()->GraphNodes.IsValidIndex(2));
+
+		Test.Advance();
+		EXPECT_EQ(Test.IslandManager->GetIslandGraph()->ItemNodes.Num(), 3);
+		EXPECT_TRUE(Test.IslandManager->GetIslandGraph()->GraphNodes.IsValidIndex(1));
+
+	}
+
 	// Test the conditions for a kinematic particle waking an island
 	// If a kinematic is being animated by velocity or by setting a target
 	// position the island should wake but only if the target velocity is
