Fix Capsule Overlap

diff --git a/Engine/Source/Programs/HeadlessChaos/Private/HeadlessChaosTestOverlap.cpp b/Engine/Source/Programs/HeadlessChaos/Private/HeadlessChaosTestOverlap.cpp
--- a/Engine/Source/Programs/HeadlessChaos/Private/HeadlessChaosTestOverlap.cpp
+++ b/Engine/Source/Programs/HeadlessChaos/Private/HeadlessChaosTestOverlap.cpp
@@ -6,6 +6,7 @@
 #include "Chaos/Box.h"
 #include "Chaos/Capsule.h"
 #include "Chaos/Sphere.h"
+#include "Chaos/GeometryQueries.h"
 #include "Chaos/TriangleMeshImplicitObject.h"
 #include "Chaos/ImplicitObject.h"
 #include "Chaos/ImplicitObjectScaled.h"
@@ -130,6 +131,29 @@ namespace ChaosTest
 					EXPECT_TRUE(bResult);
 				}
 			}
+			{
+				const FVec3 X1 = { 10.0, -1.0, 0.0 };
+				const FVec3 X2 = { 11.0, 1.0, 0.0 };
+				const FReal Radius = 0.1;
+				const FCapsule Capsule = FCapsule(X1, X2, Radius);
+				const FRigidTransform3 QueryTM(FVec3(0.0, 0.0, 0.0), FQuat::Identity);
+				bool bResult = TriangleMesh->OverlapGeom(Capsule, QueryTM, 0.0);
+				EXPECT_FALSE(bResult);
+			}
+			{
+				const FVec3 X1 = { -11.0, 2.0, 0.0 };
+				const FVec3 X2 = { -9.0, -8.0, 0.0 };
+				const FReal Radius = 0.01;
+				const FCapsule Capsule = FCapsule(X1, X2, Radius);
+				FMTDInfo MTDInfo;
+				{
+					const FRigidTransform3 QueryTM(FVec3(0.0, 0.0, 0.0), FQuat::Identity);
+					bool bResult = TriangleMesh->OverlapGeom(Capsule, QueryTM, 0.0);
+					EXPECT_FALSE(bResult);
+					bool bResultMTD = TriangleMesh->OverlapGeom(Capsule, QueryTM, 0.0, &MTDInfo);
+					EXPECT_EQ(bResult, bResultMTD);
+				}
+			}
 			{
 				// Sphere test
 				const FVec3 X = { 0.0, 0.0, 0.0 };
@@ -1077,5 +1101,249 @@ namespace ChaosTest
 
 			}
 		}
+		{
+			FTriangleMeshImplicitObject::ParticlesType TrimeshParticles(
+				{
+					{100.0, 0.0, 0.0},
+					{86.0, -75.0, 0.0},
+					{50.0, -130.0, 0.0},
+				});
+
+			TArray<TVec3<int32>> Indices(
+				{
+					{0, 1, 2},
+				});
+
+			TArray<uint16> Materials;
+			for (int32 i = 0; i < Indices.Num(); ++i)
+			{
+				Materials.Emplace(0);
+			}
+
+			TUniquePtr<FTriangleMeshImplicitObject> TriangleMesh = MakeUnique<FTriangleMeshImplicitObject>(MoveTemp(TrimeshParticles), MoveTemp(Indices), MoveTemp(Materials));
+			
+				// Capsule test
+			{
+				const FVec3 X1 = { 66.0, -99.0, 0.0 };
+				const FVec3 X2 = { 67.0, -98.0, 0.0 };
+				const FReal Radius = 1.0;
+				const FCapsule Capsule = FCapsule(X1, X2, Radius);
+				const FRigidTransform3 QueryTM(FVec3(0.0, 0.0, 0.0), FQuat::Identity);
+				bool bResult = TriangleMesh->OverlapGeom(Capsule, QueryTM, 0.0);
+				EXPECT_TRUE(bResult);
+			}
+			{
+				const FVec3 X1 = { 80.0, -92.0, 0.0 };
+				const FVec3 X2 = { 85.0, -88.0, 0.0 };
+				const FReal Radius = 1.0;
+				const FCapsule Capsule = FCapsule(X1, X2, Radius);
+				const FRigidTransform3 QueryTM(FVec3(0.0, 0.0, 0.0), FQuat::Identity);
+				bool bResult = TriangleMesh->OverlapGeom(Capsule, QueryTM, 0.0);
+				EXPECT_FALSE(bResult);
+			}
+			{
+				const FVec3 X1 = { 40.0, -127.0, 0.0 };
+				const FVec3 X2 = { 55.0, -137.0, 0.0 };
+				const FReal Radius = 1.0;
+				const FCapsule Capsule = FCapsule(X1, X2, Radius);
+				const FRigidTransform3 QueryTM(FVec3(0.0, 0.0, 0.0), FQuat::Identity);
+				bool bResult = TriangleMesh->OverlapGeom(Capsule, QueryTM, 0.0);
+				EXPECT_FALSE(bResult);
+			}
+			{
+				const FVec3 X1 = { 0.0, -110.0, 0.0 };
+				const FVec3 X2 = { 100.0, -160.0, 0.0 };
+				const FReal Radius = 1.0;
+				const FCapsule Capsule = FCapsule(X1, X2, Radius);
+				const FRigidTransform3 QueryTM(FVec3(0.0, 0.0, 0.0), FQuat::Identity);
+				bool bResult = TriangleMesh->OverlapGeom(Capsule, QueryTM, 0.0);
+				EXPECT_FALSE(bResult);
+			}
+			{
+				const FVec3 X1 = { 56.0, -138.0, 0.0 };
+				const FVec3 X2 = { 55.0, -137.0, 0.0 };
+				const FReal Radius = 1.0;
+				const FCapsule Capsule = FCapsule(X1, X2, Radius);
+				const FRigidTransform3 QueryTM(FVec3(0.0, 0.0, 0.0), FQuat::Identity);
+				bool bResult = TriangleMesh->OverlapGeom(Capsule, QueryTM, 0.0);
+				EXPECT_FALSE(bResult);
+			}
+			{
+				const FVec3 X1 = { 48.0, -130.0, 0.0 };
+				const FVec3 X2 = { 47.0, -131.0, 0.0 };
+				const FReal Radius = 1.0;
+				const FCapsule Capsule = FCapsule(X1, X2, Radius);
+				const FRigidTransform3 QueryTM(FVec3(0.0, 0.0, 0.0), FQuat::Identity);
+				bool bResult = TriangleMesh->OverlapGeom(Capsule, QueryTM, 0.0);
+				EXPECT_FALSE(bResult);
+			}
+			{
+				const FVec3 X1 = { 101.5, 0.0, 0.0 };
+				const FVec3 X2 = { 106.0, 0.0, 0.0 };
+				const FReal Radius = 1.0;
+				const FCapsule Capsule = FCapsule(X1, X2, Radius);
+				const FRigidTransform3 QueryTM(FVec3(0.0, 0.0, 0.0), FQuat::Identity);
+				bool bResult = TriangleMesh->OverlapGeom(Capsule, QueryTM, 0.0);
+				EXPECT_FALSE(bResult);
+			}
+			{
+				const FVec3 X1 = { 95.0, 0.0, 0.0 };
+				const FVec3 X2 = { 98.0, 0.0, 0.0 };
+				const FReal Radius = 1.0;
+				const FCapsule Capsule = FCapsule(X1, X2, Radius);
+				const FRigidTransform3 QueryTM(FVec3(0.0, 0.0, 0.0), FQuat::Identity);
+				bool bResult = TriangleMesh->OverlapGeom(Capsule, QueryTM, 0.0);
+				EXPECT_FALSE(bResult);
+			}
+			{
+				const FVec3 X1 = { 95.0, 0.0, 0.0 };
+				const FVec3 X2 = { 94, 0.0, 0.0 };
+				const FReal Radius = 1.0;
+				const FCapsule Capsule = FCapsule(X1, X2, Radius);
+				const FRigidTransform3 QueryTM(FVec3(0.0, 0.0, 0.0), FQuat::Identity);
+				bool bResult = TriangleMesh->OverlapGeom(Capsule, QueryTM, 0.0);
+				EXPECT_FALSE(bResult);
+			}
+			{
+				// Sphere test
+				const FVec3 X = { 66.0, -99.0, 0.0 };
+				const FReal Radius = 1.0;
+				const Chaos::FSphere Sphere = Chaos::FSphere(X, Radius);
+				{
+					const FRigidTransform3 QueryTM(FVec3(0.0, 0.0, 0.0), FQuat::Identity);
+					bool bResult = TriangleMesh->OverlapGeom(Sphere, QueryTM, 0.0);
+					// EXPECT_TRUE(bResult); // TODO Enable test
+				}
+			}
+			{
+				const TBox<FReal, 3> Box = TBox<FReal, 3>({ -0.5, -0.5, -0.5 }, { 0.5, 0.5, 0.5 });
+				{
+					const FRigidTransform3 QueryTM(FVec3(66.0, -99.0, 0.0), FQuat::Identity);
+					bool bResult = TriangleMesh->OverlapGeom(Box, QueryTM, 0.0);
+					EXPECT_TRUE(bResult);
+				}
+			}
+		}
+		{
+			FTriangleMeshImplicitObject::ParticlesType TrimeshParticles(
+				{
+					{100.0, 0.0, 0.0},
+					{86.0, -75.0, 0.0},
+					{60.0, -110.0, 0.0},
+				});
+
+			TArray<TVec3<int32>> Indices(
+				{
+					{0, 1, 2},
+				});
+
+			TArray<uint16> Materials;
+			for (int32 i = 0; i < Indices.Num(); ++i)
+			{
+				Materials.Emplace(0);
+			}
+
+			TUniquePtr<FTriangleMeshImplicitObject> TriangleMesh = MakeUnique<FTriangleMeshImplicitObject>(MoveTemp(TrimeshParticles), MoveTemp(Indices), MoveTemp(Materials));
+			{
+				// Capsule test
+				const FVec3 X1 = { 66.0, -99.0, 0.0 };
+				const FVec3 X2 = { 67.0, -98.0, 0.0 };
+				const FReal Radius = 1.0;
+				const FCapsule Capsule = FCapsule(X1, X2, Radius);
+				{
+					const FRigidTransform3 QueryTM(FVec3(0.0, 0.0, 0.0), FQuat::Identity);
+					bool bResult = TriangleMesh->OverlapGeom(Capsule, QueryTM, 0.0);
+					EXPECT_TRUE(bResult);
+				}
+			}
+			{
+				// Sphere test
+				const FVec3 X = { 66.0, -99.0, 0.0 };
+				const FReal Radius = 1.0;
+				const Chaos::FSphere Sphere = Chaos::FSphere(X, Radius);
+				{
+					const FRigidTransform3 QueryTM(FVec3(0.0, 0.0, 0.0), FQuat::Identity);
+					bool bResult = TriangleMesh->OverlapGeom(Sphere, QueryTM, 0.0);
+					// EXPECT_TRUE(bResult); // TODO Enable test
+				}
+			}
+			{
+				const TBox<FReal, 3> Box = TBox<FReal, 3>({ -0.5, -0.5, -0.5 }, { 0.5, 0.5, 0.5 });
+				{
+					const FRigidTransform3 QueryTM(FVec3(66.0, -99.0, 0.0), FQuat::Identity);
+					bool bResult = TriangleMesh->OverlapGeom(Box, QueryTM, 0.0);
+					EXPECT_TRUE(bResult);
+				}
+			}
+		}
+		{
+			FTriangleMeshImplicitObject::ParticlesType TrimeshParticles(
+				{
+					{0.0, 0.0, 0.0},
+					{-200.0, 100.0, 0.0},
+					{-200.0, -100.0, 0.0},
+				});
+
+			TArray<TVec3<int32>> Indices(
+				{
+					{0, 1, 2},
+				});
+
+			TArray<uint16> Materials;
+			for (int32 i = 0; i < Indices.Num(); ++i)
+			{
+				Materials.Emplace(0);
+			}
+
+			TUniquePtr<FTriangleMeshImplicitObject> TriangleMesh = MakeUnique<FTriangleMeshImplicitObject>(MoveTemp(TrimeshParticles), MoveTemp(Indices), MoveTemp(Materials));
+			{
+				// Capsule test
+				const FVec3 X1 = { 0.0, -100.0, 100.0 };
+				FVec3 X2 = { 0.0, 100.0, 100.0 };
+				const FReal Radius = 101.0;
+				const FCapsule Capsule = FCapsule(X1, X2, Radius);
+				{
+					const FRigidTransform3 QueryTM(FVec3(0.0, 0.0, 0.0), FQuat::Identity);
+					bool bResult = TriangleMesh->OverlapGeom(Capsule, QueryTM, 0.0);
+					EXPECT_TRUE(bResult);
+				}
+			}
+		}
+		{
+			FTriangleMeshImplicitObject::ParticlesType TrimeshParticles(
+				{
+					{0.0, 0.0, 0.0},
+					{100, -100, 0.0},
+					{-200.0, -400.0, 0.0},
+				});
+
+			TArray<TVec3<int32>> Indices(
+				{
+					{0, 1, 2},
+				});
+
+			TArray<uint16> Materials;
+			for (int32 i = 0; i < Indices.Num(); ++i)
+			{
+				Materials.Emplace(0);
+			}
+
+			TUniquePtr<FTriangleMeshImplicitObject> TriangleMesh = MakeUnique<FTriangleMeshImplicitObject>(MoveTemp(TrimeshParticles), MoveTemp(Indices), MoveTemp(Materials));
+			{
+				{
+					// Capsule test
+					const FVec3 X1 = { -100, 0, 0 };
+					const FVec3 X2 = { -200, 0, 0 };
+					const FReal Radius = 99.0f;
+					// This capsule will not touch a vertex, but will intersect an edge
+					const FCapsule Capsule = FCapsule(X1, X2, Radius); 
+					{
+						const FRigidTransform3 QueryTM(FVec3(0.0, 0.0, 0.0), FQuat::Identity);
+						bool bResult = TriangleMesh->OverlapGeom(Capsule, QueryTM, 0.0);
+						EXPECT_TRUE(bResult);
+					}
+				}
+			}
+		}
 	}
 }
diff --git a/Engine/Source/Runtime/Experimental/Chaos/Private/Chaos/Collision/TriangleOverlap.cpp b/Engine/Source/Runtime/Experimental/Chaos/Private/Chaos/Collision/TriangleOverlap.cpp
--- a/Engine/Source/Runtime/Experimental/Chaos/Private/Chaos/Collision/TriangleOverlap.cpp
+++ b/Engine/Source/Runtime/Experimental/Chaos/Private/Chaos/Collision/TriangleOverlap.cpp
@@ -35,8 +35,8 @@ namespace Chaos
 		// Edge plane normals and signed distances to each segment point
 		// Mask to compare the distance of the triangle edge plane to the capsule radius.
 		constexpr static int32 AllEdges = 0b1111;
+		const VectorRegister4Float PlaneCA = VectorNormalize(VectorCross(CA, Normal));
 		{
-			VectorRegister4Float PlaneCA = VectorNormalize(VectorCross(CA, Normal));
 			FRealSingle EdgeX1 = VectorDot3Scalar(AX1, PlaneCA);
 			FRealSingle EdgeX2 = VectorDot3Scalar(AX2, PlaneCA);
 			FRealSingle Edge = FMath::Min<FRealSingle>(EdgeX1, EdgeX2);
@@ -51,9 +51,9 @@ namespace Chaos
 			if (Edge > 0.0f)
 			{
 				VectorRegister4Float OtherNormal = VectorCross(AX1, CA);
-				PlaneCA = VectorNormalize(VectorCross(CA, OtherNormal));
-				EdgeX1 = VectorDot3Scalar(AX1, PlaneCA);
-				EdgeX2 = VectorDot3Scalar(AX2, PlaneCA);
+				VectorRegister4Float OtherPlaneCA = VectorNormalize(VectorCross(CA, OtherNormal));
+				EdgeX1 = VectorDot3Scalar(AX1, OtherPlaneCA);
+				EdgeX2 = VectorDot3Scalar(AX2, OtherPlaneCA);
 				Edge = FMath::Min<FRealSingle>(EdgeX1, EdgeX2);
 				if (Edge > Radius)
 				{
@@ -61,9 +61,9 @@ namespace Chaos
 				}
 
 				OtherNormal = VectorCross(AX2, CA);
-				PlaneCA = VectorNormalize(VectorCross(CA, OtherNormal));
-				EdgeX1 = VectorDot3Scalar(AX1, PlaneCA);
-				EdgeX2 = VectorDot3Scalar(AX2, PlaneCA);
+				OtherPlaneCA = VectorNormalize(VectorCross(CA, OtherNormal));
+				EdgeX1 = VectorDot3Scalar(AX1, OtherPlaneCA);
+				EdgeX2 = VectorDot3Scalar(AX2, OtherPlaneCA);
 				Edge = FMath::Min<FRealSingle>(EdgeX1, EdgeX2);
 				if (Edge > Radius)
 				{
@@ -71,8 +71,8 @@ namespace Chaos
 				}
 			}
 		}
+		const VectorRegister4Float PlaneAB = VectorNormalize(VectorCross(AB, Normal));
 		{
-			VectorRegister4Float PlaneAB = VectorNormalize(VectorCross(AB, Normal));
 			FRealSingle EdgeX1 = VectorDot3Scalar(BX1, PlaneAB);
 			FRealSingle EdgeX2 = VectorDot3Scalar(VectorSubtract(X2, B), PlaneAB);
 			FRealSingle Edge = FMath::Min<FRealSingle>(EdgeX1, EdgeX2);
@@ -84,9 +84,9 @@ namespace Chaos
 			if (Edge > 0.0f)
 			{
 				VectorRegister4Float OtherNormal = VectorCross(BX1, AB);
-				PlaneAB = VectorNormalize(VectorCross(AB, OtherNormal));
-				EdgeX1 = VectorDot3Scalar(BX1, PlaneAB);
-				EdgeX2 = VectorDot3Scalar(VectorSubtract(X2, B), PlaneAB);
+				VectorRegister4Float OtherPlaneAB = VectorNormalize(VectorCross(AB, OtherNormal));
+				EdgeX1 = VectorDot3Scalar(BX1, OtherPlaneAB);
+				EdgeX2 = VectorDot3Scalar(VectorSubtract(X2, B), OtherPlaneAB);
 				Edge = FMath::Min<FRealSingle>(EdgeX1, EdgeX2);
 				if (Edge > Radius)
 				{
@@ -94,9 +94,9 @@ namespace Chaos
 				}
 
 				OtherNormal = VectorCross(VectorSubtract(X2, B), AB);
-				PlaneAB = VectorNormalize(VectorCross(AB, OtherNormal));
-				EdgeX1 = VectorDot3Scalar(BX1, PlaneAB);
-				EdgeX2 = VectorDot3Scalar(VectorSubtract(X2, B), PlaneAB);
+				OtherPlaneAB = VectorNormalize(VectorCross(AB, OtherNormal));
+				EdgeX1 = VectorDot3Scalar(BX1, OtherPlaneAB);
+				EdgeX2 = VectorDot3Scalar(VectorSubtract(X2, B), OtherPlaneAB);
 				Edge = FMath::Min<FRealSingle>(EdgeX1, EdgeX2);
 				if (Edge > Radius)
 				{
@@ -104,8 +104,8 @@ namespace Chaos
 				}
 			}
 		}
+		const VectorRegister4Float PlaneBC = VectorNormalize(VectorCross(BC, Normal));
 		{
-			VectorRegister4Float PlaneBC = VectorNormalize(VectorCross(BC, Normal));
 			FRealSingle EdgeX1 = VectorDot3Scalar(CX1, PlaneBC);
 			FRealSingle EdgeX2 = VectorDot3Scalar(VectorSubtract(X2, C), PlaneBC);
 			FRealSingle Edge = FMath::Min<FRealSingle>(EdgeX1, EdgeX2);
@@ -117,9 +117,9 @@ namespace Chaos
 			if (Edge > 0.0f)
 			{
 				VectorRegister4Float OtherNormal = VectorCross(CX1, BC);
-				PlaneBC = VectorNormalize(VectorCross(BC, OtherNormal));
-				EdgeX1 = VectorDot3Scalar(CX1, PlaneBC);
-				EdgeX2 = VectorDot3Scalar(VectorSubtract(X2, C), PlaneBC);
+				VectorRegister4Float OtherPlaneBC = VectorNormalize(VectorCross(BC, OtherNormal));
+				EdgeX1 = VectorDot3Scalar(CX1, OtherPlaneBC);
+				EdgeX2 = VectorDot3Scalar(VectorSubtract(X2, C), OtherPlaneBC);
 				Edge = FMath::Min<FRealSingle>(EdgeX1, EdgeX2);
 				if (Edge > Radius)
 				{
@@ -127,9 +127,9 @@ namespace Chaos
 				}
 
 				OtherNormal = VectorCross(VectorSubtract(X2, C), BC);
-				PlaneBC = VectorNormalize(VectorCross(BC, OtherNormal));
-				EdgeX1 = VectorDot3Scalar(CX1, PlaneBC);
-				EdgeX2 = VectorDot3Scalar(VectorSubtract(X2, C), PlaneBC);
+				OtherPlaneBC = VectorNormalize(VectorCross(BC, OtherNormal));
+				EdgeX1 = VectorDot3Scalar(CX1, OtherPlaneBC);
+				EdgeX2 = VectorDot3Scalar(VectorSubtract(X2, C), OtherPlaneBC);
 				Edge = FMath::Min<FRealSingle>(EdgeX1, EdgeX2);
 				if (Edge > Radius)
 				{
@@ -151,15 +151,15 @@ namespace Chaos
 		}
 
 		// Triangle Vertices
-		FRealSingle SqrRadius = Radius * Radius;
+		const FRealSingle SqrRadius = Radius * Radius;
 		{
 			VectorRegister4Float TimeA = VectorClamp(VectorDivide(VectorDot3(X2X1, AX1), SqrX2X1), VectorZeroFloat(), VectorOneFloat());
 			TimeA = VectorBitwiseNotAnd(ZeroMask, TimeA);
 			const VectorRegister4Float PA = VectorMultiplyAdd(X1, VectorSubtract(VectorOneFloat(), TimeA), VectorMultiply(X2, TimeA));
-
-			if (VectorDot3Scalar(VectorSubtract(PA, A), CentroidA) > 0.0f)
+			const VectorRegister4Float APA = VectorSubtract(PA, A);
+			FRealSingle BetweenSeg = VectorDot3Scalar(VectorCross(PlaneAB, APA), VectorCross(PlaneCA, APA));
+			if (BetweenSeg < 0.0f && VectorDot3Scalar(APA, CentroidA) > 0.0f)
 			{
-				const VectorRegister4Float APA = VectorSubtract(PA, A);
 				if (VectorDot3Scalar(APA, APA) > SqrRadius)
 				{
 					return false;
@@ -170,10 +170,10 @@ namespace Chaos
 			VectorRegister4Float TimeB = VectorClamp(VectorDivide(VectorDot3(X2X1, VectorSubtract(X1, B)), SqrX2X1), VectorZeroFloat(), VectorOneFloat());
 			TimeB = VectorBitwiseNotAnd(ZeroMask, TimeB);
 			const VectorRegister4Float PB = VectorMultiplyAdd(X1, VectorSubtract(VectorOneFloat(), TimeB), VectorMultiply(X2, TimeB));
-
-			if (VectorDot3Scalar(VectorSubtract(PB, B), CentroidB) > 0.0f)
+			const VectorRegister4Float BPB = VectorSubtract(PB, B);
+			FRealSingle BetweenSeg = VectorDot3Scalar(VectorCross(PlaneAB, BPB), VectorCross(PlaneBC, BPB));
+			if (BetweenSeg < 0.0f && VectorDot3Scalar(BPB, CentroidB) > 0.0f)
 			{
-				const VectorRegister4Float BPB = VectorSubtract(PB, B);
 				if (VectorDot3Scalar(BPB, BPB) > SqrRadius)
 				{
 					return false;
@@ -184,10 +184,10 @@ namespace Chaos
 			VectorRegister4Float TimeC = VectorClamp(VectorDivide(VectorDot3(X2X1, VectorSubtract(X1, C)), SqrX2X1), VectorZeroFloat(), VectorOneFloat());
 			TimeC = VectorBitwiseNotAnd(ZeroMask, TimeC);
 			const VectorRegister4Float PC = VectorMultiplyAdd(X1, VectorSubtract(VectorOneFloat(), TimeC), VectorMultiply(X2, TimeC));
-
-			if (VectorDot3Scalar(VectorSubtract(PC, C), CentroidC) > 0.0f)
+			const VectorRegister4Float CPC = VectorSubtract(PC, C);
+			FRealSingle BetweenSeg = VectorDot3Scalar(VectorCross(PlaneBC, CPC), VectorCross(PlaneCA, CPC));
+			if (BetweenSeg < 0.0f && VectorDot3Scalar(CPC, CentroidC) > 0.0f)
 			{
-				const VectorRegister4Float CPC = VectorSubtract(PC, C);
 				if (VectorDot3Scalar(CPC, CPC) > SqrRadius)
 				{
 					return false;
