Fix Box overlap SAT consistency with GJK

diff --git a/Engine/Source/Programs/HeadlessChaos/Private/HeadlessChaosTestHeightfield.cpp b/Engine/Source/Programs/HeadlessChaos/Private/HeadlessChaosTestHeightfield.cpp
--- a/Engine/Source/Programs/HeadlessChaos/Private/HeadlessChaosTestHeightfield.cpp
+++ b/Engine/Source/Programs/HeadlessChaos/Private/HeadlessChaosTestHeightfield.cpp
@@ -967,7 +967,7 @@ namespace ChaosTest {
 						FVec3 Dir(1, 0, 0);
 						bool BoxResult = Heightfield.OverlapGeom(Box, QueryTM, 0.0, nullptr);
 						bool BoxResultMTD = Heightfield.OverlapGeom(Box, QueryTM, 0.0, &OutBoxMTD);
-						// EXPECT_EQ(BoxResult, BoxResultMTD); // TODO Enable test
+						EXPECT_EQ(BoxResult, BoxResultMTD);
 
 						bool CapsuleResult = Heightfield.OverlapGeom(Capsule, QueryTM, 0.0, nullptr);
 						bool CapsuleResultMTD = Heightfield.OverlapGeom(Capsule, QueryTM, 0.0, &OutCapsuleMTD);
diff --git a/Engine/Source/Programs/HeadlessChaos/Private/HeadlessChaosTestOverlap.cpp b/Engine/Source/Programs/HeadlessChaos/Private/HeadlessChaosTestOverlap.cpp
--- a/Engine/Source/Programs/HeadlessChaos/Private/HeadlessChaosTestOverlap.cpp
+++ b/Engine/Source/Programs/HeadlessChaos/Private/HeadlessChaosTestOverlap.cpp
@@ -47,6 +47,7 @@ namespace ChaosTest
 			}
 
 			TUniquePtr<FTriangleMeshImplicitObject> TriangleMesh = MakeUnique<FTriangleMeshImplicitObject>(MoveTemp(TrimeshParticles), MoveTemp(Indices), MoveTemp(Materials));
+			FMTDInfo MTDInfo;
 			{
 				// Capsule test
 				const FVec3 X1 = { 0.0, 0.0, -2.0 };
@@ -58,72 +59,100 @@ namespace ChaosTest
 					const FRigidTransform3 QueryTM(FVec3(0.0, 0.0, 0.0), FQuat::Identity);
 					bool bResult = TriangleMesh->OverlapGeom(Capsule, QueryTM, 0.0);
 					EXPECT_EQ(bResult, true);
+					bool bResultMTD = TriangleMesh->OverlapGeom(Capsule, QueryTM, 0.0, &MTDInfo);
+					EXPECT_EQ(bResult, bResultMTD);
 				}
 				{
 					const FRigidTransform3 QueryTM(FVec3(12.5, 0.0, 0.0), FQuat::Identity);
 					bool bResult = TriangleMesh->OverlapGeom(Capsule, QueryTM, 0.0);
 					EXPECT_EQ(bResult, false);
+					bool bResultMTD = TriangleMesh->OverlapGeom(Capsule, QueryTM, 0.0, &MTDInfo);
+					EXPECT_EQ(bResult, bResultMTD);
 				}
 				{
 					const FRigidTransform3 QueryTM(FVec3(10.0, 10.0, 0.0), FQuat::Identity);
 					bool bResult = TriangleMesh->OverlapGeom(Capsule, QueryTM, 0.0);
 					EXPECT_EQ(bResult, false);
+					bool bResultMTD = TriangleMesh->OverlapGeom(Capsule, QueryTM, 0.0, &MTDInfo);
+					EXPECT_EQ(bResult, bResultMTD);
 				}
 				{
 					const FRigidTransform3 QueryTM(FVec3(-10.0, -10.0, 0.0), FQuat::Identity);
 					bool bResult = TriangleMesh->OverlapGeom(Capsule, QueryTM, 0.0);
 					EXPECT_EQ(bResult, false);
+					bool bResultMTD = TriangleMesh->OverlapGeom(Capsule, QueryTM, 0.0, &MTDInfo);
+					EXPECT_EQ(bResult, bResultMTD);
 				}
 				{
 					const FRigidTransform3 QueryTM(FVec3(10.0, 0.0, 0.0), FQuat::Identity);
 					bool bResult = TriangleMesh->OverlapGeom(Capsule, QueryTM, 0.0);
 					EXPECT_EQ(bResult, true);
+					bool bResultMTD = TriangleMesh->OverlapGeom(Capsule, QueryTM, 0.0, &MTDInfo);
+					EXPECT_EQ(bResult, bResultMTD);
 				}
 				{
 					const FRigidTransform3 QueryTM(FVec3(0.0, 10.0, 0.0), FQuat::Identity);
 					bool bResult = TriangleMesh->OverlapGeom(Capsule, QueryTM, 0.0);
 					EXPECT_EQ(bResult, true);
+					bool bResultMTD = TriangleMesh->OverlapGeom(Capsule, QueryTM, 0.0, &MTDInfo);
+					EXPECT_EQ(bResult, bResultMTD);
 				}
 				{
 					const FRigidTransform3 QueryTM(FVec3(0.0, -10.0, 0.0), FQuat::Identity);
 					bool bResult = TriangleMesh->OverlapGeom(Capsule, QueryTM, 0.0);
 					EXPECT_EQ(bResult, true);
+					bool bResultMTD = TriangleMesh->OverlapGeom(Capsule, QueryTM, 0.0, &MTDInfo);
+					EXPECT_EQ(bResult, bResultMTD);
 				}
 				{
 					const FRigidTransform3 QueryTM(FVec3(-10.0, 0.0, 0.0), FQuat::Identity);
 					bool bResult = TriangleMesh->OverlapGeom(Capsule, QueryTM, 0.0);
 					EXPECT_EQ(bResult, true);
+					bool bResultMTD = TriangleMesh->OverlapGeom(Capsule, QueryTM, 0.0, &MTDInfo);
+					EXPECT_EQ(bResult, bResultMTD);
 				}
 				{
 					const FRigidTransform3 QueryTM(FVec3(0.0, 0.0, 12.0), FQuat::Identity);
 					bool bResult = TriangleMesh->OverlapGeom(Capsule, QueryTM, 0.0);
 					EXPECT_EQ(bResult, true);
+					bool bResultMTD = TriangleMesh->OverlapGeom(Capsule, QueryTM, 0.0, &MTDInfo);
+					EXPECT_EQ(bResult, bResultMTD);
 				}
 				{
 					const FRigidTransform3 QueryTM(FVec3(0.0, 0.0, 12.0), FQuat(FVec3(1.0, 0.0, 0.0), 3.1415926 / 4.0));
 					bool bResult = TriangleMesh->OverlapGeom(Capsule, QueryTM, 0.0);
 					EXPECT_EQ(bResult, false);
+					bool bResultMTD = TriangleMesh->OverlapGeom(Capsule, QueryTM, 0.0, &MTDInfo);
+					EXPECT_EQ(bResult, bResultMTD);
 				}
 				{
 					// Inside Mesh
 					const FRigidTransform3 QueryTM(FVec3(0.0, 0.0, 4.0), FQuat::Identity);
 					bool bResult = TriangleMesh->OverlapGeom(Capsule, QueryTM, 0.0);
 					EXPECT_EQ(bResult, false);
+					bool bResultMTD = TriangleMesh->OverlapGeom(Capsule, QueryTM, 0.0, &MTDInfo);
+					EXPECT_EQ(bResult, bResultMTD);
 				}
 				{
 					const FRigidTransform3 QueryTM(FVec3(6.5, 0.0, 6.5), FQuat::Identity);
 					bool bResult = TriangleMesh->OverlapGeom(Capsule, QueryTM, 0.0);
 					EXPECT_TRUE(bResult);
+					bool bResultMTD = TriangleMesh->OverlapGeom(Capsule, QueryTM, 0.0, &MTDInfo);
+					EXPECT_EQ(bResult, bResultMTD);
 				}
 				{
 					const FRigidTransform3 QueryTM(FVec3(6.5, 0.0, 6.5), FQuat(FVec3(0.0, 1.0, 0.0), -3.1415926 / 4.0));
 					bool bResult = TriangleMesh->OverlapGeom(Capsule, QueryTM, 0.0);
 					EXPECT_FALSE(bResult);
+					bool bResultMTD = TriangleMesh->OverlapGeom(Capsule, QueryTM, 0.0, &MTDInfo);
+					EXPECT_EQ(bResult, bResultMTD);
 				}
 				{
 					const FRigidTransform3 QueryTM(FVec3(2.5, 2.5, 0.0), FQuat::Identity);
 					bool bResult = TriangleMesh->OverlapGeom(Capsule, QueryTM, 0.0);
 					EXPECT_TRUE(bResult);
+					bool bResultMTD = TriangleMesh->OverlapGeom(Capsule, QueryTM, 0.0, &MTDInfo);
+					EXPECT_EQ(bResult, bResultMTD);
 				}
 				{
 					const FRigidTransform3 QueryTM(FVec3(-2.5, -2.5, 0.0), FQuat::Identity);
@@ -131,12 +160,11 @@ namespace ChaosTest
 					EXPECT_TRUE(bResult);
 				}
 			}
-{
+			{
 				const FVec3 X1 = { -11.0, 2.0, 0.0 };
 				const FVec3 X2 = { -9.0, -8.0, 0.0 };
 				const FReal Radius = 0.01;
 				const FCapsule Capsule = FCapsule(X1, X2, Radius);
-				FMTDInfo MTDInfo;
 				{
 					const FRigidTransform3 QueryTM(FVec3(0.0, 0.0, 0.0), FQuat::Identity);
 					bool bResult = TriangleMesh->OverlapGeom(Capsule, QueryTM, 0.0);
@@ -164,72 +192,100 @@ namespace ChaosTest
 					const FRigidTransform3 QueryTM(FVec3(0.0, 0.0, 0.0), FQuat::Identity);
 					bool bResult = TriangleMesh->OverlapGeom(Sphere, QueryTM, 0.0);
 					EXPECT_EQ(bResult, true);
+					bool bResultMTD = TriangleMesh->OverlapGeom(Sphere, QueryTM, 0.0, &MTDInfo);
+					EXPECT_EQ(bResult, bResultMTD);
 				}
 				{
 					const FRigidTransform3 QueryTM(FVec3(12.5, 0.0, 0.0), FQuat::Identity);
 					bool bResult = TriangleMesh->OverlapGeom(Sphere, QueryTM, 0.0);
 					EXPECT_EQ(bResult, false);
+					bool bResultMTD = TriangleMesh->OverlapGeom(Sphere, QueryTM, 0.0, &MTDInfo);
+					EXPECT_EQ(bResult, bResultMTD);
 				}
 				{
 					const FRigidTransform3 QueryTM(FVec3(7.0, 7.0, 0.0), FQuat::Identity);
 					bool bResult = TriangleMesh->OverlapGeom(Sphere, QueryTM, 0.0);
 					EXPECT_EQ(bResult, false);
+					bool bResultMTD = TriangleMesh->OverlapGeom(Sphere, QueryTM, 0.0, &MTDInfo);
+					EXPECT_EQ(bResult, bResultMTD);
 				}
 				{
 					const FRigidTransform3 QueryTM(FVec3(-7.0, -7.0, 0.0), FQuat::Identity);
 					bool bResult = TriangleMesh->OverlapGeom(Sphere, QueryTM, 0.0);
 					EXPECT_EQ(bResult, false);
+					bool bResultMTD = TriangleMesh->OverlapGeom(Sphere, QueryTM, 0.0, &MTDInfo);
+					EXPECT_EQ(bResult, bResultMTD);
 				}
 				{
 					const FRigidTransform3 QueryTM(FVec3(10.0, 0.0, 0.0), FQuat::Identity);
 					bool bResult = TriangleMesh->OverlapGeom(Sphere, QueryTM, 0.0);
 					EXPECT_EQ(bResult, true);
+					bool bResultMTD = TriangleMesh->OverlapGeom(Sphere, QueryTM, 0.0, &MTDInfo);
+					EXPECT_EQ(bResult, bResultMTD);
 				}
 				{
 					const FRigidTransform3 QueryTM(FVec3(0.0, 10.2, 0.0), FQuat::Identity);
 					bool bResult = TriangleMesh->OverlapGeom(Sphere, QueryTM, 0.0);
 					EXPECT_EQ(bResult, true);
+					bool bResultMTD = TriangleMesh->OverlapGeom(Sphere, QueryTM, 0.0, &MTDInfo);
+					EXPECT_EQ(bResult, bResultMTD);
 				}
 				{
 					const FRigidTransform3 QueryTM(FVec3(0.0, -10.2, 0.0), FQuat::Identity);
 					bool bResult = TriangleMesh->OverlapGeom(Sphere, QueryTM, 0.0);
 					EXPECT_EQ(bResult, true);
+					bool bResultMTD = TriangleMesh->OverlapGeom(Sphere, QueryTM, 0.0, &MTDInfo);
+					EXPECT_EQ(bResult, bResultMTD);
 				}
 				{
 					const FRigidTransform3 QueryTM(FVec3(-10.3, 0.0, 0.0), FQuat::Identity);
 					bool bResult = TriangleMesh->OverlapGeom(Sphere, QueryTM, 0.0);
 					EXPECT_EQ(bResult, true);
+					bool bResultMTD = TriangleMesh->OverlapGeom(Sphere, QueryTM, 0.0, &MTDInfo);
+					EXPECT_EQ(bResult, bResultMTD);
 				}
 				{
 					const FRigidTransform3 QueryTM(FVec3(0.0, 0.0, 10.9), FQuat::Identity);
 					bool bResult = TriangleMesh->OverlapGeom(Sphere, QueryTM, 0.0);
 					EXPECT_EQ(bResult, true);
+					bool bResultMTD = TriangleMesh->OverlapGeom(Sphere, QueryTM, 0.0, &MTDInfo);
+					EXPECT_EQ(bResult, bResultMTD);
 				}
 				{
 					const FRigidTransform3 QueryTM(FVec3(0.0, 0.0, 11.5), FQuat(FVec3(1.0, 0.0, 0.0), 3.1415926 / 4.0));
 					bool bResult = TriangleMesh->OverlapGeom(Sphere, QueryTM, 0.0);
 					EXPECT_EQ(bResult, false);
+					bool bResultMTD = TriangleMesh->OverlapGeom(Sphere, QueryTM, 0.0, &MTDInfo);
+					EXPECT_EQ(bResult, bResultMTD);
 				}
 				{
 					// Inside Mesh
 					const FRigidTransform3 QueryTM(FVec3(0.0, 0.0, 4.0), FQuat::Identity);
 					bool bResult = TriangleMesh->OverlapGeom(Sphere, QueryTM, 0.0);
 					EXPECT_EQ(bResult, false);
+					bool bResultMTD = TriangleMesh->OverlapGeom(Sphere, QueryTM, 0.0, &MTDInfo);
+					EXPECT_EQ(bResult, bResultMTD);
 				}
 				{
 					const FRigidTransform3 QueryTM(FVec3(5.5, 0.0, 5.5), FQuat::Identity);
 					bool bResult = TriangleMesh->OverlapGeom(Sphere, QueryTM, 0.0);
 					EXPECT_TRUE(bResult);
+					bool bResultMTD = TriangleMesh->OverlapGeom(Sphere, QueryTM, 0.0, &MTDInfo);
+					EXPECT_EQ(bResult, bResultMTD);
 				}
 				{
 					const FRigidTransform3 QueryTM(FVec3(5.8, 0.0, 5.8), FQuat(FVec3(0.0, 1.0, 0.0), -3.1415926 / 4.0));
 					bool bResult = TriangleMesh->OverlapGeom(Sphere, QueryTM, 0.0);
 					EXPECT_FALSE(bResult);
+					bool bResultMTD = TriangleMesh->OverlapGeom(Sphere, QueryTM, 0.0, &MTDInfo);
+					EXPECT_EQ(bResult, bResultMTD);
 				}
 				{
 					const FRigidTransform3 QueryTM(FVec3(-3.0, 0.0, 9.0), FQuat::Identity);
 					bool bResult = TriangleMesh->OverlapGeom(Sphere, QueryTM, 0.0);
 					EXPECT_FALSE(bResult);
+					bool bResultMTD = TriangleMesh->OverlapGeom(Sphere, QueryTM, 0.0, &MTDInfo);
+					EXPECT_EQ(bResult, bResultMTD);
 				}
 			}
 			{
@@ -244,62 +300,86 @@ namespace ChaosTest
 					const FRigidTransform3 QueryTM(FVec3(0.0, 0.0, 0.0), FQuat::Identity);
 					bool bResult = TriangleMesh->OverlapGeom(ScaledCapsule, QueryTM, 0.0, nullptr, TriMeshScale);
 					EXPECT_EQ(bResult, true);
+					bool bResultMTD = TriangleMesh->OverlapGeom(ScaledCapsule, QueryTM, 0.0, &MTDInfo, TriMeshScale);
+					EXPECT_EQ(bResult, bResultMTD);
 				}
 				{
 					const FRigidTransform3 QueryTM(FVec3(10.0, 10.0, 0.0), FQuat::Identity);
 					bool bResult = TriangleMesh->OverlapGeom(ScaledCapsule, QueryTM, 0.0, nullptr, TriMeshScale);
 					EXPECT_EQ(bResult, false);
+					bool bResultMTD = TriangleMesh->OverlapGeom(ScaledCapsule, QueryTM, 0.0, &MTDInfo, TriMeshScale);
+					EXPECT_EQ(bResult, bResultMTD);
 				}
 				{
 					const FRigidTransform3 QueryTM(FVec3(-10.0, -10.0, 0.0), FQuat::Identity);
 					bool bResult = TriangleMesh->OverlapGeom(ScaledCapsule, QueryTM, 0.0, nullptr, TriMeshScale);
 					EXPECT_EQ(bResult, false);
+					bool bResultMTD = TriangleMesh->OverlapGeom(ScaledCapsule, QueryTM, 0.0, &MTDInfo, TriMeshScale);
+					EXPECT_EQ(bResult, bResultMTD);
 				}
 				{
 					const FRigidTransform3 QueryTM(FVec3(10.0, 0.0, 0.0), FQuat::Identity);
 					bool bResult = TriangleMesh->OverlapGeom(ScaledCapsule, QueryTM, 0.0, nullptr, TriMeshScale);
 					EXPECT_EQ(bResult, true);
+					bool bResultMTD = TriangleMesh->OverlapGeom(ScaledCapsule, QueryTM, 0.0, &MTDInfo, TriMeshScale);
+					EXPECT_EQ(bResult, bResultMTD);
 				}
 				{
 					const FRigidTransform3 QueryTM(FVec3(0.0, 10.0, 0.0), FQuat::Identity);
 					bool bResult = TriangleMesh->OverlapGeom(ScaledCapsule, QueryTM, 0.0, nullptr, TriMeshScale);
 					EXPECT_EQ(bResult, true);
+					bool bResultMTD = TriangleMesh->OverlapGeom(ScaledCapsule, QueryTM, 0.0, &MTDInfo, TriMeshScale);
+					EXPECT_EQ(bResult, bResultMTD);
 				}
 				{
 					const FRigidTransform3 QueryTM(FVec3(0.0, -10.0, 0.0), FQuat::Identity);
 					bool bResult = TriangleMesh->OverlapGeom(ScaledCapsule, QueryTM, 0.0, nullptr, TriMeshScale);
 					EXPECT_EQ(bResult, true);
+					bool bResultMTD = TriangleMesh->OverlapGeom(ScaledCapsule, QueryTM, 0.0, &MTDInfo, TriMeshScale);
+					EXPECT_EQ(bResult, bResultMTD);
 				}
 				{
 					const FRigidTransform3 QueryTM(FVec3(-10.0, 0.0, 0.0), FQuat::Identity);
 					bool bResult = TriangleMesh->OverlapGeom(ScaledCapsule, QueryTM, 0.0, nullptr, TriMeshScale);
 					EXPECT_EQ(bResult, true);
+					bool bResultMTD = TriangleMesh->OverlapGeom(ScaledCapsule, QueryTM, 0.0, &MTDInfo, TriMeshScale);
+					EXPECT_EQ(bResult, bResultMTD);
 				}
 				{
 					const FRigidTransform3 QueryTM(FVec3(0.0, 0.0, 12.0), FQuat::Identity);
 					bool bResult = TriangleMesh->OverlapGeom(ScaledCapsule, QueryTM, 0.0, nullptr, TriMeshScale);
 					EXPECT_EQ(bResult, true);
+					bool bResultMTD = TriangleMesh->OverlapGeom(ScaledCapsule, QueryTM, 0.0, &MTDInfo, TriMeshScale);
+					EXPECT_EQ(bResult, bResultMTD);
 				}
 				{
 					const FRigidTransform3 QueryTM(FVec3(0.0, 0.0, 12.0), FQuat(FVec3(1.0, 0.0, 0.0), 3.1415926 / 4.0));
 					bool bResult = TriangleMesh->OverlapGeom(ScaledCapsule, QueryTM, 0.0, nullptr, TriMeshScale);
 					EXPECT_EQ(bResult, false);
+					bool bResultMTD = TriangleMesh->OverlapGeom(ScaledCapsule, QueryTM, 0.0, &MTDInfo, TriMeshScale);
+					EXPECT_EQ(bResult, bResultMTD);
 				}
 				{
 					// Inside Mesh
 					const FRigidTransform3 QueryTM(FVec3(0.0, 0.0, 4.0), FQuat::Identity);
 					bool bResult = TriangleMesh->OverlapGeom(ScaledCapsule, QueryTM, 0.0, nullptr, TriMeshScale);
 					EXPECT_EQ(bResult, false);
+					bool bResultMTD = TriangleMesh->OverlapGeom(ScaledCapsule, QueryTM, 0.0, &MTDInfo, TriMeshScale);
+					EXPECT_EQ(bResult, bResultMTD);
 				}
 				{
 					const FRigidTransform3 QueryTM(FVec3(6.5, 0.0, 6.5), FQuat::Identity);
 					bool bResult = TriangleMesh->OverlapGeom(ScaledCapsule, QueryTM, 0.0, nullptr, TriMeshScale);
 					EXPECT_TRUE(bResult);
+					bool bResultMTD = TriangleMesh->OverlapGeom(ScaledCapsule, QueryTM, 0.0, &MTDInfo, TriMeshScale);
+					EXPECT_EQ(bResult, bResultMTD);
 				}
 				{
 					const FRigidTransform3 QueryTM(FVec3(6.5, 0.0, 6.5), FQuat(FVec3(0.0, 1.0, 0.0), -3.1415926 / 4.0));
 					bool bResult = TriangleMesh->OverlapGeom(ScaledCapsule, QueryTM, 0.0, nullptr, TriMeshScale);
 					EXPECT_FALSE(bResult);
+					bool bResultMTD = TriangleMesh->OverlapGeom(ScaledCapsule, QueryTM, 0.0, &MTDInfo, TriMeshScale);
+					EXPECT_EQ(bResult, bResultMTD);
 				}
 			}
 			{
@@ -315,60 +395,80 @@ namespace ChaosTest
 					QueryTM = FRigidTransform3(QueryTM.GetLocation() * InvScale, QueryTM.GetRotation());
 					bool bResult = TriangleMesh->OverlapGeom(ScaledCapsule, QueryTM, 0.0, nullptr, TriMeshScale);
 					EXPECT_EQ(bResult, true);
+					bool bResultMTD = TriangleMesh->OverlapGeom(ScaledCapsule, QueryTM, 0.0, &MTDInfo, TriMeshScale);
+					EXPECT_EQ(bResult, bResultMTD);
 				}
 				{
 					FRigidTransform3 QueryTM(FVec3(100.0, 100.0, 0.0), FQuat::Identity);
 					QueryTM = FRigidTransform3(QueryTM.GetLocation() * InvScale, QueryTM.GetRotation());
 					bool bResult = TriangleMesh->OverlapGeom(ScaledCapsule, QueryTM, 0.0, nullptr, TriMeshScale);
 					EXPECT_EQ(bResult, false);
+					bool bResultMTD = TriangleMesh->OverlapGeom(ScaledCapsule, QueryTM, 0.0, &MTDInfo, TriMeshScale);
+					EXPECT_EQ(bResult, bResultMTD);
 				}
 				{
 					FRigidTransform3 QueryTM(FVec3(125.0, 0.0, 0.0), FQuat::Identity);
 					QueryTM = FRigidTransform3(QueryTM.GetLocation() * InvScale, QueryTM.GetRotation());
 					bool bResult = TriangleMesh->OverlapGeom(ScaledCapsule, QueryTM, 0.0, nullptr, TriMeshScale);
 					EXPECT_EQ(bResult, false);
+					bool bResultMTD = TriangleMesh->OverlapGeom(ScaledCapsule, QueryTM, 0.0, &MTDInfo, TriMeshScale);
+					EXPECT_EQ(bResult, bResultMTD);
 				}
 				{
 					FRigidTransform3 QueryTM(FVec3(-100.0, -100.0, 0.0), FQuat::Identity);
 					QueryTM = FRigidTransform3(QueryTM.GetLocation() * InvScale, QueryTM.GetRotation());
 					bool bResult = TriangleMesh->OverlapGeom(ScaledCapsule, QueryTM, 0.0, nullptr, TriMeshScale);
 					EXPECT_EQ(bResult, false);
+					bool bResultMTD = TriangleMesh->OverlapGeom(ScaledCapsule, QueryTM, 0.0, &MTDInfo, TriMeshScale);
+					EXPECT_EQ(bResult, bResultMTD);
 				}
 				{
 					FRigidTransform3 QueryTM(FVec3(100.0, 0.0, 0.0), FQuat::Identity);
 					QueryTM = FRigidTransform3(QueryTM.GetLocation() * InvScale, QueryTM.GetRotation());
 					bool bResult = TriangleMesh->OverlapGeom(ScaledCapsule, QueryTM, 0.0, nullptr, TriMeshScale);
 					EXPECT_EQ(bResult, true);
+					bool bResultMTD = TriangleMesh->OverlapGeom(ScaledCapsule, QueryTM, 0.0, &MTDInfo, TriMeshScale);
+					EXPECT_EQ(bResult, bResultMTD);
 				}
 				{
 					FRigidTransform3 QueryTM(FVec3(0.0, 100.0, 0.0), FQuat::Identity);
 					QueryTM = FRigidTransform3(QueryTM.GetLocation() * InvScale, QueryTM.GetRotation());
 					bool bResult = TriangleMesh->OverlapGeom(ScaledCapsule, QueryTM, 0.0, nullptr, TriMeshScale);
 					EXPECT_EQ(bResult, true);
+					bool bResultMTD = TriangleMesh->OverlapGeom(ScaledCapsule, QueryTM, 0.0, &MTDInfo, TriMeshScale);
+					EXPECT_EQ(bResult, bResultMTD);
 				}
 				{
 					FRigidTransform3 QueryTM(FVec3(0.0, -100.0, 0.0), FQuat::Identity);
 					QueryTM = FRigidTransform3(QueryTM.GetLocation() * InvScale, QueryTM.GetRotation());
 					bool bResult = TriangleMesh->OverlapGeom(ScaledCapsule, QueryTM, 0.0, nullptr, TriMeshScale);
 					EXPECT_EQ(bResult, true);
+					bool bResultMTD = TriangleMesh->OverlapGeom(ScaledCapsule, QueryTM, 0.0, &MTDInfo, TriMeshScale);
+					EXPECT_EQ(bResult, bResultMTD);
 				}
 				{
 					FRigidTransform3 QueryTM(FVec3(-100.0, 0.0, 0.0), FQuat::Identity);
 					QueryTM = FRigidTransform3(QueryTM.GetLocation() * InvScale, QueryTM.GetRotation());
 					bool bResult = TriangleMesh->OverlapGeom(ScaledCapsule, QueryTM, 0.0, nullptr, TriMeshScale);
 					EXPECT_EQ(bResult, true);
+					bool bResultMTD = TriangleMesh->OverlapGeom(ScaledCapsule, QueryTM, 0.0, &MTDInfo, TriMeshScale);
+					EXPECT_EQ(bResult, bResultMTD);
 				}
 				{
 					FRigidTransform3 QueryTM(FVec3(0.0, 0.0, 120.0), FQuat::Identity);
 					QueryTM = FRigidTransform3(QueryTM.GetLocation() * InvScale, QueryTM.GetRotation());
 					bool bResult = TriangleMesh->OverlapGeom(ScaledCapsule, QueryTM, 0.0, nullptr, TriMeshScale);
 					EXPECT_EQ(bResult, true);
+					bool bResultMTD = TriangleMesh->OverlapGeom(ScaledCapsule, QueryTM, 0.0, &MTDInfo, TriMeshScale);
+					EXPECT_EQ(bResult, bResultMTD);
 				}
 				{
 					FRigidTransform3 QueryTM(FVec3(0.0, 0.0, 120.0), FQuat(FVec3(1.0, 0.0, 0.0), 3.1415926 / 4.0));
 					QueryTM = FRigidTransform3(QueryTM.GetLocation() * InvScale, QueryTM.GetRotation());
 					bool bResult = TriangleMesh->OverlapGeom(ScaledCapsule, QueryTM, 0.0, nullptr, TriMeshScale);
 					EXPECT_EQ(bResult, false);
+					bool bResultMTD = TriangleMesh->OverlapGeom(ScaledCapsule, QueryTM, 0.0, &MTDInfo, TriMeshScale);
+					EXPECT_EQ(bResult, bResultMTD);
 				}
 				{
 					// Inside Mesh
@@ -376,18 +476,24 @@ namespace ChaosTest
 					QueryTM = FRigidTransform3(QueryTM.GetLocation() * InvScale, QueryTM.GetRotation());
 					bool bResult = TriangleMesh->OverlapGeom(ScaledCapsule, QueryTM, 0.0, nullptr, TriMeshScale);
 					EXPECT_EQ(bResult, false);
+					bool bResultMTD = TriangleMesh->OverlapGeom(ScaledCapsule, QueryTM, 0.0, &MTDInfo, TriMeshScale);
+					EXPECT_EQ(bResult, bResultMTD);
 				}
 				{
 					FRigidTransform3 QueryTM(FVec3(65.0, 0.0, 65.0), FQuat::Identity);
 					QueryTM = FRigidTransform3(QueryTM.GetLocation() * InvScale, QueryTM.GetRotation());
 					bool bResult = TriangleMesh->OverlapGeom(ScaledCapsule, QueryTM, 0.0, nullptr, TriMeshScale);
 					EXPECT_TRUE(bResult);
+					bool bResultMTD = TriangleMesh->OverlapGeom(ScaledCapsule, QueryTM, 0.0, &MTDInfo, TriMeshScale);
+					EXPECT_EQ(bResult, bResultMTD);
 				}
 				{
 					FRigidTransform3 QueryTM(FVec3(65.0, 0.0, 65.0), FQuat(FVec3(0.0, 1.0, 0.0), -3.1415926 / 4.0));
 					QueryTM = FRigidTransform3(QueryTM.GetLocation() * InvScale, QueryTM.GetRotation());
 					bool bResult = TriangleMesh->OverlapGeom(ScaledCapsule, QueryTM, 0.0, nullptr, TriMeshScale);
 					EXPECT_FALSE(bResult);
+					bool bResultMTD = TriangleMesh->OverlapGeom(ScaledCapsule, QueryTM, 0.0, &MTDInfo, TriMeshScale);
+					EXPECT_EQ(bResult, bResultMTD);
 				}
 			}
 			// Non uniform scale
@@ -404,66 +510,88 @@ namespace ChaosTest
 					QueryTM = FRigidTransform3(QueryTM.GetLocation() * InvScale, QueryTM.GetRotation());
 					bool bResult = TriangleMesh->OverlapGeom(ScaledCapsule, QueryTM, 0.0, nullptr, TriMeshScale);
 					EXPECT_EQ(bResult, true);
+					bool bResultMTD = TriangleMesh->OverlapGeom(ScaledCapsule, QueryTM, 0.0, &MTDInfo, TriMeshScale);
+					EXPECT_EQ(bResult, bResultMTD);
 				}
 				{
 					FRigidTransform3 QueryTM(FVec3(10.0, 10.0, 0.0), FQuat::Identity);
 					QueryTM = FRigidTransform3(QueryTM.GetLocation() * InvScale, QueryTM.GetRotation());
 					bool bResult = TriangleMesh->OverlapGeom(ScaledCapsule, QueryTM, 0.0, nullptr, TriMeshScale);
 					EXPECT_EQ(bResult, false);
+					bool bResultMTD = TriangleMesh->OverlapGeom(ScaledCapsule, QueryTM, 0.0, &MTDInfo, TriMeshScale);
+					EXPECT_EQ(bResult, bResultMTD);
 				}
 				{
 					FRigidTransform3 QueryTM(FVec3(12.5, 0.0, 0.0), FQuat::Identity);
 					QueryTM = FRigidTransform3(QueryTM.GetLocation() * InvScale, QueryTM.GetRotation());
 					bool bResult = TriangleMesh->OverlapGeom(ScaledCapsule, QueryTM, 0.0, nullptr, TriMeshScale);
 					EXPECT_EQ(bResult, false);
+					bool bResultMTD = TriangleMesh->OverlapGeom(ScaledCapsule, QueryTM, 0.0, &MTDInfo, TriMeshScale);
+					EXPECT_EQ(bResult, bResultMTD);
 				}
 				{
 					FRigidTransform3 QueryTM(FVec3(-10.0, -10.0, 0.0), FQuat::Identity);
 					QueryTM = FRigidTransform3(QueryTM.GetLocation() * InvScale, QueryTM.GetRotation());
 					bool bResult = TriangleMesh->OverlapGeom(ScaledCapsule, QueryTM, 0.0, nullptr, TriMeshScale);
 					EXPECT_EQ(bResult, false);
+					bool bResultMTD = TriangleMesh->OverlapGeom(ScaledCapsule, QueryTM, 0.0, &MTDInfo, TriMeshScale);
+					EXPECT_EQ(bResult, bResultMTD);
 				}
 				{
 					FRigidTransform3 QueryTM(FVec3(10.0, 0.0, 0.0), FQuat::Identity);
 					QueryTM = FRigidTransform3(QueryTM.GetLocation() * InvScale, QueryTM.GetRotation());
 					bool bResult = TriangleMesh->OverlapGeom(ScaledCapsule, QueryTM, 0.0, nullptr, TriMeshScale);
 					EXPECT_EQ(bResult, true);
+					bool bResultMTD = TriangleMesh->OverlapGeom(ScaledCapsule, QueryTM, 0.0, &MTDInfo, TriMeshScale);
+					EXPECT_EQ(bResult, bResultMTD);
 				}
 				{
 					FRigidTransform3 QueryTM(FVec3(0.0, 10.0, 0.0), FQuat::Identity);
 					QueryTM = FRigidTransform3(QueryTM.GetLocation() * InvScale, QueryTM.GetRotation());
 					bool bResult = TriangleMesh->OverlapGeom(ScaledCapsule, QueryTM, 0.0, nullptr, TriMeshScale);
 					EXPECT_EQ(bResult, true);
+					bool bResultMTD = TriangleMesh->OverlapGeom(ScaledCapsule, QueryTM, 0.0, &MTDInfo, TriMeshScale);
+					EXPECT_EQ(bResult, bResultMTD);
 				}
 				{
 					FRigidTransform3 QueryTM(FVec3(0.0, -10.0, 0.0), FQuat::Identity);
 					QueryTM = FRigidTransform3(QueryTM.GetLocation() * InvScale, QueryTM.GetRotation());
 					bool bResult = TriangleMesh->OverlapGeom(ScaledCapsule, QueryTM, 0.0, nullptr, TriMeshScale);
 					EXPECT_EQ(bResult, true);
+					bool bResultMTD = TriangleMesh->OverlapGeom(ScaledCapsule, QueryTM, 0.0, &MTDInfo, TriMeshScale);
+					EXPECT_EQ(bResult, bResultMTD);
 				}
 				{
 					FRigidTransform3 QueryTM(FVec3(-10.0, 0.0, 0.0), FQuat::Identity);
 					QueryTM = FRigidTransform3(QueryTM.GetLocation() * InvScale, QueryTM.GetRotation());
 					bool bResult = TriangleMesh->OverlapGeom(ScaledCapsule, QueryTM, 0.0, nullptr, TriMeshScale);
 					EXPECT_EQ(bResult, true);
+					bool bResultMTD = TriangleMesh->OverlapGeom(ScaledCapsule, QueryTM, 0.0, &MTDInfo, TriMeshScale);
+					EXPECT_EQ(bResult, bResultMTD);
 				}
 				{
 					FRigidTransform3 QueryTM(FVec3(0.0, 0.0, 22.0), FQuat::Identity);
 					QueryTM = FRigidTransform3(QueryTM.GetLocation() * InvScale, QueryTM.GetRotation());
 					bool bResult = TriangleMesh->OverlapGeom(ScaledCapsule, QueryTM, 0.0, nullptr, TriMeshScale);
 					EXPECT_EQ(bResult, true);
+					bool bResultMTD = TriangleMesh->OverlapGeom(ScaledCapsule, QueryTM, 0.0, &MTDInfo, TriMeshScale);
+					EXPECT_EQ(bResult, bResultMTD);
 				}
 				{
 					FRigidTransform3 QueryTM(FVec3(0.0, 0.0, 22.0), FQuat(FVec3(1.0, 0.0, 0.0), 3.1415926 / 4.0));
 					QueryTM = FRigidTransform3(QueryTM.GetLocation() * InvScale, QueryTM.GetRotation());
 					bool bResult = TriangleMesh->OverlapGeom(ScaledCapsule, QueryTM, 0.0, nullptr, TriMeshScale);
 					EXPECT_EQ(bResult, false);
+					bool bResultMTD = TriangleMesh->OverlapGeom(ScaledCapsule, QueryTM, 0.0, &MTDInfo, TriMeshScale);
+					EXPECT_EQ(bResult, bResultMTD);
 				}
 				{
 					FRigidTransform3 QueryTM(FVec3(0.0, 0.0, 22.0), FQuat(FVec3(1.0, 0.0, 0.0), 3.1415926 / 2.0));
 					QueryTM = FRigidTransform3(QueryTM.GetLocation() * InvScale, QueryTM.GetRotation());
 					bool bResult = TriangleMesh->OverlapGeom(ScaledCapsule, QueryTM, 0.0, nullptr, TriMeshScale);
 					EXPECT_EQ(bResult, false);
+					bool bResultMTD = TriangleMesh->OverlapGeom(ScaledCapsule, QueryTM, 0.0, &MTDInfo, TriMeshScale);
+					EXPECT_EQ(bResult, bResultMTD);
 				}
 				{
 					// Inside Mesh
@@ -471,12 +599,16 @@ namespace ChaosTest
 					QueryTM = FRigidTransform3(QueryTM.GetLocation() * InvScale, QueryTM.GetRotation());
 					bool bResult = TriangleMesh->OverlapGeom(ScaledCapsule, QueryTM, 0.0, nullptr, TriMeshScale);
 					EXPECT_EQ(bResult, false);
+					bool bResultMTD = TriangleMesh->OverlapGeom(ScaledCapsule, QueryTM, 0.0, &MTDInfo, TriMeshScale);
+					EXPECT_EQ(bResult, bResultMTD);
 				}
 				{
 					FRigidTransform3 QueryTM(FVec3(6.0, 0.0, 11.5), FQuat::Identity);
 					QueryTM = FRigidTransform3(QueryTM.GetLocation() * InvScale, QueryTM.GetRotation());
 					bool bResult = TriangleMesh->OverlapGeom(ScaledCapsule, QueryTM, 0.0, nullptr, TriMeshScale);
 					EXPECT_TRUE(bResult);
+					bool bResultMTD = TriangleMesh->OverlapGeom(ScaledCapsule, QueryTM, 0.0, &MTDInfo, TriMeshScale);
+					EXPECT_EQ(bResult, bResultMTD);
 				}
 				{
 					// Triangles parallel to the capsule axis with non uniform scale
@@ -484,6 +616,8 @@ namespace ChaosTest
 					QueryTM = FRigidTransform3(QueryTM.GetLocation() * InvScale, QueryTM.GetRotation());
 					bool bResult = TriangleMesh->OverlapGeom(ScaledCapsule, QueryTM, 0.0, nullptr, TriMeshScale);
 					EXPECT_FALSE(bResult);
+					bool bResultMTD = TriangleMesh->OverlapGeom(ScaledCapsule, QueryTM, 0.0, &MTDInfo, TriMeshScale);
+					EXPECT_EQ(bResult, bResultMTD);
 				}
 			}
 			// Non uniform scale sphere
@@ -499,60 +633,80 @@ namespace ChaosTest
 					QueryTM = FRigidTransform3(QueryTM.GetLocation() * InvScale, QueryTM.GetRotation());
 					bool bResult = TriangleMesh->OverlapGeom(ScaledSphere, QueryTM, 0.0, nullptr, TriMeshScale);
 					EXPECT_EQ(bResult, true);
+					bool bResultMTD = TriangleMesh->OverlapGeom(ScaledSphere, QueryTM, 0.0, &MTDInfo, TriMeshScale);
+					EXPECT_EQ(bResult, bResultMTD);
 				}
 				{
 					FRigidTransform3 QueryTM(FVec3(10.0, 10.0, 0.0), FQuat::Identity);
 					QueryTM = FRigidTransform3(QueryTM.GetLocation() * InvScale, QueryTM.GetRotation());
 					bool bResult = TriangleMesh->OverlapGeom(ScaledSphere, QueryTM, 0.0, nullptr, TriMeshScale);
 					EXPECT_EQ(bResult, false);
+					bool bResultMTD = TriangleMesh->OverlapGeom(ScaledSphere, QueryTM, 0.0, &MTDInfo, TriMeshScale);
+					EXPECT_EQ(bResult, bResultMTD);
 				}
 				{
 					FRigidTransform3 QueryTM(FVec3(12.5, 0.0, 0.0), FQuat::Identity);
 					QueryTM = FRigidTransform3(QueryTM.GetLocation() * InvScale, QueryTM.GetRotation());
 					bool bResult = TriangleMesh->OverlapGeom(ScaledSphere, QueryTM, 0.0, nullptr, TriMeshScale);
 					EXPECT_EQ(bResult, false);
+					bool bResultMTD = TriangleMesh->OverlapGeom(ScaledSphere, QueryTM, 0.0, &MTDInfo, TriMeshScale);
+					EXPECT_EQ(bResult, bResultMTD);
 				}
 				{
 					FRigidTransform3 QueryTM(FVec3(-10.0, -10.0, 0.0), FQuat::Identity);
 					QueryTM = FRigidTransform3(QueryTM.GetLocation() * InvScale, QueryTM.GetRotation());
 					bool bResult = TriangleMesh->OverlapGeom(ScaledSphere, QueryTM, 0.0, nullptr, TriMeshScale);
 					EXPECT_EQ(bResult, false);
+					bool bResultMTD = TriangleMesh->OverlapGeom(ScaledSphere, QueryTM, 0.0, &MTDInfo, TriMeshScale);
+					EXPECT_EQ(bResult, bResultMTD);
 				}
 				{
 					FRigidTransform3 QueryTM(FVec3(10.0, 0.0, 0.0), FQuat::Identity);
 					QueryTM = FRigidTransform3(QueryTM.GetLocation() * InvScale, QueryTM.GetRotation());
 					bool bResult = TriangleMesh->OverlapGeom(ScaledSphere, QueryTM, 0.0, nullptr, TriMeshScale);
 					EXPECT_EQ(bResult, true);
+					bool bResultMTD = TriangleMesh->OverlapGeom(ScaledSphere, QueryTM, 0.0, &MTDInfo, TriMeshScale);
+					EXPECT_EQ(bResult, bResultMTD);
 				}
 				{
 					FRigidTransform3 QueryTM(FVec3(0.0, 10.0, 0.0), FQuat::Identity);
 					QueryTM = FRigidTransform3(QueryTM.GetLocation() * InvScale, QueryTM.GetRotation());
 					bool bResult = TriangleMesh->OverlapGeom(ScaledSphere, QueryTM, 0.0, nullptr, TriMeshScale);
 					EXPECT_EQ(bResult, true);
+					bool bResultMTD = TriangleMesh->OverlapGeom(ScaledSphere, QueryTM, 0.0, &MTDInfo, TriMeshScale);
+					EXPECT_EQ(bResult, bResultMTD);
 				}
 				{
 					FRigidTransform3 QueryTM(FVec3(0.0, -10.0, 0.0), FQuat::Identity);
 					QueryTM = FRigidTransform3(QueryTM.GetLocation() * InvScale, QueryTM.GetRotation());
 					bool bResult = TriangleMesh->OverlapGeom(ScaledSphere, QueryTM, 0.0, nullptr, TriMeshScale);
 					EXPECT_EQ(bResult, true);
+					bool bResultMTD = TriangleMesh->OverlapGeom(ScaledSphere, QueryTM, 0.0, &MTDInfo, TriMeshScale);
+					EXPECT_EQ(bResult, bResultMTD);
 				}
 				{
 					FRigidTransform3 QueryTM(FVec3(-10.0, 0.0, 0.0), FQuat::Identity);
 					QueryTM = FRigidTransform3(QueryTM.GetLocation() * InvScale, QueryTM.GetRotation());
 					bool bResult = TriangleMesh->OverlapGeom(ScaledSphere, QueryTM, 0.0, nullptr, TriMeshScale);
 					EXPECT_EQ(bResult, true);
+					bool bResultMTD = TriangleMesh->OverlapGeom(ScaledSphere, QueryTM, 0.0, &MTDInfo, TriMeshScale);
+					EXPECT_EQ(bResult, bResultMTD);
 				}
 				{
 					FRigidTransform3 QueryTM(FVec3(0.0, 0.0, 20.5), FQuat::Identity);
 					QueryTM = FRigidTransform3(QueryTM.GetLocation() * InvScale, QueryTM.GetRotation());
 					bool bResult = TriangleMesh->OverlapGeom(ScaledSphere, QueryTM, 0.0, nullptr, TriMeshScale);
 					EXPECT_EQ(bResult, true);
+					bool bResultMTD = TriangleMesh->OverlapGeom(ScaledSphere, QueryTM, 0.0, &MTDInfo, TriMeshScale);
+					EXPECT_EQ(bResult, bResultMTD);
 				}
 				{
 					FRigidTransform3 QueryTM(FVec3(0.0, 0.0, 22.0), FQuat(FVec3(1.0, 0.0, 0.0), 3.1415926 / 4.0));
 					QueryTM = FRigidTransform3(QueryTM.GetLocation() * InvScale, QueryTM.GetRotation());
 					bool bResult = TriangleMesh->OverlapGeom(ScaledSphere, QueryTM, 0.0, nullptr, TriMeshScale);
 					EXPECT_EQ(bResult, false);
+					bool bResultMTD = TriangleMesh->OverlapGeom(ScaledSphere, QueryTM, 0.0, &MTDInfo, TriMeshScale);
+					EXPECT_EQ(bResult, bResultMTD);
 				}
 				{
 					// Inside Mesh
@@ -560,12 +714,16 @@ namespace ChaosTest
 					QueryTM = FRigidTransform3(QueryTM.GetLocation() * InvScale, QueryTM.GetRotation());
 					bool bResult = TriangleMesh->OverlapGeom(ScaledSphere, QueryTM, 0.0, nullptr, TriMeshScale);
 					EXPECT_EQ(bResult, false);
+					bool bResultMTD = TriangleMesh->OverlapGeom(ScaledSphere, QueryTM, 0.0, &MTDInfo, TriMeshScale);
+					EXPECT_EQ(bResult, bResultMTD);
 				}
 				{
 					FRigidTransform3 QueryTM(FVec3(5.5, 0.0, 10.5), FQuat::Identity);
 					QueryTM = FRigidTransform3(QueryTM.GetLocation() * InvScale, QueryTM.GetRotation());
 					bool bResult = TriangleMesh->OverlapGeom(ScaledSphere, QueryTM, 0.0, nullptr, TriMeshScale);
 					EXPECT_TRUE(bResult);
+					bool bResultMTD = TriangleMesh->OverlapGeom(ScaledSphere, QueryTM, 0.0, &MTDInfo, TriMeshScale);
+					EXPECT_EQ(bResult, bResultMTD);
 				}
 				{
 					// Triangles parallel to the capsule axis with non uniform scale
@@ -573,6 +731,8 @@ namespace ChaosTest
 					QueryTM = FRigidTransform3(QueryTM.GetLocation() * InvScale, QueryTM.GetRotation());
 					bool bResult = TriangleMesh->OverlapGeom(ScaledSphere, QueryTM, 0.0, nullptr, TriMeshScale);
 					EXPECT_FALSE(bResult);
+					bool bResultMTD = TriangleMesh->OverlapGeom(ScaledSphere, QueryTM, 0.0, &MTDInfo, TriMeshScale);
+					EXPECT_EQ(bResult, bResultMTD);
 				}
 			}
 
@@ -583,82 +743,114 @@ namespace ChaosTest
 					const FRigidTransform3 QueryTM(FVec3(0.0, 0.0, 0.0), FQuat::Identity);
 					bool bResult = TriangleMesh->OverlapGeom(Box, QueryTM, 0.0);
 					EXPECT_EQ(bResult, true);
+					bool bResultMTD = TriangleMesh->OverlapGeom(Box, QueryTM, 0.0, &MTDInfo);
+					EXPECT_EQ(bResult, bResultMTD);
 				}
 				{
 					const FRigidTransform3 QueryTM(FVec3(10.0, 10.0, 0.0), FQuat::Identity);
 					bool bResult = TriangleMesh->OverlapGeom(Box, QueryTM, 0.0);
 					EXPECT_EQ(bResult, false);
+					bool bResultMTD = TriangleMesh->OverlapGeom(Box, QueryTM, 0.0, &MTDInfo);
+					EXPECT_EQ(bResult, bResultMTD);
 				}
 				{
 					const FRigidTransform3 QueryTM(FVec3(-10.0, -10.0, 0.0), FQuat::Identity);
 					bool bResult = TriangleMesh->OverlapGeom(Box, QueryTM, 0.0);
 					EXPECT_EQ(bResult, false);
+					bool bResultMTD = TriangleMesh->OverlapGeom(Box, QueryTM, 0.0, &MTDInfo);
+					EXPECT_EQ(bResult, bResultMTD);
 				}
 				{
 					const FRigidTransform3 QueryTM(FVec3(10.0, 0.0, 0.0), FQuat::Identity);
 					bool bResult = TriangleMesh->OverlapGeom(Box, QueryTM, 0.0);
 					EXPECT_EQ(bResult, true);
+					bool bResultMTD = TriangleMesh->OverlapGeom(Box, QueryTM, 0.0, &MTDInfo);
+					EXPECT_EQ(bResult, bResultMTD);
 				}
 				{
 					const FRigidTransform3 QueryTM(FVec3(0.0, 10.0, 0.0), FQuat::Identity);
 					bool bResult = TriangleMesh->OverlapGeom(Box, QueryTM, 0.0);
 					EXPECT_EQ(bResult, true);
+					bool bResultMTD = TriangleMesh->OverlapGeom(Box, QueryTM, 0.0, &MTDInfo);
+					EXPECT_EQ(bResult, bResultMTD);
 				}
 				{
 					const FRigidTransform3 QueryTM(FVec3(0.0, -10.0, 0.0), FQuat::Identity);
 					bool bResult = TriangleMesh->OverlapGeom(Box, QueryTM, 0.0);
 					EXPECT_EQ(bResult, true);
+					bool bResultMTD = TriangleMesh->OverlapGeom(Box, QueryTM, 0.0, &MTDInfo);
+					EXPECT_EQ(bResult, bResultMTD);
 				}
 				{
 					const FRigidTransform3 QueryTM(FVec3(-10.0, 0.0, 0.0), FQuat::Identity);
 					bool bResult = TriangleMesh->OverlapGeom(Box, QueryTM, 0.0);
 					EXPECT_EQ(bResult, true);
+					bool bResultMTD = TriangleMesh->OverlapGeom(Box, QueryTM, 0.0, &MTDInfo);
+					EXPECT_EQ(bResult, bResultMTD);
 				}
 				{
 					const FRigidTransform3 QueryTM(FVec3(0.0, 0.0, 12.0), FQuat::Identity);
 					bool bResult = TriangleMesh->OverlapGeom(Box, QueryTM, 0.0);
 					EXPECT_EQ(bResult, true);
+					bool bResultMTD = TriangleMesh->OverlapGeom(Box, QueryTM, 0.0, &MTDInfo);
+					EXPECT_EQ(bResult, bResultMTD);
 				}
 				{
 					const FRigidTransform3 QueryTM(FVec3(0.0, 0.0, 12.0), FQuat(FVec3(1.0, 0.0, 0.0), 3.1415926 / 4.0));
 					bool bResult = TriangleMesh->OverlapGeom(Box, QueryTM, 0.0);
 					EXPECT_EQ(bResult, false);
+					bool bResultMTD = TriangleMesh->OverlapGeom(Box, QueryTM, 0.0, &MTDInfo);
+					EXPECT_EQ(bResult, bResultMTD);
 				}
 				{
 					// Inside Mesh
 					const FRigidTransform3 QueryTM(FVec3(0.0, 0.0, 4.0), FQuat::Identity);
 					bool bResult = TriangleMesh->OverlapGeom(Box, QueryTM, 0.0);
 					EXPECT_EQ(bResult, false);
+					bool bResultMTD = TriangleMesh->OverlapGeom(Box, QueryTM, 0.0, &MTDInfo);
+					EXPECT_EQ(bResult, bResultMTD);
 				}
 				{
 					const FRigidTransform3 QueryTM(FVec3(0.0, 0.0, -2.5), FQuat(FVec3(1.0, 0.0, 0.0), 3.1415926 / 4.0));
 					bool bResult = TriangleMesh->OverlapGeom(Box, QueryTM, 0.0);
 					EXPECT_EQ(bResult, false);
+					bool bResultMTD = TriangleMesh->OverlapGeom(Box, QueryTM, 0.0, &MTDInfo);
+					EXPECT_EQ(bResult, bResultMTD);
 				}
 				{
 					const FRigidTransform3 QueryTM(FVec3(5.0, 5.0, -2.5), FQuat(FVec3(0.0, 1.0, 0.0), 3.1415926 / 6.0));
 					bool bResult = TriangleMesh->OverlapGeom(Box, QueryTM, 0.0);
 					EXPECT_EQ(bResult, false);
+					bool bResultMTD = TriangleMesh->OverlapGeom(Box, QueryTM, 0.0, &MTDInfo);
+					EXPECT_EQ(bResult, bResultMTD);
 				}
 				{
 					const FRigidTransform3 QueryTM(FVec3(-5.0, 5.0, -2.5), FQuat(FVec3(1.0, 0.0, 0.0), 3.1415926 / 4.0));
 					bool bResult = TriangleMesh->OverlapGeom(Box, QueryTM, 0.0);
 					EXPECT_EQ(bResult, false);
+					bool bResultMTD = TriangleMesh->OverlapGeom(Box, QueryTM, 0.0, &MTDInfo);
+					EXPECT_EQ(bResult, bResultMTD);
 				}
 				{
 					const FRigidTransform3 QueryTM(FVec3(-5.0, -5.0, -2.5), FQuat(FVec3(1.0, 0.0, 0.0), -3.1415926 / 6.0));
 					bool bResult = TriangleMesh->OverlapGeom(Box, QueryTM, 0.0);
 					EXPECT_EQ(bResult, false);
+					bool bResultMTD = TriangleMesh->OverlapGeom(Box, QueryTM, 0.0, &MTDInfo);
+					EXPECT_EQ(bResult, bResultMTD);
 				}
 				{
 					const FRigidTransform3 QueryTM(FVec3(0.0, 0.0, -2.5), FQuat(FVec3(1.0, 0.0, 0.0), 3.1415926 / 4.0));
 					bool bResult = TriangleMesh->OverlapGeom(Box, QueryTM, 0.0);
 					EXPECT_EQ(bResult, false);
+					bool bResultMTD = TriangleMesh->OverlapGeom(Box, QueryTM, 0.0, &MTDInfo);
+					EXPECT_EQ(bResult, bResultMTD);
 				}
 				{
 					const FRigidTransform3 QueryTM(FVec3(7.0, 0.0, 10.0), FQuat::Identity);
 					bool bResult = TriangleMesh->OverlapGeom(Box, QueryTM, 0.0);
 					EXPECT_EQ(bResult, false);
+					bool bResultMTD = TriangleMesh->OverlapGeom(Box, QueryTM, 0.0, &MTDInfo);
+					EXPECT_EQ(bResult, bResultMTD);
 				}
 			}
 			{
@@ -680,52 +872,72 @@ namespace ChaosTest
 					const FRigidTransform3 QueryTM(FVec3(0.0, 0.0, 0.0), FQuat::Identity);
 					bool bResult = TriangleMesh->OverlapGeom(ScaledBox, QueryTM, 0.0, nullptr, TriMeshScale);
 					EXPECT_EQ(bResult, true);
+					bool bResultMTD = TriangleMesh->OverlapGeom(ScaledBox, QueryTM, 0.0, &MTDInfo, TriMeshScale);
+					EXPECT_EQ(bResult, bResultMTD);
 				}
 				{
 					const FRigidTransform3 QueryTM(FVec3(10.0, 10.0, 0.0), FQuat::Identity);
 					bool bResult = TriangleMesh->OverlapGeom(ScaledBox, QueryTM, 0.0, nullptr, TriMeshScale);
 					EXPECT_EQ(bResult, false);
+					bool bResultMTD = TriangleMesh->OverlapGeom(ScaledBox, QueryTM, 0.0, &MTDInfo, TriMeshScale);
+					EXPECT_EQ(bResult, bResultMTD);
 				}
 				{
 					const FRigidTransform3 QueryTM(FVec3(-10.0, -10.0, 0.0), FQuat::Identity);
 					bool bResult = TriangleMesh->OverlapGeom(ScaledBox, QueryTM, 0.0, nullptr, TriMeshScale);
 					EXPECT_EQ(bResult, false);
+					bool bResultMTD = TriangleMesh->OverlapGeom(ScaledBox, QueryTM, 0.0, &MTDInfo, TriMeshScale);
+					EXPECT_EQ(bResult, bResultMTD);
 				}
 				{
 					const FRigidTransform3 QueryTM(FVec3(10.0, 0.0, 0.0), FQuat::Identity);
 					bool bResult = TriangleMesh->OverlapGeom(ScaledBox, QueryTM, 0.0, nullptr, TriMeshScale);
 					EXPECT_EQ(bResult, true);
+					bool bResultMTD = TriangleMesh->OverlapGeom(ScaledBox, QueryTM, 0.0, &MTDInfo, TriMeshScale);
+					EXPECT_EQ(bResult, bResultMTD);
 				}
 				{
 					const FRigidTransform3 QueryTM(FVec3(0.0, 10.0, 0.0), FQuat::Identity);
 					bool bResult = TriangleMesh->OverlapGeom(ScaledBox, QueryTM, 0.0, nullptr, TriMeshScale);
 					EXPECT_EQ(bResult, true);
+					bool bResultMTD = TriangleMesh->OverlapGeom(ScaledBox, QueryTM, 0.0, &MTDInfo, TriMeshScale);
+					EXPECT_EQ(bResult, bResultMTD);
 				}
 				{
 					const FRigidTransform3 QueryTM(FVec3(0.0, -10.0, 0.0), FQuat::Identity);
 					bool bResult = TriangleMesh->OverlapGeom(ScaledBox, QueryTM, 0.0, nullptr, TriMeshScale);
 					EXPECT_EQ(bResult, true);
+					bool bResultMTD = TriangleMesh->OverlapGeom(ScaledBox, QueryTM, 0.0, &MTDInfo, TriMeshScale);
+					EXPECT_EQ(bResult, bResultMTD);
 				}
 				{
 					const FRigidTransform3 QueryTM(FVec3(-10.0, 0.0, 0.0), FQuat::Identity);
 					bool bResult = TriangleMesh->OverlapGeom(ScaledBox, QueryTM, 0.0, nullptr, TriMeshScale);
 					EXPECT_EQ(bResult, true);
+					bool bResultMTD = TriangleMesh->OverlapGeom(ScaledBox, QueryTM, 0.0, &MTDInfo, TriMeshScale);
+					EXPECT_EQ(bResult, bResultMTD);
 				}
 				{
 					const FRigidTransform3 QueryTM(FVec3(0.0, 0.0, 12.0), FQuat::Identity);
 					bool bResult = TriangleMesh->OverlapGeom(ScaledBox, QueryTM, 0.0, nullptr, TriMeshScale);
 					EXPECT_EQ(bResult, true);
+					bool bResultMTD = TriangleMesh->OverlapGeom(ScaledBox, QueryTM, 0.0, &MTDInfo, TriMeshScale);
+					EXPECT_EQ(bResult, bResultMTD);
 				}
 				{
 					const FRigidTransform3 QueryTM(FVec3(0.0, 0.0, 12.0), FQuat(FVec3(1.0, 0.0, 0.0), 3.1415926 / 4.0));
 					bool bResult = TriangleMesh->OverlapGeom(ScaledBox, QueryTM, 0.0, nullptr, TriMeshScale);
 					EXPECT_EQ(bResult, false);
+					bool bResultMTD = TriangleMesh->OverlapGeom(ScaledBox, QueryTM, 0.0, &MTDInfo, TriMeshScale);
+					EXPECT_EQ(bResult, bResultMTD);
 				}
 				{
 					// Inside Mesh
 					const FRigidTransform3 QueryTM(FVec3(0.0, 0.0, 4.0), FQuat::Identity);
 					bool bResult = TriangleMesh->OverlapGeom(ScaledBox, QueryTM, 0.0, nullptr, TriMeshScale);
 					EXPECT_EQ(bResult, false);
+					bool bResultMTD = TriangleMesh->OverlapGeom(ScaledBox, QueryTM, 0.0, &MTDInfo, TriMeshScale);
+					EXPECT_EQ(bResult, bResultMTD);
 				}
 			}
 			// Box Scaled test non uniform transform with rotation
@@ -733,26 +945,33 @@ namespace ChaosTest
 				TSharedPtr<TBox<FReal, 3>> BigBoxSafe = MakeShared<TBox<FReal, 3>>(FVec3(-100.0, -100.0, -100.0), FVec3(100.0, 100.0, 100.0));
 				FVec3 TriMeshScale = { 0.01, 0.01, 0.05 };
 				TImplicitObjectScaled<TBox<FReal, 3>> ScaledBox = TImplicitObjectScaled<TBox<FReal, 3>>(BigBoxSafe, TriMeshScale);
-
 				{
 					const FRigidTransform3 QueryTM(FVec3(5.5, 0.0, 10.0), FQuat::Identity);
 					bool bResult = TriangleMesh->OverlapGeom(ScaledBox, QueryTM, 0.0, nullptr, TriMeshScale);
 					EXPECT_EQ(bResult, true);
+					bool bResultMTD = TriangleMesh->OverlapGeom(ScaledBox, QueryTM, 0.0, &MTDInfo, TriMeshScale);
+					EXPECT_EQ(bResult, bResultMTD);
 				}
 				{
-					const FRigidTransform3 QueryTM(FVec3(7.0, 0.0, 10.0), FQuat::Identity);
+					const FRigidTransform3 QueryTM(FVec3(7.0, 0.0, 10.001), FQuat::Identity);
 					bool bResult = TriangleMesh->OverlapGeom(ScaledBox, QueryTM, 0.0, nullptr, TriMeshScale);
 					EXPECT_EQ(bResult, false);
+					bool bResultMTD = TriangleMesh->OverlapGeom(ScaledBox, QueryTM, 0.0, &MTDInfo, TriMeshScale);
+					EXPECT_EQ(bResult, bResultMTD);
 				}
 				{
 					const FRigidTransform3 QueryTM(FVec3(7.0, 0.0, 10.0), FQuat(FVec3(0.0, 1.0, 0.0), -3.1415926 / 4.0));
 					bool bResult = TriangleMesh->OverlapGeom(ScaledBox, QueryTM, 0.0, nullptr, TriMeshScale);
 					EXPECT_EQ(bResult, false);
+					bool bResultMTD = TriangleMesh->OverlapGeom(ScaledBox, QueryTM, 0.0, &MTDInfo, TriMeshScale);
+					EXPECT_EQ(bResult, bResultMTD);
 				}
 				{
 					const FRigidTransform3 QueryTM(FVec3(7.0, 0.0, 10.0), FQuat(FVec3(0.0, 1.0, 0.0), 3.1415926 / 4.0));
 					bool bResult = TriangleMesh->OverlapGeom(ScaledBox, QueryTM, 0.0, nullptr, TriMeshScale);
 					EXPECT_EQ(bResult, true);
+					bool bResultMTD = TriangleMesh->OverlapGeom(ScaledBox, QueryTM, 0.0, &MTDInfo, TriMeshScale);
+					EXPECT_EQ(bResult, bResultMTD);
 				}
 			}
 			{
@@ -760,12 +979,13 @@ namespace ChaosTest
 				FVec3 TriMeshScale = { 10.0, 10.0, 2.0 };
 				FVec3 InvScale = 1.0 / TriMeshScale;
 				TImplicitObjectScaled<TBox<FReal, 3>> ScaledBox = TImplicitObjectScaled<TBox<FReal, 3>>(BigBoxSafe, InvScale);
-
 				{
 					FRigidTransform3 QueryTM(FVec3(50, 0.0, 10.0), FQuat::Identity);
 					QueryTM = FRigidTransform3(QueryTM.GetLocation() * InvScale, QueryTM.GetRotation());
 					bool bResult = TriangleMesh->OverlapGeom(ScaledBox, QueryTM, 0.0, nullptr, TriMeshScale);
 					EXPECT_EQ(bResult, true);
+					bool bResultMTD = TriangleMesh->OverlapGeom(ScaledBox, QueryTM, 0.0, &MTDInfo, TriMeshScale);
+					EXPECT_EQ(bResult, bResultMTD);
 				}
 				{
 					// Inside mesh
@@ -773,36 +993,48 @@ namespace ChaosTest
 					QueryTM = FRigidTransform3(QueryTM.GetLocation() * InvScale, QueryTM.GetRotation());
 					bool bResult = TriangleMesh->OverlapGeom(ScaledBox, QueryTM, 0.0, nullptr, TriMeshScale);
 					EXPECT_EQ(bResult, false);
+					bool bResultMTD = TriangleMesh->OverlapGeom(ScaledBox, QueryTM, 0.0, &MTDInfo, TriMeshScale);
+					EXPECT_EQ(bResult, bResultMTD);
 				}
 				{
 					FRigidTransform3 QueryTM(FVec3(52.0, 0.0, 11.0), FQuat::Identity);
 					QueryTM = FRigidTransform3(QueryTM.GetLocation() * InvScale, QueryTM.GetRotation());
 					bool bResult = TriangleMesh->OverlapGeom(ScaledBox, QueryTM, 0.0, nullptr, TriMeshScale);
 					EXPECT_EQ(bResult, false);
+					bool bResultMTD = TriangleMesh->OverlapGeom(ScaledBox, QueryTM, 0.0, &MTDInfo, TriMeshScale);
+					EXPECT_EQ(bResult, bResultMTD);
 				}
 				{
 					FRigidTransform3 QueryTM(FVec3(52.0, 0.0, 11.0), FQuat(FVec3(0.0, 1.0, 0.0), -3.1415926 / 4.0));
 					QueryTM = FRigidTransform3(QueryTM.GetLocation() * InvScale, QueryTM.GetRotation());
 					bool bResult = TriangleMesh->OverlapGeom(ScaledBox, QueryTM, 0.0, nullptr, TriMeshScale);
 					EXPECT_EQ(bResult, true);
+					bool bResultMTD = TriangleMesh->OverlapGeom(ScaledBox, QueryTM, 0.0, &MTDInfo, TriMeshScale);
+					EXPECT_EQ(bResult, bResultMTD);
 				}
 				{
 					FRigidTransform3 QueryTM(FVec3(51.0, 0.0, 11.0), FQuat::Identity);
 					QueryTM = FRigidTransform3(QueryTM.GetLocation() * InvScale, QueryTM.GetRotation());
 					bool bResult = TriangleMesh->OverlapGeom(ScaledBox, QueryTM, 0.0, nullptr, TriMeshScale);
 					EXPECT_EQ(bResult, true);
+					bool bResultMTD = TriangleMesh->OverlapGeom(ScaledBox, QueryTM, 0.0, &MTDInfo, TriMeshScale);
+					EXPECT_EQ(bResult, bResultMTD);
 				}
 				{
 					FRigidTransform3 QueryTM(FVec3(51.5, 0.0, 11.0), FQuat(FVec3(0.0, 1.0, 0.0), 4.0 / 5.0 * 3.1415926 / 2.0));
 					QueryTM = FRigidTransform3(QueryTM.GetLocation() * InvScale, QueryTM.GetRotation());
 					bool bResult = TriangleMesh->OverlapGeom(ScaledBox, QueryTM, 0.0, nullptr, TriMeshScale);
 					EXPECT_EQ(bResult, true);
+					bool bResultMTD = TriangleMesh->OverlapGeom(ScaledBox, QueryTM, 0.0, &MTDInfo, TriMeshScale);
+					EXPECT_EQ(bResult, bResultMTD);
 				}
 				{
 					FRigidTransform3 QueryTM(FVec3(51.5, 0.0, 11.0), FQuat(FVec3(0.0, 1.0, 0.0), -4.0 / 5.0 * 3.1415926 / 2.0));
 					QueryTM = FRigidTransform3(QueryTM.GetLocation() * InvScale, QueryTM.GetRotation());
 					bool bResult = TriangleMesh->OverlapGeom(ScaledBox, QueryTM, 0.0, nullptr, TriMeshScale);
 					EXPECT_EQ(bResult, false);
+					bool bResultMTD = TriangleMesh->OverlapGeom(ScaledBox, QueryTM, 0.0, &MTDInfo, TriMeshScale);
+					EXPECT_EQ(bResult, bResultMTD);
 				}
 			}
 			{
@@ -811,12 +1043,13 @@ namespace ChaosTest
 				FVec3 TriMeshScale = { 10.0, 10.0, 2.0 };
 				FVec3 InvScale = 1.0 / TriMeshScale;
 				TImplicitObjectScaled<TBox<FReal, 3>> ScaledBox = TImplicitObjectScaled<TBox<FReal, 3>>(BigBoxSafe, InvScale);
-
 				{
 					FRigidTransform3 QueryTM(FVec3(50, 0.0, 10.0), FQuat::Identity);
 					QueryTM = FRigidTransform3(QueryTM.GetLocation() * InvScale, QueryTM.GetRotation());
 					bool bResult = TriangleMesh->OverlapGeom(ScaledBox, QueryTM, 0.0, nullptr, TriMeshScale);
 					EXPECT_EQ(bResult, true);
+					bool bResultMTD = TriangleMesh->OverlapGeom(ScaledBox, QueryTM, 0.0, &MTDInfo, TriMeshScale);
+					EXPECT_EQ(bResult, bResultMTD);
 				}
 				{
 					// Inside mesh
@@ -824,42 +1057,56 @@ namespace ChaosTest
 					QueryTM = FRigidTransform3(QueryTM.GetLocation() * InvScale, QueryTM.GetRotation());
 					bool bResult = TriangleMesh->OverlapGeom(ScaledBox, QueryTM, 0.0, nullptr, TriMeshScale);
 					EXPECT_EQ(bResult, true);
+					bool bResultMTD = TriangleMesh->OverlapGeom(ScaledBox, QueryTM, 0.0, &MTDInfo, TriMeshScale);
+					EXPECT_EQ(bResult, bResultMTD);
 				}
 				{
 					FRigidTransform3 QueryTM(FVec3(52.0, 0.0, 11.0), FQuat::Identity);
 					QueryTM = FRigidTransform3(QueryTM.GetLocation() * InvScale, QueryTM.GetRotation());
 					bool bResult = TriangleMesh->OverlapGeom(ScaledBox, QueryTM, 0.0, nullptr, TriMeshScale);
 					EXPECT_EQ(bResult, false);
+					bool bResultMTD = TriangleMesh->OverlapGeom(ScaledBox, QueryTM, 0.0, &MTDInfo, TriMeshScale);
+					EXPECT_EQ(bResult, bResultMTD);
 				}
 				{
 					FRigidTransform3 QueryTM(FVec3(0.0, 52.0, 11.0), FQuat::Identity);
 					QueryTM = FRigidTransform3(QueryTM.GetLocation() * InvScale, QueryTM.GetRotation());
 					bool bResult = TriangleMesh->OverlapGeom(ScaledBox, QueryTM, 0.0, nullptr, TriMeshScale);
 					EXPECT_EQ(bResult, true);
+					bool bResultMTD = TriangleMesh->OverlapGeom(ScaledBox, QueryTM, 0.0, &MTDInfo, TriMeshScale);
+					EXPECT_EQ(bResult, bResultMTD);
 				}
 				{
 					FRigidTransform3 QueryTM(FVec3(54.0, 0.0, 11.0), FQuat(FVec3(0.0, 1.0, 0.0), -3.1415926 / 2.0 * (1.5 / 5.0)));
 					QueryTM = FRigidTransform3(QueryTM.GetLocation() * InvScale, QueryTM.GetRotation());
 					bool bResult = TriangleMesh->OverlapGeom(ScaledBox, QueryTM, 0.0, nullptr, TriMeshScale);
 					EXPECT_EQ(bResult, false);
+					bool bResultMTD = TriangleMesh->OverlapGeom(ScaledBox, QueryTM, 0.0, &MTDInfo, TriMeshScale);
+					EXPECT_EQ(bResult, bResultMTD);
 				}
 				{
 					FRigidTransform3 QueryTM(FVec3(0.0, 54.0, 11.0), FQuat(FVec3(1.0, 0.0, 0.0), -3.1415926 / 2.0 * (1.5 / 5.0)));
 					QueryTM = FRigidTransform3(QueryTM.GetLocation() * InvScale, QueryTM.GetRotation());
 					bool bResult = TriangleMesh->OverlapGeom(ScaledBox, QueryTM, 0.0, nullptr, TriMeshScale);
 					EXPECT_EQ(bResult, true);
+					bool bResultMTD = TriangleMesh->OverlapGeom(ScaledBox, QueryTM, 0.0, &MTDInfo, TriMeshScale);
+					EXPECT_EQ(bResult, bResultMTD);
 				}
 				{
 					FRigidTransform3 QueryTM(FVec3(52.0, 0.0, 11.0), FQuat(FVec3(0.0, 1.0, 0.0), (1.5 / 5.0) * 3.1415926 / 2.0));
 					QueryTM = FRigidTransform3(QueryTM.GetLocation() * InvScale, QueryTM.GetRotation());
 					bool bResult = TriangleMesh->OverlapGeom(ScaledBox, QueryTM, 0.0, nullptr, TriMeshScale);
 					EXPECT_EQ(bResult, false);
+					bool bResultMTD = TriangleMesh->OverlapGeom(ScaledBox, QueryTM, 0.0, &MTDInfo, TriMeshScale);
+					EXPECT_EQ(bResult, bResultMTD);
 				}
 				{
 					FRigidTransform3 QueryTM(FVec3(0.0, 52.0, 11.0), FQuat(FVec3(1.0, 0.0, 0.0), (1.5 / 5.0) * 3.1415926 / 2.0));
 					QueryTM = FRigidTransform3(QueryTM.GetLocation() * InvScale, QueryTM.GetRotation());
 					bool bResult = TriangleMesh->OverlapGeom(ScaledBox, QueryTM, 0.0, nullptr, TriMeshScale);
 					EXPECT_EQ(bResult, true);
+					bool bResultMTD = TriangleMesh->OverlapGeom(ScaledBox, QueryTM, 0.0, &MTDInfo, TriMeshScale);
+					EXPECT_EQ(bResult, bResultMTD);
 				}
 			}
 			{
@@ -869,52 +1116,72 @@ namespace ChaosTest
 					const FRigidTransform3 QueryTM(FVec3(0.0, 0.0, 0.0), FQuat::Identity);
 					bool bResult = TriangleMesh->OverlapGeom(Sphere, QueryTM, 0.0);
 					EXPECT_EQ(bResult, true);
+					bool bResultMTD = TriangleMesh->OverlapGeom(Sphere, QueryTM, 0.0, &MTDInfo);
+					EXPECT_EQ(bResult, bResultMTD);
 				}
 				{
 					const FRigidTransform3 QueryTM(FVec3(10.0, 10.0, 0.0), FQuat::Identity);
 					bool bResult = TriangleMesh->OverlapGeom(Sphere, QueryTM, 0.0);
 					EXPECT_EQ(bResult, false);
+					bool bResultMTD = TriangleMesh->OverlapGeom(Sphere, QueryTM, 0.0, &MTDInfo);
+					EXPECT_EQ(bResult, bResultMTD);
 				}
 				{
 					const FRigidTransform3 QueryTM(FVec3(-10.0, -10.0, 0.0), FQuat::Identity);
 					bool bResult = TriangleMesh->OverlapGeom(Sphere, QueryTM, 0.0);
 					EXPECT_EQ(bResult, false);
+					bool bResultMTD = TriangleMesh->OverlapGeom(Sphere, QueryTM, 0.0, &MTDInfo);
+					EXPECT_EQ(bResult, bResultMTD);
 				}
 				{
 					const FRigidTransform3 QueryTM(FVec3(10.0, 0.0, 0.0), FQuat::Identity);
 					bool bResult = TriangleMesh->OverlapGeom(Sphere, QueryTM, 0.0);
 					EXPECT_EQ(bResult, true);
+					bool bResultMTD = TriangleMesh->OverlapGeom(Sphere, QueryTM, 0.0, &MTDInfo);
+					EXPECT_EQ(bResult, bResultMTD);
 				}
 				{
 					const FRigidTransform3 QueryTM(FVec3(0.0, 10.0, 0.0), FQuat::Identity);
 					bool bResult = TriangleMesh->OverlapGeom(Sphere, QueryTM, 0.0);
 					EXPECT_EQ(bResult, true);
+					bool bResultMTD = TriangleMesh->OverlapGeom(Sphere, QueryTM, 0.0, &MTDInfo);
+					EXPECT_EQ(bResult, bResultMTD);
 				}
 				{
 					const FRigidTransform3 QueryTM(FVec3(0.0, -10.0, 0.0), FQuat::Identity);
 					bool bResult = TriangleMesh->OverlapGeom(Sphere, QueryTM, 0.0);
 					EXPECT_EQ(bResult, true);
+					bool bResultMTD = TriangleMesh->OverlapGeom(Sphere, QueryTM, 0.0, &MTDInfo);
+					EXPECT_EQ(bResult, bResultMTD);
 				}
 				{
 					const FRigidTransform3 QueryTM(FVec3(-10.0, 0.0, 0.0), FQuat::Identity);
 					bool bResult = TriangleMesh->OverlapGeom(Sphere, QueryTM, 0.0);
 					EXPECT_EQ(bResult, true);
+					bool bResultMTD = TriangleMesh->OverlapGeom(Sphere, QueryTM, 0.0, &MTDInfo);
+					EXPECT_EQ(bResult, bResultMTD);
 				}
 				{
 					const FRigidTransform3 QueryTM(FVec3(0.0, 0.0, 12.0), FQuat::Identity);
 					bool bResult = TriangleMesh->OverlapGeom(Sphere, QueryTM, 0.0);
 					EXPECT_EQ(bResult, true);
+					bool bResultMTD = TriangleMesh->OverlapGeom(Sphere, QueryTM, 0.0, &MTDInfo);
+					EXPECT_EQ(bResult, bResultMTD);
 				}
 				{
 					const FRigidTransform3 QueryTM(FVec3(0.0, 0.0, 12.0), FQuat(FVec3(1.0, 0.0, 0.0), 3.1415926 / 4.0));
 					bool bResult = TriangleMesh->OverlapGeom(Sphere, QueryTM, 0.0);
 					EXPECT_EQ(bResult, true);
+					bool bResultMTD = TriangleMesh->OverlapGeom(Sphere, QueryTM, 0.0, &MTDInfo);
+					EXPECT_EQ(bResult, bResultMTD);
 				}
 				{
 					// Inside Mesh
 					const FRigidTransform3 QueryTM(FVec3(0.0, 0.0, 4.0), FQuat::Identity);
 					bool bResult = TriangleMesh->OverlapGeom(Sphere, QueryTM, 0.0);
 					EXPECT_EQ(bResult, false);
+					bool bResultMTD = TriangleMesh->OverlapGeom(Sphere, QueryTM, 0.0, &MTDInfo);
+					EXPECT_EQ(bResult, bResultMTD);
 				}
 			}
 		}
@@ -949,6 +1216,7 @@ namespace ChaosTest
 			}
 
 			TUniquePtr<FTriangleMeshImplicitObject> TriangleMesh = MakeUnique<FTriangleMeshImplicitObject>(MoveTemp(TrimeshParticles), MoveTemp(Indices), MoveTemp(Materials));
+			FMTDInfo MTDInfo;
 			{
 				// Capsule test
 				const FVec3 X1 = { 0.0, 0.0, -2.0 };
@@ -960,26 +1228,36 @@ namespace ChaosTest
 					const FRigidTransform3 QueryTM(FVec3(0.0, 0.0, 5.0), FQuat::Identity);
 					bool bResult = TriangleMesh->OverlapGeom(Capsule, QueryTM, 0.0);
 					EXPECT_EQ(bResult, false);
+					bool bResultMTD = TriangleMesh->OverlapGeom(Capsule, QueryTM, 0.0, &MTDInfo);
+					EXPECT_EQ(bResult, bResultMTD);
 				}
 				{
 					const FRigidTransform3 QueryTM(FVec3(10.0, 0.0, 11.0), FQuat::Identity);
 					bool bResult = TriangleMesh->OverlapGeom(Capsule, QueryTM, 0.0);
 					EXPECT_EQ(bResult, true);
+					bool bResultMTD = TriangleMesh->OverlapGeom(Capsule, QueryTM, 0.0, &MTDInfo);
+					EXPECT_EQ(bResult, bResultMTD);
 				}
 				{
 					const FRigidTransform3 QueryTM(FVec3(11.5, 0.0, 11.0), FQuat::Identity);
 					bool bResult = TriangleMesh->OverlapGeom(Capsule, QueryTM, 0.0);
 					EXPECT_EQ(bResult, false);
+					bool bResultMTD = TriangleMesh->OverlapGeom(Capsule, QueryTM, 0.0, &MTDInfo);
+					EXPECT_EQ(bResult, bResultMTD);
 				}
 				{
 					const FRigidTransform3 QueryTM(FVec3(11.5, 0.0, 11.0), FQuat(FVec3(0.0, 1.0, 0.0), 3.1415926 / 4.0));
 					bool bResult = TriangleMesh->OverlapGeom(Capsule, QueryTM, 0.0);
 					EXPECT_EQ(bResult, true);
+					bool bResultMTD = TriangleMesh->OverlapGeom(Capsule, QueryTM, 0.0, &MTDInfo);
+					EXPECT_EQ(bResult, bResultMTD);
 				}
 				{
 					const FRigidTransform3 QueryTM(FVec3(11.5, 0.0, 11.0), FQuat(FVec3(0.0, 1.0, 0.0), -3.1415926 / 4.0));
 					bool bResult = TriangleMesh->OverlapGeom(Capsule, QueryTM, 0.0);
 					EXPECT_EQ(bResult, false);
+					bool bResultMTD = TriangleMesh->OverlapGeom(Capsule, QueryTM, 0.0, &MTDInfo);
+					EXPECT_EQ(bResult, bResultMTD);
 				}
 			}
 			{
@@ -992,6 +1270,8 @@ namespace ChaosTest
 					const FRigidTransform3 QueryTM(FVec3(0.0, 0.0, 5.0), FQuat::Identity);
 					bool bResult = TriangleMesh->OverlapGeom(Capsule, QueryTM, 0.0);
 					EXPECT_TRUE(bResult);
+					bool bResultMTD = TriangleMesh->OverlapGeom(Capsule, QueryTM, 0.0, &MTDInfo);
+					EXPECT_EQ(bResult, bResultMTD);
 				}
 			}
 			{
@@ -1003,26 +1283,36 @@ namespace ChaosTest
 					const FRigidTransform3 QueryTM(FVec3(0.0, 0.0, 0.0), FQuat::Identity);
 					bool bResult = TriangleMesh->OverlapGeom(Capsule, QueryTM, 0.0);
 					EXPECT_TRUE(bResult);
+					bool bResultMTD = TriangleMesh->OverlapGeom(Capsule, QueryTM, 0.0, &MTDInfo);
+					EXPECT_EQ(bResult, bResultMTD);
 				}
 				{
 					const FRigidTransform3 QueryTM(FVec3(0.0, 0.0, 5.0), FQuat::Identity);
 					bool bResult = TriangleMesh->OverlapGeom(Capsule, QueryTM, 0.0);
 					EXPECT_FALSE(bResult);
+					bool bResultMTD = TriangleMesh->OverlapGeom(Capsule, QueryTM, 0.0, &MTDInfo);
+					EXPECT_EQ(bResult, bResultMTD);
 				}
 				{
 					const FRigidTransform3 QueryTM(FVec3(0.0, 0.0, 10.0), FQuat::Identity);
 					bool bResult = TriangleMesh->OverlapGeom(Capsule, QueryTM, 0.0);
 					EXPECT_TRUE(bResult);
+					bool bResultMTD = TriangleMesh->OverlapGeom(Capsule, QueryTM, 0.0, &MTDInfo);
+					EXPECT_EQ(bResult, bResultMTD);
 				}
 				{
 					const FRigidTransform3 QueryTM(FVec3(5.0, -10.25, 0.0), FQuat::Identity);
 					bool bResult = TriangleMesh->OverlapGeom(Capsule, QueryTM, 0.0);
 					EXPECT_TRUE(bResult);
+					bool bResultMTD = TriangleMesh->OverlapGeom(Capsule, QueryTM, 0.0, &MTDInfo);
+					EXPECT_EQ(bResult, bResultMTD);
 				}
 				{
 					const FRigidTransform3 QueryTM(FVec3(5.0, -10.6, 0.0), FQuat::Identity);
 					bool bResult = TriangleMesh->OverlapGeom(Capsule, QueryTM, 0.0);
 					EXPECT_FALSE(bResult);
+					bool bResultMTD = TriangleMesh->OverlapGeom(Capsule, QueryTM, 0.0, &MTDInfo);
+					EXPECT_EQ(bResult, bResultMTD);
 				}
 			}
 		}
@@ -1046,6 +1336,7 @@ namespace ChaosTest
 			}
 
 			TUniquePtr<FTriangleMeshImplicitObject> TriangleMesh = MakeUnique<FTriangleMeshImplicitObject>(MoveTemp(TrimeshParticles), MoveTemp(Indices), MoveTemp(Materials));
+			FMTDInfo MTDInfo;
 			{
 				// Sphere test
 				const FVec3 X = { 0.0, 0.0, 0.0 };
@@ -1055,21 +1346,29 @@ namespace ChaosTest
 					const FRigidTransform3 QueryTM(FVec3(0.0, 0.0, 0.0), FQuat::Identity);
 					bool bResult = TriangleMesh->OverlapGeom(Sphere, QueryTM, 0.0);
 					EXPECT_TRUE(bResult);
+					bool bResultMTD = TriangleMesh->OverlapGeom(Sphere, QueryTM, 0.0, &MTDInfo);
+					EXPECT_EQ(bResult, bResultMTD);
 				}
 				{
 					const FRigidTransform3 QueryTM(FVec3(3.9, 0.0, 100.0), FQuat::Identity);
 					bool bResult = TriangleMesh->OverlapGeom(Sphere, QueryTM, 0.0);
 					EXPECT_TRUE(bResult);
+					bool bResultMTD = TriangleMesh->OverlapGeom(Sphere, QueryTM, 0.0, &MTDInfo);
+					EXPECT_EQ(bResult, bResultMTD);
 				}
 				{
 					const FRigidTransform3 QueryTM(FVec3(4.0, 0.0, 90.0), FQuat::Identity);
 					bool bResult = TriangleMesh->OverlapGeom(Sphere, QueryTM, 0.0);
 					EXPECT_TRUE(bResult);
+					bool bResultMTD = TriangleMesh->OverlapGeom(Sphere, QueryTM, 0.0, &MTDInfo);
+					EXPECT_EQ(bResult, bResultMTD);
 				}
 				{
 					const FRigidTransform3 QueryTM(FVec3(4.5, 0.0, 90.0), FQuat::Identity);
 					bool bResult = TriangleMesh->OverlapGeom(Sphere, QueryTM, 0.0);
 					EXPECT_FALSE(bResult);
+					bool bResultMTD = TriangleMesh->OverlapGeom(Sphere, QueryTM, 0.0, &MTDInfo);
+					EXPECT_EQ(bResult, bResultMTD);
 				}
 			}
 			{
@@ -1082,21 +1381,29 @@ namespace ChaosTest
 					const FRigidTransform3 QueryTM(FVec3(0.0, 0.0, 0.0), FQuat::Identity);
 					bool bResult = TriangleMesh->OverlapGeom(Capsule, QueryTM, 0.0);
 					EXPECT_TRUE(bResult);
+					bool bResultMTD = TriangleMesh->OverlapGeom(Capsule, QueryTM, 0.0, &MTDInfo);
+					EXPECT_EQ(bResult, bResultMTD);
 				}
 				{
 					const FRigidTransform3 QueryTM(FVec3(3.9, 0.0, 100.0), FQuat::Identity);
 					bool bResult = TriangleMesh->OverlapGeom(Capsule, QueryTM, 0.0);
 					EXPECT_TRUE(bResult);
+					bool bResultMTD = TriangleMesh->OverlapGeom(Capsule, QueryTM, 0.0, &MTDInfo);
+					EXPECT_EQ(bResult, bResultMTD);
 				}
 				{
 					const FRigidTransform3 QueryTM(FVec3(4.0, 0.0, 90.0), FQuat::Identity);
 					bool bResult = TriangleMesh->OverlapGeom(Capsule, QueryTM, 0.0);
 					EXPECT_TRUE(bResult);
+					bool bResultMTD = TriangleMesh->OverlapGeom(Capsule, QueryTM, 0.0, &MTDInfo);
+					EXPECT_EQ(bResult, bResultMTD);
 				}
 				{
 					const FRigidTransform3 QueryTM(FVec3(4.5, 0.0, 90.0), FQuat::Identity);
 					bool bResult = TriangleMesh->OverlapGeom(Capsule, QueryTM, 0.0);
 					EXPECT_FALSE(bResult);
+					bool bResultMTD = TriangleMesh->OverlapGeom(Capsule, QueryTM, 0.0, &MTDInfo);
+					EXPECT_EQ(bResult, bResultMTD);
 				}
 
 			}
diff --git a/Engine/Source/Runtime/Experimental/Chaos/Private/Chaos/Collision/TriangleOverlap.cpp b/Engine/Source/Runtime/Experimental/Chaos/Private/Chaos/Collision/TriangleOverlap.cpp
--- a/Engine/Source/Runtime/Experimental/Chaos/Private/Chaos/Collision/TriangleOverlap.cpp
+++ b/Engine/Source/Runtime/Experimental/Chaos/Private/Chaos/Collision/TriangleOverlap.cpp
@@ -399,7 +399,7 @@ namespace Chaos
 					Axis = VectorNegate(Axis);
 				}
 				const FRealSingle Separation = VectorDot3Scalar(VectorSubtract(ClosestPoint, PlaneVertex), Axis);
-				if (Separation > 0.0f)
+				if (Separation > UE_KINDA_SMALL_NUMBER)
 				{
 					return false;
 				}
@@ -510,7 +510,9 @@ namespace Chaos
 		FRealSingle ClosestDist = VectorDot3Scalar(Normal, ClosCent);
 		FRealSingle FurthestDist = VectorDot3Scalar(Normal, FurtCent);
 
-		return FMath::Sign(ClosestDist) != FMath::Sign(FurthestDist);
+		return !(FMath::Abs<FRealSingle>(ClosestDist) > UE_KINDA_SMALL_NUMBER &&
+			FMath::Abs<FRealSingle>(FurthestDist) > UE_KINDA_SMALL_NUMBER &&
+			FMath::Sign(ClosestDist) == FMath::Sign(FurthestDist));
 	}
 
 	const VectorRegister4Float FAABBSimd::SignBit = GlobalVectorConstants::SignBit();
@@ -564,7 +566,7 @@ namespace Chaos
 		HalfExtents = VectorLoadFloat3(&HalfExtentsf.X);
 	}
 
-	FORCEINLINE bool FAABBSimd::ComputeEdgeOverlap(const VectorRegister4Float& TriangleEdge, const VectorRegister4Float& TriangleVertex, const VectorRegister4Float& Centroid, const VectorRegister4Float& Normal, const VectorRegister4Float& LocalClosest) const
+	FORCEINLINE_DEBUGGABLE bool FAABBSimd::ComputeEdgeOverlap(const VectorRegister4Float& TriangleEdge, const VectorRegister4Float& TriangleVertex, const VectorRegister4Float& Centroid, const VectorRegister4Float& Normal, const VectorRegister4Float& LocalClosest) const
 	{
 		// Triangle edge vs box edges 
 		const VectorRegister4Float ClosestPoint = VectorAdd(LocalClosest, Position);
@@ -593,7 +595,7 @@ namespace Chaos
 					Axis = VectorBitwiseXor(SignBit, Axis);
 				}
 				const FRealSingle ScaledSeparation = VectorDot3Scalar(VectorSubtract(ClosestPoint, TriangleVertex), Axis);
-				if (ScaledSeparation > 0.0f)
+				if (ScaledSeparation > UE_KINDA_SMALL_NUMBER)
 				{
 					return false;
 				}
@@ -674,7 +676,9 @@ namespace Chaos
 		FRealSingle ClosestDist = VectorDot3Scalar(Normal, ClosCent);
 		FRealSingle FurthestDist = VectorDot3Scalar(Normal, FurtCent);
 
-		if (FMath::Sign(ClosestDist) == FMath::Sign(FurthestDist))
+		if (FMath::Abs<FRealSingle>(ClosestDist) > UE_KINDA_SMALL_NUMBER &&
+			FMath::Abs<FRealSingle>(FurthestDist) > UE_KINDA_SMALL_NUMBER &&
+			FMath::Sign(ClosestDist) == FMath::Sign(FurthestDist))
 		{
 			return false;
 		}
