Chaos - Island Manager removes statics and kinematics that have no constraints to prevent accumulation

diff --git a/Engine/Source/Programs/HeadlessChaos/Private/HeadlessChaosTestGraphEvolution.cpp b/Engine/Source/Programs/HeadlessChaos/Private/HeadlessChaosTestGraphEvolution.cpp
--- a/Engine/Source/Programs/HeadlessChaos/Private/HeadlessChaosTestGraphEvolution.cpp
+++ b/Engine/Source/Programs/HeadlessChaos/Private/HeadlessChaosTestGraphEvolution.cpp
@@ -903,31 +903,47 @@ namespace ChaosTest
 		EXPECT_TRUE(Test.ConstraintHandles[0]->IsSleeping());
 	}
 
-	// Test dynamic particle to be removed from island if no constraints
+	// Test isolated kinematic particles are removed from the island manager if they have no constraints
 	GTEST_TEST(GraphEvolutionTests, TestConstraintGraph_KinematicRemoveFromGraph)
 	{
+		// Create a scene with 3 dynamic particles
 		FGraphEvolutionTest Test(3);
 
 		Test.Advance();
+
+		// Initially all particles are dynamic and should be in the graph
 		EXPECT_EQ(Test.IslandManager->NumIslands(), 3);
 		EXPECT_EQ(Test.IslandManager->GetIslandGraph()->ItemNodes.Num(), 3);
+		EXPECT_TRUE(Test.ParticleHandles[0]->IsInConstraintGraph());	// Dynamic
+		EXPECT_TRUE(Test.ParticleHandles[1]->IsInConstraintGraph());	// Dynamic
+		EXPECT_TRUE(Test.ParticleHandles[2]->IsInConstraintGraph());	// Dynamic
+		EXPECT_TRUE(Test.IslandManager->GetIslandGraph()->GraphNodes.IsValidIndex(0));
 		EXPECT_TRUE(Test.IslandManager->GetIslandGraph()->GraphNodes.IsValidIndex(1));
+		EXPECT_TRUE(Test.IslandManager->GetIslandGraph()->GraphNodes.IsValidIndex(2));
 
+		// Change a particle to kinematic. Since it has no collisions or joints it should
+		// be removed from the constraint graph at the next update.
 		Test.Evolution.SetParticleObjectState(Test.ParticleHandles[1], EObjectStateType::Kinematic);
 
 		Test.Advance();
+
+		// Should now have only 2 nodes in the graph, each in their own island
 		EXPECT_EQ(Test.IslandManager->NumIslands(), 2);
-		EXPECT_EQ(Test.IslandManager->GetIslandGraph()->ItemNodes.Num(), 3);
-		// Kinematic is still in a graph, and will be updated
-		EXPECT_TRUE(Test.IslandManager->GetIslandGraph()->GraphNodes.IsValidIndex(1));
-		// The two dynamics body are also there
-		EXPECT_TRUE(Test.IslandManager->GetIslandGraph()->GraphNodes.IsValidIndex(0));
-		EXPECT_TRUE(Test.IslandManager->GetIslandGraph()->GraphNodes.IsValidIndex(2));
+		EXPECT_EQ(Test.IslandManager->GetIslandGraph()->ItemNodes.Num(), 2);
+
+		// Kinematic was removed from the graph, but dynamics are still there
+		EXPECT_TRUE(Test.ParticleHandles[0]->IsInConstraintGraph());	// Dynamic
+		EXPECT_FALSE(Test.ParticleHandles[1]->IsInConstraintGraph());	// Kinematic
+		EXPECT_TRUE(Test.ParticleHandles[2]->IsInConstraintGraph());	// Dynamic
 
 		Test.Advance();
-		EXPECT_EQ(Test.IslandManager->GetIslandGraph()->ItemNodes.Num(), 3);
-		EXPECT_TRUE(Test.IslandManager->GetIslandGraph()->GraphNodes.IsValidIndex(1));
 
+		// State should not have changed with a second tick
+		EXPECT_EQ(Test.IslandManager->NumIslands(), 2);
+		EXPECT_EQ(Test.IslandManager->GetIslandGraph()->ItemNodes.Num(), 2);
+		EXPECT_TRUE(Test.ParticleHandles[0]->IsInConstraintGraph());	// Dynamic
+		EXPECT_FALSE(Test.ParticleHandles[1]->IsInConstraintGraph());	// Kinematic
+		EXPECT_TRUE(Test.ParticleHandles[2]->IsInConstraintGraph());	// Dynamic
 	}
 
 	// Test the conditions for a kinematic particle waking an island
diff --git a/Engine/Source/Runtime/Experimental/Chaos/Private/Chaos/Island/IslandManager.cpp b/Engine/Source/Runtime/Experimental/Chaos/Private/Chaos/Island/IslandManager.cpp
--- a/Engine/Source/Runtime/Experimental/Chaos/Private/Chaos/Island/IslandManager.cpp
+++ b/Engine/Source/Runtime/Experimental/Chaos/Private/Chaos/Island/IslandManager.cpp
@@ -239,6 +239,15 @@ void FPBDIslandManager::InitializeGraph(const TParticleView<FPBDRigidParticles>&
 			FGeometryParticleHandle* ParticleHandle = IslandGraph->GraphNodes[NodeIndex].NodeItem;
 
 			IslandGraph->UpdateNode(NodeIndex, IsDynamicParticle(ParticleHandle), IsStationaryParticle(ParticleHandle));
+
+			// If the particle is static or kinematic and has no edges, remove it
+			if (!IslandGraph->GraphNodes[NodeIndex].bValidNode)
+			{
+				if (IslandGraph->GraphNodes[NodeIndex].NodeEdges.Num() == 0)
+				{
+					RemoveParticle(ParticleHandle);
+				}
+			}
 		}
 	}
 
