Fix compilation error in Headless

diff --git a/Engine/Source/Programs/HeadlessChaos/Private/GeometryCollection/GeometryCollectionTestEvents.cpp b/Engine/Source/Programs/HeadlessChaos/Private/GeometryCollection/GeometryCollectionTestEvents.cpp
--- a/Engine/Source/Programs/HeadlessChaos/Private/GeometryCollection/GeometryCollectionTestEvents.cpp
+++ b/Engine/Source/Programs/HeadlessChaos/Private/GeometryCollection/GeometryCollectionTestEvents.cpp
@@ -30,6 +30,12 @@ namespace GeometryCollectionTest
 		FVector Data2;
 	};
 
+	template<typename PayloadType>
+	bool IsEventDataEmpty(const PayloadType* Buffer)
+	{
+		return false;
+	}
+
 	class MyEventHandler
 	{
 	public:
