Chaos - fix HeadlessChaos compiler errors (gravity changes)

diff --git a/Engine/Source/Programs/HeadlessChaos/Private/GeometryCollection/GeometryCollectionTestCollisionResolution.cpp b/Engine/Source/Programs/HeadlessChaos/Private/GeometryCollection/GeometryCollectionTestCollisionResolution.cpp
--- a/Engine/Source/Programs/HeadlessChaos/Private/GeometryCollection/GeometryCollectionTestCollisionResolution.cpp
+++ b/Engine/Source/Programs/HeadlessChaos/Private/GeometryCollection/GeometryCollectionTestCollisionResolution.cpp
@@ -51,7 +51,7 @@ namespace GeometryCollectionTest
 			EXPECT_EQ(Collection->DynamicCollection->Simplicials[3], nullptr);
 			EXPECT_EQ(UnitTest.Solver->GetParticles().GetGeometryCollectionParticles().CollisionParticles(0), nullptr);
 
-			const FReal MaxRestingSeparation = -UnitTest.Solver->GetEvolution()->GetGravityForces().GetAcceleration().Z * UnitTest.Dt * UnitTest.Dt;	// PBD resting separation will be up to this
+			const FReal MaxRestingSeparation = -UnitTest.Solver->GetEvolution()->GetGravityForces().GetAcceleration(0).Z * UnitTest.Dt * UnitTest.Dt;	// PBD resting separation will be up to this
 			EXPECT_LT(FMath::Abs(Collection->RestCollection->Transform[0].GetTranslation().Z)-10.f, KINDA_SMALL_NUMBER);
 			EXPECT_LT(FMath::Abs(Collection->DynamicCollection->Transform[0].GetTranslation().Z - 1.0), 0.1 + MaxRestingSeparation);
 		}
diff --git a/Engine/Source/Programs/HeadlessChaos/Private/HeadlessChaosTestConstraints.cpp b/Engine/Source/Programs/HeadlessChaos/Private/HeadlessChaosTestConstraints.cpp
--- a/Engine/Source/Programs/HeadlessChaos/Private/HeadlessChaosTestConstraints.cpp
+++ b/Engine/Source/Programs/HeadlessChaos/Private/HeadlessChaosTestConstraints.cpp
@@ -155,7 +155,7 @@ namespace ChaosTest {
 			InitEvolutionSettings(Evolution);
 
 			// disable gravity
-			Evolution.GetGravityForces().SetAcceleration(FVec3(0, 0, 0));
+			Evolution.GetGravityForces().SetAcceleration(FVec3(0, 0, 0), 0);
 
 			// chassis particle
 			auto* DynamicParticle = Evolution.CreateDynamicParticles(1)[0];
@@ -201,7 +201,7 @@ namespace ChaosTest {
 			InitEvolutionSettings(Evolution);
 
 			// disable gravity
-			Evolution.GetGravityForces().SetAcceleration(FVec3(0, 0, 0));
+			Evolution.GetGravityForces().SetAcceleration(FVec3(0, 0, 0), 0);
 
 			// chassis particle
 			auto* DynamicParticle = Evolution.CreateDynamicParticles(1)[0];
@@ -261,7 +261,7 @@ namespace ChaosTest {
 		Evolution.SetNumVelocityIterations(1);
 
 		// disable gravity
-		Evolution.GetGravityForces().SetAcceleration(FVec3(0, 0, -980.f));
+		Evolution.GetGravityForces().SetAcceleration(FVec3(0, 0, -980.f), 0);
 
 		float Mass = 1.0f;
 
diff --git a/Engine/Source/Programs/HeadlessChaos/Private/HeadlessChaosTestGraphEvolution.cpp b/Engine/Source/Programs/HeadlessChaos/Private/HeadlessChaosTestGraphEvolution.cpp
--- a/Engine/Source/Programs/HeadlessChaos/Private/HeadlessChaosTestGraphEvolution.cpp
+++ b/Engine/Source/Programs/HeadlessChaos/Private/HeadlessChaosTestGraphEvolution.cpp
@@ -30,7 +30,7 @@ namespace ChaosTest
 			// Bind the constraints to the evolution
 			Evolution.AddConstraintContainer(Constraints);
 
-			Evolution.GetGravityForces().SetAcceleration(FVec3(0));
+			Evolution.GetGravityForces().SetAcceleration(FVec3(0), 0);
 
 			// Create some particles and constraints connecting them in a chain: 0-1-2-3-...-N
 			ParticleHandles = Evolution.CreateDynamicParticles(NumParticles);
diff --git a/Engine/Source/Programs/HeadlessChaos/Private/HeadlessChaosTestRewind.cpp b/Engine/Source/Programs/HeadlessChaos/Private/HeadlessChaosTestRewind.cpp
--- a/Engine/Source/Programs/HeadlessChaos/Private/HeadlessChaosTestRewind.cpp
+++ b/Engine/Source/Programs/HeadlessChaos/Private/HeadlessChaosTestRewind.cpp
@@ -501,7 +501,7 @@ namespace ChaosTest {
 
 				auto& Particle = Proxy->GetGameThreadAPI();
 				Particle.SetGravityEnabled(true);
-				Solver->GetEvolution()->GetGravityForces().SetAcceleration(FVec3(0, 0, -1));
+				Solver->GetEvolution()->GetGravityForces().SetAcceleration(FVec3(0, 0, -1), 0);
 
 				Particle.SetV(FVec3(0, 0, 0));
 
@@ -1458,7 +1458,7 @@ namespace ChaosTest {
 					{
 						if (PhysicsStep == 3)
 						{
-							Proxy->GetPhysicsThreadAPI()->SetAcceleration(FVec3(0, 0, 10) + Proxy->GetPhysicsThreadAPI()->Acceleration());
+							Proxy->GetPhysicsThreadAPI()->SetAcceleration(FVec3(0, 0, 10) + Proxy->GetPhysicsThreadAPI()->Acceleration(), 0);
 							Proxy->GetPhysicsThreadAPI()->SetAngularAcceleration(FVec3(0, 0, 10) + Proxy->GetPhysicsThreadAPI()->AngularAcceleration());
 						}
 
@@ -1684,7 +1684,7 @@ namespace ChaosTest {
 		TRewindHelper::TestDynamicSphere([](auto* Solver, FReal SimDt, int32 Optimization, auto Proxy, auto Sphere)
 		{
 				auto& Particle = Proxy->GetGameThreadAPI();
-				Solver->GetEvolution()->GetGravityForces().SetAcceleration(FVec3(0, 0, -1));
+				Solver->GetEvolution()->GetGravityForces().SetAcceleration(FVec3(0, 0, -1), 0);
 				Particle.SetGravityEnabled(true);
 				Particle.SetX(FVec3(0, 0, 100));
 
@@ -1829,7 +1829,7 @@ namespace ChaosTest {
 
 				{
 					auto& Particle = Proxy->GetGameThreadAPI();
-					Solver->GetEvolution()->GetGravityForces().SetAcceleration(FVec3(0, 0, -1));
+					Solver->GetEvolution()->GetGravityForces().SetAcceleration(FVec3(0, 0, -1), 0);
 					Particle.SetGravityEnabled(true);
 					Particle.SetX(FVec3(0, 0, 100));
 
@@ -1989,7 +1989,7 @@ namespace ChaosTest {
 
 				{
 					auto& Particle = Proxy->GetGameThreadAPI();
-					Solver->GetEvolution()->GetGravityForces().SetAcceleration(FVec3(0, 0, -1));
+					Solver->GetEvolution()->GetGravityForces().SetAcceleration(FVec3(0, 0, -1), 0);
 					Particle.SetGravityEnabled(true);
 					Particle.SetX(FVec3(0, 0, 100));
 
@@ -2113,7 +2113,7 @@ namespace ChaosTest {
 
 				{
 					auto& Particle = Proxy->GetGameThreadAPI();
-					Solver->GetEvolution()->GetGravityForces().SetAcceleration(FVec3(0, 0, -1));
+					Solver->GetEvolution()->GetGravityForces().SetAcceleration(FVec3(0, 0, -1), 0);
 					Particle.SetGravityEnabled(false);
 					Particle.SetX(FVec3(0, 0, 0));
 					Particle.SetV(FVec3(0, 0, -1));
@@ -2217,7 +2217,7 @@ namespace ChaosTest {
 
 			{
 				auto& Particle = Proxy->GetGameThreadAPI();
-				Solver->GetEvolution()->GetGravityForces().SetAcceleration(FVec3(0, 0, -1));
+				Solver->GetEvolution()->GetGravityForces().SetAcceleration(FVec3(0, 0, -1), 0);
 				Particle.SetGravityEnabled(true);
 				Particle.SetX(FVec3(0, 0, 100));
 
@@ -2253,7 +2253,7 @@ namespace ChaosTest {
 
 			{
 				auto& Particle = Proxy->GetGameThreadAPI();
-				Solver->GetEvolution()->GetGravityForces().SetAcceleration(FVec3(0, 0, -1));
+				Solver->GetEvolution()->GetGravityForces().SetAcceleration(FVec3(0, 0, -1), 0);
 				Particle.SetGravityEnabled(true);
 				Particle.SetX(FVec3(0, 0, 100));
 				Particle.SetResimType(EResimType::ResimAsFollower);
@@ -3488,7 +3488,7 @@ namespace ChaosTest {
 				Dynamic.SetGeometry(Sphere);
 				Dynamic.SetGravityEnabled(true);
 				Solver->RegisterObject(DynamicProxy);
-				Solver->GetEvolution()->GetGravityForces().SetAcceleration(FVec3(0, 0, -1));
+				Solver->GetEvolution()->GetGravityForces().SetAcceleration(FVec3(0, 0, -1), 0);
 
 				Kinematic.SetGeometry(Box);
 				Solver->RegisterObject(KinematicProxy);
@@ -3596,7 +3596,7 @@ namespace ChaosTest {
 			Dynamic.SetGeometry(Sphere);
 			Dynamic.SetGravityEnabled(true);
 			Solver->RegisterObject(DynamicProxy);
-			Solver->GetEvolution()->GetGravityForces().SetAcceleration(FVec3(0, 0, -1));
+			Solver->GetEvolution()->GetGravityForces().SetAcceleration(FVec3(0, 0, -1), 0);
 
 			Kinematic.SetGeometry(Box);
 			Solver->RegisterObject(KinematicProxy);
@@ -3710,7 +3710,7 @@ namespace ChaosTest {
 			Dynamic.SetGeometry(Sphere);
 			Dynamic.SetGravityEnabled(true);
 			Solver->RegisterObject(DynamicProxy);
-			Solver->GetEvolution()->GetGravityForces().SetAcceleration(FVec3(0, 0, -1));
+			Solver->GetEvolution()->GetGravityForces().SetAcceleration(FVec3(0, 0, -1), 0);
 
 			Kinematic.SetGeometry(Box);
 			Solver->RegisterObject(KinematicProxy);
@@ -3964,7 +3964,7 @@ namespace ChaosTest {
 			Dynamic.SetGeometry(Box);
 			Dynamic.SetGravityEnabled(true);
 			Solver->RegisterObject(DynamicProxy);
-			Solver->GetEvolution()->GetGravityForces().SetAcceleration(FVec3(0, 0, -1));
+			Solver->GetEvolution()->GetGravityForces().SetAcceleration(FVec3(0, 0, -1), 0);
 			Dynamic.SetObjectState(EObjectStateType::Dynamic);
 
 			Storage.Add(DynamicProxy);
@@ -3999,7 +3999,7 @@ namespace ChaosTest {
 			Dynamic.SetGeometry(Box);
 			Dynamic.SetGravityEnabled(true);
 			Solver->RegisterObject(DynamicProxy);
-			Solver->GetEvolution()->GetGravityForces().SetAcceleration(FVec3(0, 0, -1));
+			Solver->GetEvolution()->GetGravityForces().SetAcceleration(FVec3(0, 0, -1), 0);
 			Dynamic.SetObjectState(EObjectStateType::Dynamic);
 			Dynamic.SetX(StartPos);
 			Dynamic.SetR(StartRotation);
@@ -4084,7 +4084,7 @@ namespace ChaosTest {
 			Dynamic.SetGeometry(Box);
 			Dynamic.SetGravityEnabled(true);
 			Solver->RegisterObject(DynamicProxy);
-			Solver->GetEvolution()->GetGravityForces().SetAcceleration(FVec3(0, 0, Gravity));
+			Solver->GetEvolution()->GetGravityForces().SetAcceleration(FVec3(0, 0, Gravity), 0);
 			Dynamic.SetObjectState(EObjectStateType::Dynamic);
 
 			Storage.Add(DynamicProxy);
@@ -4178,7 +4178,7 @@ namespace ChaosTest {
 
 		const auto InitLambda = [&SmallBox, &Box](auto& Solver, auto)
 		{
-			Solver->GetEvolution()->GetGravityForces().SetAcceleration(FVec3(0, 0, -980));
+			Solver->GetEvolution()->GetGravityForces().SetAcceleration(FVec3(0, 0, -980), 0);
 			TArray<FPhysicsActorHandle> Storage;
 			for (int Idx = 0; Idx < 5; ++Idx)
 			{
@@ -4245,7 +4245,7 @@ namespace ChaosTest {
 		const auto InitLambda = [&SmallBox, &Box, NumParticles](auto& Solver, const TArray<int32>* Mapping)
 		{
 			auto MappingHelper = [Mapping](const int32 Idx) { return Mapping ? (*Mapping)[Idx] : Idx; };
-			Solver->GetEvolution()->GetGravityForces().SetAcceleration(FVec3(0, 0, -980));
+			Solver->GetEvolution()->GetGravityForces().SetAcceleration(FVec3(0, 0, -980), 0);
 			TArray<FPhysicsActorHandle> Storage;
 			for (int Idx = 0; Idx < NumParticles; ++Idx)
 			{
diff --git a/Engine/Source/Programs/HeadlessChaos/Private/HeadlessChaosTestVehicles.cpp b/Engine/Source/Programs/HeadlessChaos/Private/HeadlessChaosTestVehicles.cpp
--- a/Engine/Source/Programs/HeadlessChaos/Private/HeadlessChaosTestVehicles.cpp
+++ b/Engine/Source/Programs/HeadlessChaos/Private/HeadlessChaosTestVehicles.cpp
@@ -1457,7 +1457,7 @@ namespace ChaosTest
 		FPBDRigidsEvolutionGBF Evolution(Particles, PhysicalMaterials);
 
 		float BodyMass = 1000.0f;
-		float Gravity = FMath::Abs(Evolution.GetGravityForces().GetAcceleration().Z);
+		float Gravity = FMath::Abs(Evolution.GetGravityForces().GetAcceleration(0).Z);
 
 		FSimpleSuspensionConfig Setup;
 		Setup.MaxLength = 20.0f;
diff --git a/Engine/Source/Programs/HeadlessChaos/Public/HeadlessChaosTestConstraints.h b/Engine/Source/Programs/HeadlessChaos/Public/HeadlessChaosTestConstraints.h
--- a/Engine/Source/Programs/HeadlessChaos/Public/HeadlessChaosTestConstraints.h
+++ b/Engine/Source/Programs/HeadlessChaos/Public/HeadlessChaosTestConstraints.h
@@ -36,7 +36,7 @@ namespace ChaosTest
 			Evolution.SetNumPositionIterations(NumIterations);
 			Evolution.SetNumVelocityIterations(1);
 			Evolution.SetNumProjectionIterations(1);
-			Evolution.GetGravityForces().SetAcceleration(Gravity * FVec3(0, 0, -1));
+			Evolution.GetGravityForces().SetAcceleration(Gravity * FVec3(0, 0, -1), 0);
 		}
 
 		virtual ~FConstraintsTest()
