Fix Tetrahedron::Inside(), and add Outside() tests to HeadlessChaosTestImplicits.

diff --git a/Engine/Source/Programs/HeadlessChaos/Private/HeadlessChaosTestImplicits.cpp b/Engine/Source/Programs/HeadlessChaos/Private/HeadlessChaosTestImplicits.cpp
--- a/Engine/Source/Programs/HeadlessChaos/Private/HeadlessChaosTestImplicits.cpp
+++ b/Engine/Source/Programs/HeadlessChaos/Private/HeadlessChaosTestImplicits.cpp
@@ -248,8 +248,10 @@ namespace ChaosTest {
 
 		// Center
 		TVec3<double> Pt = Center;
-		bool Hit = Tet.Inside(Pt); // generates tris
+		bool Hit = !Tet.Outside(Pt); // generates tris
 		EXPECT_TRUE(Hit);
+		bool Inside = Tet.Inside(Pt);
+		EXPECT_TRUE(Hit == Inside);
 		bool RobustHit = Tet.RobustInside(Pt);
 		EXPECT_TRUE(RobustHit);
 		TVec3<double> Bary = Tet.GetFirstThreeBarycentricCoordinates(Pt);
@@ -263,8 +265,10 @@ namespace ChaosTest {
 
 		// Point
 		Pt = Tet[0];
-		Hit = Tet.Inside(Pt, 0.001);
+		Hit = !Tet.Outside(Pt, 0.001);
 		EXPECT_TRUE(Hit);
+		Inside = Tet.Inside(Pt);
+		EXPECT_TRUE(Hit == Inside);
 		RobustHit = Tet.RobustInside(Pt, -0.001);
 		EXPECT_TRUE(RobustHit);
 		Bary = Tet.GetFirstThreeBarycentricCoordinates(Pt);
@@ -277,8 +281,10 @@ namespace ChaosTest {
 
 
 		Pt[0] -= 0.1;
-		Hit = Tet.Inside(Pt, 0.001);
+		Hit = !Tet.Outside(Pt, 0.001);
 		EXPECT_FALSE(Hit);
+		Inside = Tet.Inside(Pt);
+		EXPECT_TRUE(Hit == Inside);
 		RobustHit = Tet.RobustInside(Pt, -0.001);
 		EXPECT_FALSE(RobustHit);
 		Surf = Tet.ProjectToSurface(Tris, Pt);
@@ -298,8 +304,10 @@ namespace ChaosTest {
 
 		// Edge
 		Pt = TVec3<double>(0.5, 0, 0);
-		Hit = Tet.Inside(Pt, 0.001);
+		Hit = !Tet.Outside(Pt, 0.001);
 		EXPECT_TRUE(Hit);
+		Inside = Tet.Inside(Pt);
+		EXPECT_TRUE(Hit == Inside);
 		RobustHit = Tet.RobustInside(Pt, -0.001);
 		EXPECT_TRUE(RobustHit);
 		Bary = Tet.GetFirstThreeBarycentricCoordinates(Pt);
@@ -311,8 +319,10 @@ namespace ChaosTest {
 		EXPECT_VECTOR_NEAR(ClosestBary, TVec4<double>(.5, .5, 0, 0), 0.001);
 
 		Pt[1] -= 0.1;
-		Hit = Tet.Inside(Pt, 0.001);
+		Hit = !Tet.Outside(Pt, 0.001);
 		EXPECT_FALSE(Hit);
+		Inside = Tet.Inside(Pt);
+		EXPECT_TRUE(Hit == Inside);
 		RobustHit = Tet.RobustInside(Pt, -0.001);
 		EXPECT_FALSE(RobustHit);
 		Surf = Tet.ProjectToSurface(Tris, Pt);
@@ -334,8 +344,10 @@ namespace ChaosTest {
 
 		// Face
 		Pt = Tris[0].GetCentroid();
-		Hit = Tet.Inside(Pt, 0.001);
+		Hit = !Tet.Outside(Pt, 0.001);
 		EXPECT_TRUE(Hit);
+		Inside = Tet.Inside(Pt);
+		EXPECT_TRUE(Hit == Inside);
 		RobustHit = Tet.RobustInside(Pt, -0.001);
 		EXPECT_TRUE(RobustHit);
 		Bary = Tet.GetFirstThreeBarycentricCoordinates(Pt);
@@ -347,8 +359,10 @@ namespace ChaosTest {
 		EXPECT_VECTOR_NEAR(ClosestBary, TVec4<double>(.3333, .3333, .3333, 0), 0.001);
 
 		Pt += Tris[0].GetNormal() * 0.1;
-		Hit = Tet.Inside(Pt, 0.001);
+		Hit = !Tet.Outside(Pt, 0.001);
 		EXPECT_FALSE(Hit);
+		Inside = Tet.Inside(Pt);
+		EXPECT_TRUE(Hit == Inside);
 		RobustHit = Tet.RobustInside(Pt, -0.001);
 		EXPECT_FALSE(RobustHit);
 		Surf = Tet.ProjectToSurface(Tris, Pt);
diff --git a/Engine/Source/Runtime/Experimental/Chaos/Public/Chaos/Tetrahedron.h b/Engine/Source/Runtime/Experimental/Chaos/Public/Chaos/Tetrahedron.h
--- a/Engine/Source/Runtime/Experimental/Chaos/Public/Chaos/Tetrahedron.h
+++ b/Engine/Source/Runtime/Experimental/Chaos/Public/Chaos/Tetrahedron.h
@@ -175,7 +175,7 @@ namespace Chaos {
 			checkSlow(Tris.Num() == 4);
 
 			for(int32 i=0; i < Tris.Num(); i++)
-				if (Inside(Tris[i].GetPlane(), Location, -HalfThickness))
+				if (!Inside(Tris[i].GetPlane(), Location, -HalfThickness))
 				{
 					return false;
 				}
