change headless chaos GeometryCollection_RigidBodies_ClusterTest_SingleLevelBreaking test to not use a cluster union due to a previous change that creates cluster unions even for clusters of size 1

diff --git a/Engine/Source/Programs/HeadlessChaos/Private/GeometryCollection/GeometryCollectionTestClustering.cpp b/Engine/Source/Programs/HeadlessChaos/Private/GeometryCollection/GeometryCollectionTestClustering.cpp
--- a/Engine/Source/Programs/HeadlessChaos/Private/GeometryCollection/GeometryCollectionTestClustering.cpp
+++ b/Engine/Source/Programs/HeadlessChaos/Private/GeometryCollection/GeometryCollectionTestClustering.cpp
@@ -397,7 +397,7 @@ namespace GeometryCollectionTest
 		Params.Simulating = true;
 		Params.EnableClustering = true;
 		Params.DamageThreshold = { 0.1f };
-		Params.ClusterGroupIndex = 1;
+		Params.ClusterGroupIndex = 0;
 		FGeometryCollectionWrapper* Collection = TNewSimulationObject<GeometryType::GeometryCollectionWithSuppliedRestCollection>::Init(Params)->template As<FGeometryCollectionWrapper>();
 
 
