Chaos Cloth - Improve Headless Chaos managed array collection functional test, to not fail with consecutive 0 elements and remove randomness.

#rnx

diff --git a/Engine/Source/Programs/HeadlessChaos/Private/HeadlessChaosTestCloth.cpp b/Engine/Source/Programs/HeadlessChaos/Private/HeadlessChaosTestCloth.cpp
--- a/Engine/Source/Programs/HeadlessChaos/Private/HeadlessChaosTestCloth.cpp
+++ b/Engine/Source/Programs/HeadlessChaos/Private/HeadlessChaosTestCloth.cpp
@@ -352,7 +352,17 @@ namespace ChaosTest {
 
 				if (NumValues)
 				{
-					const int32 ValuePosition = (ViewPosition > 0) ? ValueEnd[ViewPosition - 1] + 1: 0;
+					int32 PrevValuePosition = INDEX_NONE;
+					for (int32 Index = ViewPosition - 1; Index >= 0; --Index)
+					{
+						if (ValueEnd[Index] != INDEX_NONE)
+						{
+							PrevValuePosition = ValueEnd[Index];
+							break;
+						}
+					}
+					const int32 ValuePosition = PrevValuePosition + 1;
+
 					EXPECT_TRUE(ManagedArrayCollection.InsertElements(NumValues, ValuePosition, DataGroup) == ValuePosition);
 			
 					ValueStart[ViewPosition] = ValuePosition;
@@ -503,18 +513,23 @@ namespace ChaosTest {
 		// Insert 5 views from the start of the array
 		for (int32 Index = 0; Index < 5; ++Index)
 		{
-			InsertView(0, FMath::Rand() / 32);
+			InsertView(0, (Index + 1) * 100);
+		}
+		// Insert 2 more views for testing consecutive empty views
+		for (int32 Index = 0; Index < 2; ++Index)
+		{
+			InsertView(0, 0);
 		}
 		EXPECT_TRUE(HasKeptIntegrity());
 		EXPECT_TRUE(HasKeptOrder());
 
 		// Add 1 view in the middle of the array
-		InsertView(ManagedArrayCollection.NumElements(ViewGroup) / 2, FMath::Rand() / 32);
+		InsertView(ManagedArrayCollection.NumElements(ViewGroup) / 2, 500);
 		EXPECT_TRUE(HasKeptIntegrity());
 		EXPECT_TRUE(HasKeptOrder());
 
 		// Add 1 view at the end of the array
-		InsertView(ManagedArrayCollection.NumElements(ViewGroup), FMath::Rand() / 32);
+		InsertView(ManagedArrayCollection.NumElements(ViewGroup), 1000);
 		EXPECT_TRUE(HasKeptIntegrity());
 		EXPECT_TRUE(HasKeptOrder());
 
@@ -533,8 +548,8 @@ namespace ChaosTest {
 		EXPECT_TRUE(HasKeptIntegrity());
 		EXPECT_TRUE(HasKeptOrder());
 
-		// Remove all remaining 3 views
-		RemoveViews(0, 3);
+		// Remove all remaining 5 views
+		RemoveViews(0, 5);
 		EXPECT_TRUE(ManagedArrayCollection.NumElements(ViewGroup) == 0);
 		EXPECT_TRUE(ManagedArrayCollection.NumElements(DataGroup) == 0);
 
