Chaos - fix EPA and GJK for some degenerate cases
- Fix Simd GJK handling of EPA error codes
- Make EPA tolerances be relative rather than absolute
- Fix EPA tolerance use of Abs()

diff --git a/Engine/Source/Programs/HeadlessChaos/Private/HeadlessChaosTestEPA.cpp b/Engine/Source/Programs/HeadlessChaos/Private/HeadlessChaosTestEPA.cpp
--- a/Engine/Source/Programs/HeadlessChaos/Private/HeadlessChaosTestEPA.cpp
+++ b/Engine/Source/Programs/HeadlessChaos/Private/HeadlessChaosTestEPA.cpp
@@ -1197,128 +1197,170 @@ namespace ChaosTest
 
 	}
 
+	//
+	// Performs the same initially overlapping sweep twice, with slightly different rotations, gives different normals.
+	// this is convex box slightly penetrating surface of triangle mesh. Trimesh normals point up.
+	//
+	GTEST_TEST(EPATests, EPARealFailures_ConvexTrimeshRotationalDifferencesBreakNormal)
+	{
+		using namespace Chaos;
+		TArray<FConvex::FVec3Type> ConvexBoxSurfaceParticles(
+			{
+			{50.0999985, -50.1124992, -50.1250000},
+			{-50.0999985, 50.1250000, -50.1250000},
+			{50.0999985, 50.1250000, -50.0999985},
+			{50.0999985, 50.1250000, 50.1250000},
+			{-50.0999985, -50.1124992, -50.0999985},
+			{-50.0999985, -50.1124992, 50.1250000},
+			{50.0999985, -50.1124992, 50.1250000},
+			{-50.0999985, 50.1250000, 50.1250000},
+			});
 
-		//
-		// Performs the same initially overlapping sweep twice, with slightly different rotations, gives different normals.
-		// this is convex box slightly penetrating surface of triangle mesh. Trimesh normals point up.
-		//
-		GTEST_TEST(EPATests, EPARealFailures_ConvexTrimeshRotationalDifferencesBreakNormal)
-		{
-			using namespace Chaos;
-			TArray<FConvex::FVec3Type> ConvexBoxSurfaceParticles(
-				{
-				{50.0999985, -50.1124992, -50.1250000},
-				{-50.0999985, 50.1250000, -50.1250000},
-				{50.0999985, 50.1250000, -50.0999985},
-				{50.0999985, 50.1250000, 50.1250000},
-				{-50.0999985, -50.1124992, -50.0999985},
-				{-50.0999985, -50.1124992, 50.1250000},
-				{50.0999985, -50.1124992, 50.1250000},
-				{-50.0999985, 50.1250000, 50.1250000},
-				});
-
-			FConvex ConvexBox(MoveTemp(ConvexBoxSurfaceParticles), 0.0f);
+		FConvex ConvexBox(MoveTemp(ConvexBoxSurfaceParticles), 0.0f);
 
-			FTriangleMeshImplicitObject::ParticlesType TrimeshParticles(
-				{
-					{50.0000000, 50.0000000, -8.04061356e-15},
-					{50.0000000, -50.0000000, 8.04061356e-15},
-					{-50.0000000, 50.0000000, -8.04061356e-15},
-					{-50.0000000, -50.0000000, 8.04061356e-15}
-				});
+		FTriangleMeshImplicitObject::ParticlesType TrimeshParticles(
+			{
+				{50.0000000, 50.0000000, -8.04061356e-15},
+				{50.0000000, -50.0000000, 8.04061356e-15},
+				{-50.0000000, 50.0000000, -8.04061356e-15},
+				{-50.0000000, -50.0000000, 8.04061356e-15}
+			});
 
-			TArray<TVec3<int32>> Indices;
-			Indices.Emplace(1, 0, 2);
-			Indices.Emplace(1, 2, 3);
+		TArray<TVec3<int32>> Indices;
+		Indices.Emplace(1, 0, 2);
+		Indices.Emplace(1, 2, 3);
 
-			TArray<uint16> Materials;
-			Materials.Emplace(0);
-			Materials.Emplace(0);
-			TUniquePtr<FTriangleMeshImplicitObject> TriangleMesh = MakeUnique<FTriangleMeshImplicitObject>(MoveTemp(TrimeshParticles), MoveTemp(Indices), MoveTemp(Materials));
-			TImplicitObjectScaled<FTriangleMeshImplicitObject> ScaledTriangleMesh = TImplicitObjectScaled<FTriangleMeshImplicitObject>(MakeSerializable(TriangleMesh), nullptr, FVec3(11.5, 11.5, 11.5));
+		TArray<uint16> Materials;
+		Materials.Emplace(0);
+		Materials.Emplace(0);
+		TUniquePtr<FTriangleMeshImplicitObject> TriangleMesh = MakeUnique<FTriangleMeshImplicitObject>(MoveTemp(TrimeshParticles), MoveTemp(Indices), MoveTemp(Materials));
+		TImplicitObjectScaled<FTriangleMeshImplicitObject> ScaledTriangleMesh = TImplicitObjectScaled<FTriangleMeshImplicitObject>(MakeSerializable(TriangleMesh), nullptr, FVec3(11.5, 11.5, 11.5));
 
-			FQuat Rotation0(0.00488796039, 0.00569311855, -0.000786740216, 0.999971569);
-			FQuat Rotation1(0.0117356628, -0.0108017093, -0.000888462295, 0.999872327);
+		FQuat Rotation0(0.00488796039, 0.00569311855, -0.000786740216, 0.999971569);
+		FQuat Rotation1(0.0117356628, -0.0108017093, -0.000888462295, 0.999872327);
 
-			FVec3 Translation(309.365723, -69.4132690, 51.2289352);
+		FVec3 Translation(309.365723, -69.4132690, 51.2289352);
 
-			TRigidTransform<FReal, 3> Transform0(Translation, Rotation0);
-			TRigidTransform<FReal, 3> Transform1(Translation, Rotation1);
+		TRigidTransform<FReal, 3> Transform0(Translation, Rotation0);
+		TRigidTransform<FReal, 3> Transform1(Translation, Rotation1);
 
-			FVec3 Dir(-0.00339674903, 5.76980747e-05, -0.999994159);
-			FReal Length = 1.83530724;
+		FVec3 Dir(-0.00339674903, 5.76980747e-05, -0.999994159);
+		FReal Length = 1.83530724;
 
-			FReal OutTime = -1;
-			FVec3 Normal(0.0f);
-			FVec3 Position(0.0f);
-			int32 FaceIndex = -1;
-			FVec3 FaceNormal(0.0);
-			bool bResult = ScaledTriangleMesh.LowLevelSweepGeom(ConvexBox, Transform0, Dir, Length, OutTime, Position, Normal, FaceIndex, FaceNormal, 0.0f, true);
+		FReal OutTime = -1;
+		FVec3 Normal(0.0f);
+		FVec3 Position(0.0f);
+		int32 FaceIndex = -1;
+		FVec3 FaceNormal(0.0);
+		bool bResult = ScaledTriangleMesh.LowLevelSweepGeom(ConvexBox, Transform0, Dir, Length, OutTime, Position, Normal, FaceIndex, FaceNormal, 0.0f, true);
 
-			bResult = ScaledTriangleMesh.LowLevelSweepGeom(ConvexBox, Transform1, Dir, Length, OutTime, Position, Normal, FaceIndex, FaceNormal, 0.0f, true);
+		bResult = ScaledTriangleMesh.LowLevelSweepGeom(ConvexBox, Transform1, Dir, Length, OutTime, Position, Normal, FaceIndex, FaceNormal, 0.0f, true);
 			
-			// Observe that normals are in opposite direction, while rotations are very similar.
+		// Observe that normals are in opposite direction, while rotations are very similar.
 
-		}
+	}
 
-		//
-		// Player can clip through RockWall trimesh, this repros a failure, MTD seems wrong.
-		// This is failing GJKRaycast2 call.
-		// Fixed: (11457046) ClosestB computation was transformed wrong messing up normal.
-		//
-		GTEST_TEST(EPATests, EPARealFailures_CapsuleVsTrimeshRockWallWrongNormalGJKRaycast2)
-		{
-			using namespace Chaos;
-			// Triangle w/ world scale
-			FTriangle Triangle({
-				{-306.119476, 1674.38647, 117.138489},
-				{-491.015747, 1526.35803, 116.067123},
-				{-91.0660172, 839.028320, 118.413063}
-				});
+	//
+	// Player can clip through RockWall trimesh, this repros a failure, MTD seems wrong.
+	// This is failing GJKRaycast2 call.
+	// Fixed: (11457046) ClosestB computation was transformed wrong messing up normal.
+	//
+	GTEST_TEST(EPATests, EPARealFailures_CapsuleVsTrimeshRockWallWrongNormalGJKRaycast2)
+	{
+		using namespace Chaos;
+		// Triangle w/ world scale
+		FTriangle Triangle({
+			{-306.119476, 1674.38647, 117.138489},
+			{-491.015747, 1526.35803, 116.067123},
+			{-91.0660172, 839.028320, 118.413063}
+			});
 
-			FTriangleRegister TriangleReg({
-				MakeVectorRegisterFloat(-306.119476f, 1674.38647f, 117.138489f, 0.0f),
-				MakeVectorRegisterFloat(-491.015747f, 1526.35803f, 116.067123f, 0.0f),
-				MakeVectorRegisterFloat(-91.0660172f, 839.028320f, 118.413063f, 0.0f)
-				});
+		FTriangleRegister TriangleReg({
+			MakeVectorRegisterFloat(-306.119476f, 1674.38647f, 117.138489f, 0.0f),
+			MakeVectorRegisterFloat(-491.015747f, 1526.35803f, 116.067123f, 0.0f),
+			MakeVectorRegisterFloat(-91.0660172f, 839.028320f, 118.413063f, 0.0f)
+			});
 
 
-			FVec3 ExpectedNormal = FVec3::CrossProduct(Triangle[1] - Triangle[0], Triangle[2] - Triangle[0]);
-			ExpectedNormal.Normalize();
+		FVec3 ExpectedNormal = FVec3::CrossProduct(Triangle[1] - Triangle[0], Triangle[2] - Triangle[0]);
+		ExpectedNormal.Normalize();
 
-			TRigidTransform<FReal, 3> StartTM(FVec3(-344.031799, 1210.37158, 134.252747), FQuat(-0.255716801, -0.714108050, 0.0788889676, -0.646866322), FVec3(1));
+		TRigidTransform<FReal, 3> StartTM(FVec3(-344.031799, 1210.37158, 134.252747), FQuat(-0.255716801, -0.714108050, 0.0788889676, -0.646866322), FVec3(1));
 
-			// Wrapping in 1,1,1 scale is unnecessary, but this is technically what is happening when sweeping against scaled trimesh.
-			TUniquePtr<FCapsule> Capsule = MakeUnique<FCapsule>(FVec3(0, 0, -33), FVec3(0, 0, 33), 42);
-			TImplicitObjectScaled<FCapsule> ScaledCapsule = TImplicitObjectScaled<FCapsule>(MakeSerializable(Capsule), nullptr, FVec3(1));
+		// Wrapping in 1,1,1 scale is unnecessary, but this is technically what is happening when sweeping against scaled trimesh.
+		TUniquePtr<FCapsule> Capsule = MakeUnique<FCapsule>(FVec3(0, 0, -33), FVec3(0, 0, 33), 42);
+		TImplicitObjectScaled<FCapsule> ScaledCapsule = TImplicitObjectScaled<FCapsule>(MakeSerializable(Capsule), nullptr, FVec3(1));
 
 
-			const FVec3 Dir(-0.102473199, 0.130887285, -0.986087084);
-			const FReal LengthScale = 9.31486130;
-			const FReal  CurrentLength = 2.14465737;
-			const FReal Length = LengthScale * CurrentLength;
-			const bool bComputeMTD = true;
-			const FReal Thickness = 0;
+		const FVec3 Dir(-0.102473199, 0.130887285, -0.986087084);
+		const FReal LengthScale = 9.31486130;
+		const FReal  CurrentLength = 2.14465737;
+		const FReal Length = LengthScale * CurrentLength;
+		const bool bComputeMTD = true;
+		const FReal Thickness = 0;
 
-			FReal OutTime = -1.0f;
-			FVec3 Normal(0.0f);
-			FVec3 Position(0.0f);
-			int32 FaceIndex = -1;
+		FReal OutTime = -1.0f;
+		FVec3 Normal(0.0f);
+		FVec3 Position(0.0f);
+		int32 FaceIndex = -1;
 
-			// This is local to trimesh, world scale.
-			bool bResult = GJKRaycast2<FReal>(TriangleReg, ScaledCapsule, StartTM, Dir, Length, OutTime, Position, Normal, Thickness, bComputeMTD);
+		// This is local to trimesh, world scale.
+		bool bResult = GJKRaycast2<FReal>(TriangleReg, ScaledCapsule, StartTM, Dir, Length, OutTime, Position, Normal, Thickness, bComputeMTD);
 
-			// Compare results against GJKPenetration, sweep is initial overlap, so this should be the same.
-			FVec3 Normal2, ClosestA, ClosestB;
-			int32 ClosestVertexIndexA, ClosestVertexIndexB;
-			FReal OutTime2;
-			bool bResult2 = GJKPenetration(Triangle, ScaledCapsule, StartTM, OutTime2, ClosestA, ClosestB, Normal2, ClosestVertexIndexA, ClosestVertexIndexB);
+		// Compare results against GJKPenetration, sweep is initial overlap, so this should be the same.
+		FVec3 Normal2, ClosestA, ClosestB;
+		int32 ClosestVertexIndexA, ClosestVertexIndexB;
+		FReal OutTime2;
+		bool bResult2 = GJKPenetration(Triangle, ScaledCapsule, StartTM, OutTime2, ClosestA, ClosestB, Normal2, ClosestVertexIndexA, ClosestVertexIndexB);
 
 
-			EXPECT_VECTOR_NEAR(Normal, Normal2, KINDA_SMALL_NUMBER);
-			EXPECT_NEAR(OutTime, -OutTime2, KINDA_SMALL_NUMBER);
+		EXPECT_VECTOR_NEAR(Normal, Normal2, KINDA_SMALL_NUMBER);
+		EXPECT_NEAR(OutTime, -OutTime2, KINDA_SMALL_NUMBER);
+
+		const FVec3 ClosestBShouldBe{ -287.344025, 1211.66296, 101.851364 };
+		EXPECT_VECTOR_NEAR(ClosestB, ClosestBShouldBe, KINDA_SMALL_NUMBER);
+	}
+
+	//
+	// A real failure case where EPA was terminating too early and returning an incorrect normal
+	// because the UpperBound-LowerBound tolerance check was an absolute rather than relative value
+	// and had an incorrect use of Abs()
+	//
+	GTEST_TEST(EPATests, EPARealFailures_ConvexVsTriangleWrongNormalGJKRaycast2)
+	{
+		using namespace Chaos;
+
+		FTriangleRegister Triangle({
+			MakeVectorRegisterFloat(0.00000000f, 0.00000000f, 0.00000000f, 0.00000000f),
+			MakeVectorRegisterFloat(100.000000f, 100.000000f, 0.00000000f, 0.00000000f),
+			MakeVectorRegisterFloat(0.00000000f, 100.000000f, 0.00000000f, 0.00000000f)
+			});
+
+		TArray<FVec3f> ConvexVerts = {
+			{1.64776051, 0.976988614, 8.85045052},
+			{1.64776051, -0.976994812, 8.85045052},
+			{-1.64776051, -0.976994812, 8.85045052},
+			{-1.64776051, 0.976988614, 8.85045052},
+			{1.64776051, -0.976994812, -0.191102192},
+			{1.64776051, 0.976988614, -0.191102192},
+			{-1.64776051, -0.976994812, -0.191102028},
+			{-1.64776051, 0.976988614, -0.191102028},
+		};
+
+		FImplicitConvex3 Convex(ConvexVerts, 0.0f);
+
+		VectorRegister4Float TranslationSimd = MakeVectorRegisterFloat(13.7357206f, 81.0178833f, 0.975698411f, 0.00000000f);
+		VectorRegister4Float RotationSimd = MakeVectorRegisterFloat(-0.349331319f, -0.614945233f, 0.615562916f, 0.347695649f);
+		VectorRegister4Float DirSimd = MakeVectorRegisterFloat(0.00133391307f, -0.00691976305f, -0.999975145f, 0.00000000f);
+		FReal CurrentLength = 1.0890472489398730;
+
+		FRealSingle Distance;
+		VectorRegister4Float PositionSimd, NormalSimd;
+		bool bHit = GJKRaycast2ImplSimd(Triangle, Convex, RotationSimd, TranslationSimd, DirSimd, FRealSingle(CurrentLength), Distance, PositionSimd, NormalSimd, true, GlobalVectorConstants::Float1000);
+		EXPECT_TRUE(bHit);
+
+		// We should get a hit with a normal pointing upwards but we were getting a normal facing downwards
+		EXPECT_NEAR(NormalSimd.m128_f32[2], 1.0f, UE_KINDA_SMALL_NUMBER);
+	}
 
-			const FVec3 ClosestBShouldBe{ -287.344025, 1211.66296, 101.851364 };
-			EXPECT_VECTOR_NEAR(ClosestB, ClosestBShouldBe, KINDA_SMALL_NUMBER);
-		}
 }
\ No newline at end of file
diff --git a/Engine/Source/Runtime/Experimental/Chaos/Public/Chaos/EPA.h b/Engine/Source/Runtime/Experimental/Chaos/Public/Chaos/EPA.h
--- a/Engine/Source/Runtime/Experimental/Chaos/Public/Chaos/EPA.h
+++ b/Engine/Source/Runtime/Experimental/Chaos/Public/Chaos/EPA.h
@@ -426,6 +426,7 @@ enum class EEPAResult
 	NoValidContact,             // No valid contact have been found, no contacts will be returned
 };
 
+UE_DEPRECATED(5.3, "Not used")
 inline const bool IsEPASuccess(EEPAResult EPAResult)
 {
 	return (EPAResult == EEPAResult::Ok) || (EPAResult == EEPAResult::MaxIterations);
@@ -448,7 +449,7 @@ EEPAResult EPA(TArray<TVec3<T>>& VertsABuffer, TArray<TVec3<T>>& VertsBBuffer, c
 
 template <typename T>
 EEPAResult EPA(TArray<TVec3<T>>& VertsABuffer, TArray<TVec3<T>>& VertsBBuffer, const TFunctionRef<TVector<T, 3>(const TVec3<T>& V)>& SupportA,
-	const TFunctionRef<TVector<T, 3>(const TVec3<T>& V)>& SupportB, T& OutPenetration, TVec3<T>& OutDir, TVec3<T>& WitnessA, TVec3<T>& WitnessB, const FReal Eps = 1.e-2f)
+	const TFunctionRef<TVector<T, 3>(const TVec3<T>& V)>& SupportB, T& OutPenetration, TVec3<T>& OutDir, TVec3<T>& WitnessA, TVec3<T>& WitnessB, const FReal EpsRel = 1.e-2f)
 {
 	struct FEPAEntryWrapper
 	{
@@ -547,8 +548,10 @@ EEPAResult EPA(TArray<TVec3<T>>& VertsABuffer, TArray<TVec3<T>>& VertsBBuffer, c
 
 		LowerBound = Entry.Distance;
 
-		//It's possible the origin is not contained by the CSO. In this case the upper bound will be negative, at which point we should just exit. Maybe return a different enum value?
-		if (FMath::Abs(UpperBound - LowerBound) <= Eps)
+		// It's possible the origin is not contained by the CSO, probably because of numerical error. 
+		// In this case the upper bound will be negative, at which point we should just exit. 
+		const T UpperBoundTolerance = (T(1) + EpsRel) * FMath::Abs(LowerBound);
+		if (UpperBound <= UpperBoundTolerance)
 		{
 			ResultStatus = EEPAResult::Ok;
 			LastEntry = Entry;
diff --git a/Engine/Source/Runtime/Experimental/Chaos/Public/Chaos/EPAVectorized.h b/Engine/Source/Runtime/Experimental/Chaos/Public/Chaos/EPAVectorized.h
--- a/Engine/Source/Runtime/Experimental/Chaos/Public/Chaos/EPAVectorized.h
+++ b/Engine/Source/Runtime/Experimental/Chaos/Public/Chaos/EPAVectorized.h
@@ -431,7 +431,7 @@ FORCEINLINE_DEBUGGABLE void VectorComputeEPAResults(const VectorRegister4Float*
 template <typename TSupportA, typename TSupportB>
 EEPAResult VectorEPA(TArray<VectorRegister4Float>& VertsABuffer, TArray<VectorRegister4Float>& VertsBBuffer, const TSupportA& SupportA, const TSupportB& SupportB, VectorRegister4Float& OutPenetration, VectorRegister4Float& OutDir, VectorRegister4Float& WitnessA, VectorRegister4Float& WitnessB)
 {
-	constexpr VectorRegister4Float Eps = MakeVectorRegisterFloatConstant(1.e-2f, 1.e-2f, 1.e-2f, 1.e-2f);
+	constexpr VectorRegister4Float EpsRel = MakeVectorRegisterFloatConstant(1.e-2f, 1.e-2f, 1.e-2f, 1.e-2f);
 	const VectorRegister4Float OriginInsideEps = VectorZero();
 
 	TEPAWorkingArray<VectorTEPAEntry> Entries;
@@ -518,13 +518,11 @@ EEPAResult VectorEPA(TArray<VectorRegister4Float>& VertsABuffer, TArray<VectorRe
 		UpperBound = VectorMin(DistanceToSupportPlane, UpperBound);
 		LowerBound = Entry.Distance;
 
-		
-		const VectorRegister4Float DiffUpLow = VectorAbs(VectorSubtract(UpperBound, LowerBound));
-		const VectorRegister4Float EpsIsGEAbsBounds = VectorCompareGE(Eps, DiffUpLow);
-
-
-		//It's possible the origin is not contained by the CSO. In this case the upper bound will be negative, at which point we should just exit. Maybe return a different enum value?
-		if (VectorMaskBits(EpsIsGEAbsBounds))
+		// It's possible the origin is not contained by the CSO, probably because of numerical error. 
+		// In this case the upper bound will be negative, at which point we should just exit. 
+		const VectorRegister4Float UpperBoundTolerance = VectorMultiply(VectorAdd(VectorOneFloat(), EpsRel), VectorAbs(LowerBound));
+		const VectorRegister4Float UpperBoundLEUpperBoundTolerance = VectorCompareLE(UpperBound, UpperBoundTolerance);
+		if (VectorMaskBits(UpperBoundLEUpperBoundTolerance))
 		{
 			ResultStatus = EEPAResult::Ok;
 			LastEntry = Entry;
diff --git a/Engine/Source/Runtime/Experimental/Chaos/Public/Chaos/GJK.h b/Engine/Source/Runtime/Experimental/Chaos/Public/Chaos/GJK.h
--- a/Engine/Source/Runtime/Experimental/Chaos/Public/Chaos/GJK.h
+++ b/Engine/Source/Runtime/Experimental/Chaos/Public/Chaos/GJK.h
@@ -1764,37 +1764,39 @@ namespace Chaos
 			VectorRegister4Float InGJKPreDist2GTEps2 = VectorCompareGT(GJKPreDist2, Eps2Simd);
 			VectorRegister4Float LimitMaxGTInGJKPreDist2 = VectorCompareGT(LimitMax, GJKPreDist2);
 			VectorRegister4Float IsDone = VectorBitwiseAnd(InflationGTZero, VectorBitwiseAnd(InGJKPreDist2GTEps2, LimitMaxGTInGJKPreDist2));
-			if (VectorMaskBits(IsDone))
-			{
-				OutNormal = Normal;
-				VectorRegister4Float ClosestA = VectorZeroFloat();
-				VectorRegister4Float ClosestB = VectorZeroFloat();
 
-				if (NumIterations)
-				{
-					VectorRegister4Float Barycentrics[4];
-					Barycentrics[0] = VectorSwizzle(Barycentric, 0, 0, 0, 0);
-					Barycentrics[1] = VectorSwizzle(Barycentric, 1, 1, 1, 1);
-					Barycentrics[2] = VectorSwizzle(Barycentric, 2, 2, 2, 2);
-					Barycentrics[3] = VectorSwizzle(Barycentric, 3, 3, 3, 3);
+			VectorRegister4Float GJKClosestA = VectorZeroFloat();
+			VectorRegister4Float GJKClosestB = VectorZeroFloat();
 
-					alignas(16) int32 NumVertsInts[4];
-					VectorIntStoreAligned(NumVerts, NumVertsInts);
-					const int NumVertsInt = NumVertsInts[0];
-					for (int i = 0; i < NumVertsInt; ++i)
-					{
-						ClosestA = VectorMultiplyAdd(As[i], Barycentrics[i], ClosestA);
-						ClosestB = VectorMultiplyAdd(Bs[i], Barycentrics[i], ClosestB);
-					}
-				}
-				else
+			if (NumIterations)
+			{
+				VectorRegister4Float Barycentrics[4];
+				Barycentrics[0] = VectorSwizzle(Barycentric, 0, 0, 0, 0);
+				Barycentrics[1] = VectorSwizzle(Barycentric, 1, 1, 1, 1);
+				Barycentrics[2] = VectorSwizzle(Barycentric, 2, 2, 2, 2);
+				Barycentrics[3] = VectorSwizzle(Barycentric, 3, 3, 3, 3);
+
+				alignas(16) int32 NumVertsInts[4];
+				VectorIntStoreAligned(NumVerts, NumVertsInts);
+				const int NumVertsInt = NumVertsInts[0];
+				for (int i = 0; i < NumVertsInt; ++i)
 				{
-					//didn't even go into gjk loop
-					ClosestA = As[0];
-					ClosestB = Bs[0];
+					GJKClosestA = VectorMultiplyAdd(As[i], Barycentrics[i], GJKClosestA);
+					GJKClosestB = VectorMultiplyAdd(Bs[i], Barycentrics[i], GJKClosestB);
 				}
+			}
+			else
+			{
+				//didn't even go into gjk loop
+				GJKClosestA = As[0];
+				GJKClosestB = Bs[0];
+			}
+
+			if (VectorMaskBits(IsDone))
+			{
+				OutNormal = Normal;
 
-				const VectorRegister4Float ClosestBInA = VectorAdd(StartPoint, ClosestB);
+				const VectorRegister4Float ClosestBInA = VectorAdd(StartPoint, GJKClosestB);
 				const VectorRegister4Float InGJKPreDist = VectorSqrt(GJKPreDist2);
 				OutNormal = VectorNormalizeAccurate(V);
 
@@ -1802,7 +1804,7 @@ namespace Chaos
 				Penetration = VectorMin(Penetration, LimitMax);
 				Penetration = VectorMax(Penetration, VectorZeroFloat());
 
-				const VectorRegister4Float ClosestLocal = VectorNegateMultiplyAdd(OutNormal, MarginBSimd, ClosestB);
+				const VectorRegister4Float ClosestLocal = VectorNegateMultiplyAdd(OutNormal, MarginBSimd, GJKClosestB);
 
 				OutPosition = VectorAdd(VectorMultiplyAdd(OutNormal, Penetration, StartPoint), ClosestLocal);
 				Penetration = VectorNegate(Penetration);
@@ -1842,26 +1844,37 @@ namespace Chaos
 					VectorRegister4Float Penetration;
 					VectorRegister4Float MTD, ClosestA, ClosestBInA;
 					const EEPAResult EPAResult = VectorEPA(VertsA, VertsB, A, SupportBAtOrigin, Penetration, MTD, ClosestA, ClosestBInA);
-					if (IsEPASuccess(EPAResult))
+					if ((EPAResult == EEPAResult::Ok) || (EPAResult == EEPAResult::MaxIterations))
 					{
 						OutNormal = MTD;
 						VectorStoreFloat1(Penetration, &OutTime);
 						OutTime = -OutTime - (A.Margin + B.Margin);
 						OutPosition = ClosestA;
 					}
-					else if (EPAResult == EEPAResult::NoValidContact)
+					else if (EPAResult == EEPAResult::BadInitialSimplex)
 					{
-						return false;
+						// The origin is outside the simplex. Must be a touching contact and EPA setup will have calculated the normal
+						// and penetration but we keep the position generated by GJK.
+						OutNormal = MTD;
+						VectorStoreFloat1(Penetration, &OutTime);
+						OutTime = -OutTime - (A.Margin + B.Margin);
+						OutPosition = VectorMultiplyAdd(OutNormal, MarginASimd, GJKClosestA);
 					}
-					else
+					else if (EPAResult == EEPAResult::Degenerate)
 					{
-						//assume touching hit
+						// Assume touching hit and use the last GJK result because
+						// EPA does not necessarily return a good normal when GJK provides it with 
+						// a near-degenerate simplex that does not contain the origin (it can reject
+						// the nearest face and return an arbitrarily bad normal)
 						OutTime = -(A.Margin + B.Margin);
-						OutNormal = MTD;
-						OutPosition = VectorMultiplyAdd(OutNormal, MarginASimd, As[0]);
+						OutNormal = Normal;
+						OutPosition = VectorMultiplyAdd(OutNormal, MarginASimd, GJKClosestA);
+					}
+					else // EEPAResult::NoValidContact
+					{
+						return false;
 					}
 				}
-
 				else
 				{
 					//didn't even go into gjk loop, touching hit
