Chaos - add unit test to check TArray Reserve and Reset policy doesn't change (we require both to tightly allocate buffers)

diff --git a/Engine/Source/Programs/HeadlessChaos/Private/HeadlessChaosApp.cpp b/Engine/Source/Programs/HeadlessChaos/Private/HeadlessChaosApp.cpp
--- a/Engine/Source/Programs/HeadlessChaos/Private/HeadlessChaosApp.cpp
+++ b/Engine/Source/Programs/HeadlessChaos/Private/HeadlessChaosApp.cpp
@@ -402,6 +402,66 @@ TEST(GeometryCollection_FieldTest,Fields_SumScalarLeftSide) { GeometryCollection
 TEST(GeometryCollection_FieldTest,Fields_Culling) { GeometryCollectionTest::Fields_Culling(); SUCCEED(); }
 TEST(GeometryCollection_FieldTest,Fields_SerializeAPI) { GeometryCollectionTest::Fields_SerializeAPI(); SUCCEED(); }
 
+GTEST_TEST(ArrayTests, TestArrayMax)
+{
+	// The first 3 arrays without Reserve will over-allocate. We aren't testing anything
+	// useful on these - they are just here for examples...
+
+	// This allocates space for 4 elements
+	TArray<int32> Ints1;
+	Ints1.SetNum(1);
+	EXPECT_LE(Ints1.Num(), Ints1.Max());
+
+	// This allocates space for 5 elements
+	TArray<int32> Ints2;
+	Ints2.SetNum(5);
+	EXPECT_LE(Ints2.Num(), Ints2.Max());
+
+	// This allocates space for 4 elements and grows to 22 elements
+	TArray<int32> Ints3;
+	Ints3.SetNum(1);
+	Ints3.SetNum(5);
+	EXPECT_LE(Ints3.Num(), Ints3.Max());
+
+	// We rely on tight-fitting arrays for memory conservation. Make sure that the
+	// default Reserve policy still enforces tight-fitting arrays.
+	TArray<int32> Ints4;
+	Ints4.Reserve(1);
+	Ints4.SetNum(1);
+	EXPECT_EQ(Ints4.Max(), Ints4.Num());
+
+	TArray<int32> Ints5;
+	Ints5.Reserve(5);
+	Ints5.SetNum(5);
+	EXPECT_EQ(Ints5.Max(), Ints5.Num());
+
+	TArray<int32> Ints6;
+	Ints6.Reserve(1);
+	Ints6.SetNum(1);
+	Ints6.Reserve(5);
+	Ints6.SetNum(5);
+	EXPECT_EQ(Ints6.Max(), Ints6.Num());
+
+	// Reset also sets the exact buffer size and we are relying on that too
+	TArray<int32> Ints7;
+	Ints7.Reset(1);
+	Ints7.SetNum(1);
+	EXPECT_EQ(Ints7.Max(), Ints7.Num());
+
+	TArray<int32> Ints8;
+	Ints8.Reset(5);
+	Ints8.SetNum(5);
+	EXPECT_EQ(Ints8.Max(), Ints8.Num());
+
+	TArray<int32> Ints9;
+	Ints9.Reset(1);
+	Ints9.SetNum(1);
+	Ints9.Reset(5);
+	Ints9.SetNum(5);
+	EXPECT_EQ(Ints9.Max(), Ints9.Num());
+
+}
+
 //TEST(GeometryCollectionTest,RigidBodies_CollisionGroup); // fix me
 //
 // Broken
