Fixed HeadlessChaos compile after SOA memory changes (double->float)

diff --git a/Engine/Source/Programs/HeadlessChaos/Private/GeometryCollection/GeometryCollectionTestSimulationField.cpp b/Engine/Source/Programs/HeadlessChaos/Private/GeometryCollection/GeometryCollectionTestSimulationField.cpp
--- a/Engine/Source/Programs/HeadlessChaos/Private/GeometryCollection/GeometryCollectionTestSimulationField.cpp
+++ b/Engine/Source/Programs/HeadlessChaos/Private/GeometryCollection/GeometryCollectionTestSimulationField.cpp
@@ -1134,7 +1134,7 @@ namespace GeometryCollectionTest
 		const auto& ClusterMap = Clustering.GetChildrenMap();
 
 		
-		Chaos::TArrayCollectionArray<Chaos::FReal>& StrainArray = UnitTest.Solver->GetEvolution()->GetRigidClustering().GetStrainArray();
+		Chaos::TArrayCollectionArray<Chaos::FRealSingle>& StrainArray = UnitTest.Solver->GetEvolution()->GetRigidClustering().GetStrainArray();
 
 		EXPECT_FALSE(ParticleHandles[0]->Disabled());
 		EXPECT_TRUE(ParticleHandles[1]->Disabled());
@@ -1203,7 +1203,7 @@ namespace GeometryCollectionTest
 		const auto& ClusterMap = Clustering.GetChildrenMap();
 
 
-		Chaos::TArrayCollectionArray<Chaos::FReal>& StrainArray = UnitTest.Solver->GetEvolution()->GetRigidClustering().GetStrainArray();
+		Chaos::TArrayCollectionArray<Chaos::FRealSingle>& StrainArray = UnitTest.Solver->GetEvolution()->GetRigidClustering().GetStrainArray();
 
 		EXPECT_FALSE(ParticleHandles[0]->Disabled());
 		EXPECT_TRUE(ParticleHandles[1]->Disabled());
diff --git a/Engine/Source/Programs/HeadlessChaos/Private/HeadlessChaosTestClustering.cpp b/Engine/Source/Programs/HeadlessChaos/Private/HeadlessChaosTestClustering.cpp
--- a/Engine/Source/Programs/HeadlessChaos/Private/HeadlessChaosTestClustering.cpp
+++ b/Engine/Source/Programs/HeadlessChaos/Private/HeadlessChaosTestClustering.cpp
@@ -246,7 +246,7 @@ namespace ChaosTest {
 		PhysicalMaterial->DisabledLinearThreshold = 0;
 		PhysicalMaterial->DisabledAngularThreshold = 0;
 
-		Chaos::TArrayCollectionArray<FReal>& SolverStrainArray = Evolution.GetRigidClustering().GetStrainArray();
+		Chaos::TArrayCollectionArray<FRealSingle>& SolverStrainArray = Evolution.GetRigidClustering().GetStrainArray();
 
 		for (int i = 0; i < NumBoxes + NumBoxes / 4 + 1; ++i)
 		{
