fix headless chaos tests GeometryCollection_RigidBodies_ClusterTest_ReleaseClusterParticle_ClusteredNode, GeometryCollection_RigidBodies_ClusterTest_ReleaseClusterParticle_ClusteredKinematicNode, and GeometryCollection_RigidBodies_ClusterTest_ReleaseClusterParticles_ClusterNodeAndSubClusterNode

diff --git a/Engine/Source/Programs/HeadlessChaos/Private/GeometryCollection/GeometryCollectionTestClustering.cpp b/Engine/Source/Programs/HeadlessChaos/Private/GeometryCollection/GeometryCollectionTestClustering.cpp
--- a/Engine/Source/Programs/HeadlessChaos/Private/GeometryCollection/GeometryCollectionTestClustering.cpp
+++ b/Engine/Source/Programs/HeadlessChaos/Private/GeometryCollection/GeometryCollectionTestClustering.cpp
@@ -1783,6 +1783,8 @@ namespace GeometryCollectionTest
 		UnitTest.Initialize();
 
 		auto& Clustering = UnitTest.Solver->GetEvolution()->GetRigidClustering();
+		Chaos::FClusterUnionManager& ClusterUnionManager = Clustering.GetClusterUnionManager();
+
 		const auto& ClusterMap = Clustering.GetChildrenMap();
 
 		TSharedPtr<FGeometryDynamicCollection> DynamicCollection1 = Collection1->DynamicCollection;
@@ -1809,6 +1811,19 @@ namespace GeometryCollectionTest
 		EXPECT_EQ(SovlerParticleHandles.Size(), 7);
 		EXPECT_EQ(ClusterMap.Num(), 1);
 		EXPECT_TRUE(ClusterMapContains(ClusterMap, SovlerParticleHandles.Handle(6)->CastToRigidParticle(), { Collection1Handles[1],Collection1Handles[0], Collection2Handles[1],Collection2Handles[0] }));
+		EXPECT_TRUE(ClusterUnionManager.FindClusterUnionFromExplicitIndex(1) != nullptr);
+
+		Chaos::FClusterUnion& ClusterUnion = *ClusterUnionManager.FindClusterUnionFromExplicitIndex(1);
+		EXPECT_EQ(ClusterUnion.InternalCluster, SovlerParticleHandles.Handle(6).Get());
+		// A bit of an assumption that the root particle is the 2nd index in these arrays.
+		EXPECT_TRUE(ClusterUnion.ChildParticles.Contains(Collection1Handles[0]));
+		EXPECT_TRUE(ClusterUnion.ChildParticles.Contains(Collection1Handles[1]));
+		EXPECT_TRUE(ClusterUnion.ChildParticles.Contains(Collection2Handles[0]));
+		EXPECT_TRUE(ClusterUnion.ChildParticles.Contains(Collection2Handles[1]));
+		EXPECT_EQ(ClusterUnion.ChildParticles.Num(), 4);
+		EXPECT_EQ(ClusterUnion.ExplicitIndex, 1);
+		EXPECT_TRUE(ClusterMapContains(ClusterMap, ClusterUnion.InternalCluster, { Collection1Handles[0], Collection1Handles[1], Collection2Handles[0], Collection2Handles[1] }));
+		EXPECT_EQ(ClusterMap[ClusterUnion.InternalCluster].Num(), 4);
 
 		TArray<FTransform> Collection1_PreReleaseTM; GeometryCollectionAlgo::GlobalMatrices(DynamicCollection1->Transform, DynamicCollection1->Parent, Collection1_PreReleaseTM);
 		TArray<FTransform> Collection2_PreReleaseTM; GeometryCollectionAlgo::GlobalMatrices(DynamicCollection2->Transform, DynamicCollection2->Parent, Collection2_PreReleaseTM);
@@ -1819,15 +1834,20 @@ namespace GeometryCollectionTest
 
 		UnitTest.Solver->GetEvolution()->GetRigidClustering().ReleaseClusterParticles({ Collection1Handles[0],Collection1Handles[1] });
 
-		EXPECT_EQ(SovlerParticleHandles.Size(), 8);
-		EXPECT_EQ(ClusterMap.Num(), 1);
-		EXPECT_TRUE(ClusterMapContains(ClusterMap, SovlerParticleHandles.Handle(7)->CastToRigidParticle(), { Collection2Handles[1],Collection2Handles[0] }));
+		EXPECT_EQ(SovlerParticleHandles.Size(), 7);
+		EXPECT_FALSE(ClusterUnion.ChildParticles.Contains(Collection1Handles[0]));
+		EXPECT_FALSE(ClusterUnion.ChildParticles.Contains(Collection1Handles[1]));
+		EXPECT_TRUE(ClusterUnion.ChildParticles.Contains(Collection2Handles[0]));
+		EXPECT_TRUE(ClusterUnion.ChildParticles.Contains(Collection2Handles[1]));
+		EXPECT_EQ(ClusterUnion.ChildParticles.Num(), 2);
+		EXPECT_EQ(ClusterUnion.ExplicitIndex, 1);
+		EXPECT_FALSE(ClusterMapContains(ClusterMap, ClusterUnion.InternalCluster, { Collection1Handles[0], Collection1Handles[1] }));
+		EXPECT_TRUE(ClusterMapContains(ClusterMap, ClusterUnion.InternalCluster, { Collection2Handles[0], Collection2Handles[1] }));
+		EXPECT_EQ(ClusterMap[ClusterUnion.InternalCluster].Num(), 2);
 
 		UnitTest.Advance();
 
-		EXPECT_EQ(SovlerParticleHandles.Size(), 8);
-		EXPECT_EQ(ClusterMap.Num(), 1);
-		EXPECT_TRUE(ClusterMapContains(ClusterMap, SovlerParticleHandles.Handle(7)->CastToRigidParticle(), {Collection2Handles[1],Collection2Handles[0] }));
+		EXPECT_EQ(SovlerParticleHandles.Size(), 7);
 
 		TArray<FTransform> Collection1_PostReleaseTM; GeometryCollectionAlgo::GlobalMatrices(DynamicCollection1->Transform, DynamicCollection1->Parent, Collection1_PostReleaseTM);
 		TArray<FTransform> Collection2_PostReleaseTM; GeometryCollectionAlgo::GlobalMatrices(DynamicCollection2->Transform, DynamicCollection2->Parent, Collection2_PostReleaseTM);
@@ -1882,6 +1902,7 @@ namespace GeometryCollectionTest
 		UnitTest.Initialize();
 
 		auto& Clustering = UnitTest.Solver->GetEvolution()->GetRigidClustering();
+		Chaos::FClusterUnionManager& ClusterUnionManager = Clustering.GetClusterUnionManager();
 		const auto& ClusterMap = Clustering.GetChildrenMap();
 
 		TArray<FTransform> Collection1_InitialTM; GeometryCollectionAlgo::GlobalMatrices(DynamicCollection1->Transform, DynamicCollection1->Parent, Collection1_InitialTM);
@@ -1904,6 +1925,19 @@ namespace GeometryCollectionTest
 		EXPECT_EQ(SovlerParticleHandles.Size(), 7);
 		EXPECT_EQ(ClusterMap.Num(), 1);
 		EXPECT_TRUE(ClusterMapContains(ClusterMap, SovlerParticleHandles.Handle(6)->CastToRigidParticle(), { Collection1Handles[1],Collection1Handles[0], Collection2Handles[1],Collection2Handles[0] }));
+		EXPECT_TRUE(ClusterUnionManager.FindClusterUnionFromExplicitIndex(1) != nullptr);
+
+		Chaos::FClusterUnion& ClusterUnion = *ClusterUnionManager.FindClusterUnionFromExplicitIndex(1);
+		EXPECT_EQ(ClusterUnion.InternalCluster, SovlerParticleHandles.Handle(6).Get());
+		// A bit of an assumption that the root particle is the 2nd index in these arrays.
+		EXPECT_TRUE(ClusterUnion.ChildParticles.Contains(Collection1Handles[0]));
+		EXPECT_TRUE(ClusterUnion.ChildParticles.Contains(Collection1Handles[1]));
+		EXPECT_TRUE(ClusterUnion.ChildParticles.Contains(Collection2Handles[0]));
+		EXPECT_TRUE(ClusterUnion.ChildParticles.Contains(Collection2Handles[1]));
+		EXPECT_EQ(ClusterUnion.ChildParticles.Num(), 4);
+		EXPECT_EQ(ClusterUnion.ExplicitIndex, 1);
+		EXPECT_TRUE(ClusterMapContains(ClusterMap, ClusterUnion.InternalCluster, { Collection1Handles[0], Collection1Handles[1], Collection2Handles[0], Collection2Handles[1] }));
+		EXPECT_EQ(ClusterMap[ClusterUnion.InternalCluster].Num(), 4);
 
 		TArray<FTransform> Collection1_PreReleaseTM; GeometryCollectionAlgo::GlobalMatrices(DynamicCollection1->Transform, DynamicCollection1->Parent, Collection1_PreReleaseTM);
 		TArray<FTransform> Collection2_PreReleaseTM; GeometryCollectionAlgo::GlobalMatrices(DynamicCollection2->Transform, DynamicCollection2->Parent, Collection2_PreReleaseTM);
@@ -1915,15 +1949,20 @@ namespace GeometryCollectionTest
 
 		UnitTest.Solver->GetEvolution()->GetRigidClustering().ReleaseClusterParticles({ Collection1Handles[0],Collection1Handles[1] });
 
-		EXPECT_EQ(SovlerParticleHandles.Size(), 8);
-		EXPECT_EQ(ClusterMap.Num(), 1);
-		EXPECT_TRUE(ClusterMapContains(ClusterMap, SovlerParticleHandles.Handle(7)->CastToRigidParticle(), { Collection2Handles[1],Collection2Handles[0] }));
+		EXPECT_EQ(SovlerParticleHandles.Size(), 7);
+		EXPECT_FALSE(ClusterUnion.ChildParticles.Contains(Collection1Handles[0]));
+		EXPECT_FALSE(ClusterUnion.ChildParticles.Contains(Collection1Handles[1]));
+		EXPECT_TRUE(ClusterUnion.ChildParticles.Contains(Collection2Handles[0]));
+		EXPECT_TRUE(ClusterUnion.ChildParticles.Contains(Collection2Handles[1]));
+		EXPECT_EQ(ClusterUnion.ChildParticles.Num(), 2);
+		EXPECT_EQ(ClusterUnion.ExplicitIndex, 1);
+		EXPECT_FALSE(ClusterMapContains(ClusterMap, ClusterUnion.InternalCluster, { Collection1Handles[0], Collection1Handles[1] }));
+		EXPECT_TRUE(ClusterMapContains(ClusterMap, ClusterUnion.InternalCluster, { Collection2Handles[0], Collection2Handles[1] }));
+		EXPECT_EQ(ClusterMap[ClusterUnion.InternalCluster].Num(), 2);
 
 		UnitTest.Advance();
 
-		EXPECT_EQ(SovlerParticleHandles.Size(), 8);
-		EXPECT_EQ(ClusterMap.Num(), 1);
-		EXPECT_TRUE(ClusterMapContains(ClusterMap, SovlerParticleHandles.Handle(7)->CastToRigidParticle(), { Collection2Handles[1],Collection2Handles[0] }));
+		EXPECT_EQ(SovlerParticleHandles.Size(), 7);
 
 		// validate that DynamicCollection2 became dynamic and fell from the cluster. 
 
@@ -2062,6 +2101,8 @@ namespace GeometryCollectionTest
 		UnitTest.Advance();
 
 		FRigidClustering& Clustering = UnitTest.Solver->GetEvolution()->GetRigidClustering();
+		Chaos::FClusterUnionManager& ClusterUnionManager = Clustering.GetClusterUnionManager();
+
 		const FRigidClustering::FClusterMap& ClusterMap = Clustering.GetChildrenMap();
 		const Chaos::TArrayCollectionArray<Chaos::ClusterId>& ClusterIdsArray = Clustering.GetClusterIdsArray();
 
@@ -2072,6 +2113,10 @@ namespace GeometryCollectionTest
 			EXPECT_TRUE(ClusterMapContains(ClusterMap,ParticleHandles[4],{ParticleHandles[2]}));
 		});
 
+		EXPECT_TRUE(ClusterUnionManager.FindClusterUnionFromExplicitIndex(1) != nullptr);
+
+		Chaos::FClusterUnion& ClusterUnion = *ClusterUnionManager.FindClusterUnionFromExplicitIndex(1);
+
 		// Test releasing a specific unioned cluster
 		// We end up with the following cluster tree
 		//     4
@@ -2096,6 +2141,13 @@ namespace GeometryCollectionTest
 				EXPECT_NE(ClusterIdsArray[0].Id, nullptr);
 				EXPECT_EQ(ClusterIdsArray[1].Id, nullptr);
 				EXPECT_EQ(ClusterIdsArray[2].Id, nullptr);
+
+				EXPECT_TRUE(ClusterUnion.ChildParticles.Contains(ParticleHandles[2]));
+				EXPECT_EQ(ClusterUnion.ChildParticles.Num(), 1);
+				EXPECT_EQ(ClusterUnion.ExplicitIndex, 1);
+				EXPECT_TRUE(ClusterMapContains(ClusterMap, ClusterUnion.InternalCluster, { ParticleHandles[2] }));
+				EXPECT_EQ(ClusterMap[ClusterUnion.InternalCluster].Num(), 1);
+				EXPECT_EQ(ClusterMap.Num(), 2);
 			}
 			else
 			{
@@ -2104,7 +2156,12 @@ namespace GeometryCollectionTest
 				EXPECT_EQ(ClusterIdsArray[1].Id, nullptr);
 				EXPECT_EQ(ClusterIdsArray[2].Id, nullptr);
 
-				EXPECT_EQ(ClusterMap.Num(), 1);
+				EXPECT_FALSE(ClusterUnion.ChildParticles.Contains(ParticleHandles[2]));
+				EXPECT_EQ(ClusterUnion.ChildParticles.Num(), 0);
+				EXPECT_EQ(ClusterUnion.ExplicitIndex, 1);
+				EXPECT_FALSE(ClusterMapContains(ClusterMap, ClusterUnion.InternalCluster, { ParticleHandles[2] }));
+				EXPECT_EQ(ClusterMap[ClusterUnion.InternalCluster].Num(), 0);
+				EXPECT_EQ(ClusterMap.Num(), 2);
 				EXPECT_TRUE(ClusterMapContains(ClusterMap, ParticleHandles[2], { ParticleHandles[0], ParticleHandles[1] }));
 			}				
 		}
