properly call FClusterUnionManager::HandleDeferredClusterUnionUpdateProperties after ReleaseClusterParticles call without relying on it getting to the next breaking model

diff --git a/Engine/Source/Programs/HeadlessChaos/Private/GeometryCollection/GeometryCollectionTestClustering.cpp b/Engine/Source/Programs/HeadlessChaos/Private/GeometryCollection/GeometryCollectionTestClustering.cpp
--- a/Engine/Source/Programs/HeadlessChaos/Private/GeometryCollection/GeometryCollectionTestClustering.cpp
+++ b/Engine/Source/Programs/HeadlessChaos/Private/GeometryCollection/GeometryCollectionTestClustering.cpp
@@ -1833,6 +1833,7 @@ namespace GeometryCollectionTest
 		}
 
 		UnitTest.Solver->GetEvolution()->GetRigidClustering().ReleaseClusterParticles({ Collection1Handles[0],Collection1Handles[1] });
+		ClusterUnionManager.HandleDeferredClusterUnionUpdateProperties();
 
 		EXPECT_EQ(SovlerParticleHandles.Size(), 7);
 		EXPECT_FALSE(ClusterUnion.ChildParticles.Contains(Collection1Handles[0]));
@@ -1948,6 +1949,7 @@ namespace GeometryCollectionTest
 		}
 
 		UnitTest.Solver->GetEvolution()->GetRigidClustering().ReleaseClusterParticles({ Collection1Handles[0],Collection1Handles[1] });
+		ClusterUnionManager.HandleDeferredClusterUnionUpdateProperties();
 
 		EXPECT_EQ(SovlerParticleHandles.Size(), 7);
 		EXPECT_FALSE(ClusterUnion.ChildParticles.Contains(Collection1Handles[0]));
diff --git a/Engine/Source/Runtime/Experimental/GeometryCollectionEngine/Private/GeometryCollection/GeometryCollectionComponent.cpp b/Engine/Source/Runtime/Experimental/GeometryCollectionEngine/Private/GeometryCollection/GeometryCollectionComponent.cpp
--- a/Engine/Source/Runtime/Experimental/GeometryCollectionEngine/Private/GeometryCollection/GeometryCollectionComponent.cpp
+++ b/Engine/Source/Runtime/Experimental/GeometryCollectionEngine/Private/GeometryCollection/GeometryCollectionComponent.cpp
@@ -1937,6 +1937,8 @@ bool UGeometryCollectionComponent::ProcessRepData(const float DeltaTime, const f
 		}
 	}
 
+	RigidClustering.GetClusterUnionManager().HandleDeferredClusterUnionUpdateProperties();
+
 	// Keep track of whether we did some "work" on this frame so we can turn off the async tick after
 	// multiple frames of not doing anything.
 	bool bProcessed = false;
