Fix unit test compilation error

diff --git a/Engine/Source/Programs/HeadlessChaos/Private/GeometryCollection/GeometryCollectionTestEvents.cpp b/Engine/Source/Programs/HeadlessChaos/Private/GeometryCollection/GeometryCollectionTestEvents.cpp
--- a/Engine/Source/Programs/HeadlessChaos/Private/GeometryCollection/GeometryCollectionTestEvents.cpp
+++ b/Engine/Source/Programs/HeadlessChaos/Private/GeometryCollection/GeometryCollectionTestEvents.cpp
@@ -30,6 +30,12 @@ namespace GeometryCollectionTest
 		FVector Data2;
 	};
 
+	template<typename PayloadType>
+	const TMap<IPhysicsProxyBase*, TArray<int32>>* GetProxyToIndexMap(const PayloadType* Buffer)
+	{
+		return nullptr;
+	}
+
 	template<typename PayloadType>
 	bool IsEventDataEmpty(const PayloadType* Buffer)
 	{
