Minor compile fix for WinArm64.

diff --git a/Engine/Source/Programs/HeadlessChaos/Private/HeadlessChaosTestEPA.cpp b/Engine/Source/Programs/HeadlessChaos/Private/HeadlessChaosTestEPA.cpp
--- a/Engine/Source/Programs/HeadlessChaos/Private/HeadlessChaosTestEPA.cpp
+++ b/Engine/Source/Programs/HeadlessChaos/Private/HeadlessChaosTestEPA.cpp
@@ -1360,7 +1360,7 @@ namespace ChaosTest
 		EXPECT_TRUE(bHit);
 
 		// We should get a hit with a normal pointing upwards but we were getting a normal facing downwards
-		EXPECT_NEAR(NormalSimd.m128_f32[2], 1.0f, UE_KINDA_SMALL_NUMBER);
+		EXPECT_NEAR(VectorGetComponent(NormalSimd, 2), 1.0f, UE_KINDA_SMALL_NUMBER);
 	}
 
 }
\ No newline at end of file
