Fix Sweep cast on Height Field

diff --git a/Engine/Source/Programs/HeadlessChaos/Private/HeadlessChaosTestHeightfield.cpp b/Engine/Source/Programs/HeadlessChaos/Private/HeadlessChaosTestHeightfield.cpp
--- a/Engine/Source/Programs/HeadlessChaos/Private/HeadlessChaosTestHeightfield.cpp
+++ b/Engine/Source/Programs/HeadlessChaos/Private/HeadlessChaosTestHeightfield.cpp
@@ -712,6 +712,62 @@ namespace ChaosTest {
 		}
 	}
 
+	void SweepBigSphereTest()
+	{
+		const int32 Columns = 10;
+		const int32 Rows = 10;
+		const FReal CountToWorldScale = 100;
+
+
+		TSphere<FReal, 3> Sphere(FVec3(0.0, 0.0, 0.0), 250.0);
+
+		TArray<FReal> Heights;
+		Heights.AddZeroed(Rows * Columns);
+
+		for (int32 Row = 0; Row < Rows; ++Row)
+		{
+			for (int32 Col = 0; Col < Columns; ++Col)
+			{
+				if (Col == 5)
+				{
+					Heights[Row * Columns + Col] = 10;
+				}
+				else
+				{
+					Heights[Row * Columns + Col] = 0;
+				}
+			}
+		}
+
+		TArray<FReal> HeightsCopy = Heights;
+		FHeightField Heightfield(MoveTemp(HeightsCopy), TArray<uint8>(), Rows, Columns, FVec3(100, 100, 100));
+		const auto& Bounds = Heightfield.BoundingBox();	//Current API forces us to do this to cache the bounds
+
+		FReal TOI;
+		FVec3 Position, Normal, FaceNormal;
+		int32 FaceIdx;
+		{
+			const FVec3 Start(200, 500, 700);
+			FRigidTransform3 StartTM(Start, TRotation<FReal, 3>::Identity);
+			FVec3 Dir(1.0, -0.1, 0.0);
+			Dir.Normalize();
+
+			bool Result = Heightfield.SweepGeom(Sphere, StartTM, Dir, 50, TOI, Position, Normal, FaceIdx, FaceNormal, 0.0, true);
+			EXPECT_TRUE(Result);
+		}
+		// Test extent in the opposite direction of the sweep
+		{
+			const FVec3 Start(350, 500, 700);
+			FRigidTransform3 StartTM(Start, TRotation<FReal, 3>::Identity);
+			FVec3 Dir(-1.0, 0.1, 0.0);
+			Dir.Normalize();
+
+			bool Result = Heightfield.SweepGeom(Sphere, StartTM, Dir, 50, TOI, Position, Normal, FaceIdx, FaceNormal, 0.0, true);
+			EXPECT_TRUE(Result);
+		}
+
+	}
+
 	void SweepSmallBoxTest()
 	{
 		const int32 Columns = 10;
@@ -929,6 +985,7 @@ namespace ChaosTest {
 	void SweepTest()
 	{
 		SweepSmallSphereTest();
+		SweepBigSphereTest();
 		SweepSmallBoxTest();
 		SweepBoxTest();
 	}
diff --git a/Engine/Source/Runtime/Experimental/Chaos/Private/Chaos/HeightField.cpp b/Engine/Source/Runtime/Experimental/Chaos/Private/Chaos/HeightField.cpp
--- a/Engine/Source/Runtime/Experimental/Chaos/Private/Chaos/HeightField.cpp
+++ b/Engine/Source/Runtime/Experimental/Chaos/Private/Chaos/HeightField.cpp
@@ -1409,7 +1409,10 @@ namespace Chaos
 		FVec2 ClippedStart;
 		FVec2 ClippedEnd;
 
-		if (InflatedBounds.ClipLine(StartPoint, StartPoint + Dir * Length, ClippedStart, ClippedEnd))
+		const FVec3 InflateStartPoint = StartPoint - Dir * InHalfExtents;
+		const FVec3 InflateEndPoint = StartPoint + Dir * (InHalfExtents + Length);
+
+		if (InflatedBounds.ClipLine(InflateStartPoint, InflateEndPoint, ClippedStart, ClippedEnd))
 		{
 			// Rasterize the line over the grid
 			TVec2<int32> StartCell = FlatGrid.Cell(ClippedStart / Scale2D);
