Add unit test for GeometryCollection View Enable/Disable

diff --git a/Engine/Source/Programs/HeadlessChaos/Private/GeometryCollection/GeometryCollectionTestSimulation.cpp b/Engine/Source/Programs/HeadlessChaos/Private/GeometryCollection/GeometryCollectionTestSimulation.cpp
--- a/Engine/Source/Programs/HeadlessChaos/Private/GeometryCollection/GeometryCollectionTestSimulation.cpp
+++ b/Engine/Source/Programs/HeadlessChaos/Private/GeometryCollection/GeometryCollectionTestSimulation.cpp
@@ -223,6 +223,39 @@ namespace GeometryCollectionTest
 
 	}
 
+	GTEST_TEST(AllTraits, GeometryCollection_RigidBodies_Enabling)
+	{
+		CreationParameters Params;
+		Params.SimplicialType = ESimplicialType::Chaos_Simplicial_Box;
+		Params.ImplicitType = EImplicitTypeEnum::Chaos_Implicit_Box;
+
+		Params.DynamicState = EObjectStateTypeEnum::Chaos_Object_Dynamic;
+		Params.RootTransform.SetLocation(FVector(0.f, 0.f, 15.f));
+		FGeometryCollectionWrapper* MovingCollection = TNewSimulationObject<GeometryType::GeometryCollectionWithSingleRigid>::Init(Params)->template As<FGeometryCollectionWrapper>();
+
+		FFramework UnitTest;
+		UnitTest.AddSimulationObject(MovingCollection);
+		UnitTest.Initialize();
+
+		const auto& Transform0 = MovingCollection->DynamicCollection->Transform[0];
+		for (int i = 0; i < 5; i++)
+		{
+			UnitTest.Advance();
+			EXPECT_EQ(MovingCollection->DynamicCollection->DynamicState[0], (int32)EObjectStateTypeEnum::Chaos_Object_Dynamic);
+			EXPECT_LT(Transform0.GetTranslation().Z, 15.0f);
+		}
+		// Disabled particle
+		MovingCollection->PhysObject->DisableParticles_External({ 0 });
+		FReal CurrentPosition = Transform0.GetTranslation().Z;
+		for (int i = 0; i < 5; i++)
+		{
+			UnitTest.Advance();
+			EXPECT_EQ(MovingCollection->DynamicCollection->DynamicState[0], (int32)EObjectStateTypeEnum::Chaos_Object_Dynamic);
+			EXPECT_EQ(Transform0.GetTranslation().Z, CurrentPosition);
+		}
+	}
+
+
 	// CollisionGroup == 0 : Collide With Everything Except CollisionGroup=-1
 	// CollisionGroup == -1 : Collide With Nothing Including CollisionGroup=0
 	// CollisionGroup_A == CollisionGroup_B : Collide With Each Other
