Fix Sweep Query on Height Field. Extend of the shape query was not considered correctly in the sweep direction.

diff --git a/Engine/Source/Programs/HeadlessChaos/Private/HeadlessChaosTestHeightfield.cpp b/Engine/Source/Programs/HeadlessChaos/Private/HeadlessChaosTestHeightfield.cpp
--- a/Engine/Source/Programs/HeadlessChaos/Private/HeadlessChaosTestHeightfield.cpp
+++ b/Engine/Source/Programs/HeadlessChaos/Private/HeadlessChaosTestHeightfield.cpp
@@ -765,6 +765,25 @@ namespace ChaosTest {
 			bool Result = Heightfield.SweepGeom(Sphere, StartTM, Dir, 50, TOI, Position, Normal, FaceIdx, FaceNormal, 0.0, true);
 			EXPECT_TRUE(Result);
 		}
+		{
+			const FVec3 Start(180, 500, 700);
+			FRigidTransform3 StartTM(Start, TRotation<FReal, 3>::Identity);
+			FVec3 Dir(1.0, 0.0, -1.0);
+			Dir.Normalize();
+
+			bool Result = Heightfield.SweepGeom(Sphere, StartTM, Dir, 50, TOI, Position, Normal, FaceIdx, FaceNormal, 0.0, true);
+			EXPECT_TRUE(Result);
+		}
+		// Test extent in the opposite direction of the sweep
+		{
+			const FVec3 Start(220, 500, 700);
+			FRigidTransform3 StartTM(Start, TRotation<FReal, 3>::Identity);
+			FVec3 Dir(-1.0, 0.0, 1.0);
+			Dir.Normalize();
+
+			bool Result = Heightfield.SweepGeom(Sphere, StartTM, Dir, 50, TOI, Position, Normal, FaceIdx, FaceNormal, 0.0, true);
+			EXPECT_TRUE(Result);
+		}
 
 	}
 
diff --git a/Engine/Source/Runtime/Experimental/Chaos/Private/Chaos/HeightField.cpp b/Engine/Source/Runtime/Experimental/Chaos/Private/Chaos/HeightField.cpp
--- a/Engine/Source/Runtime/Experimental/Chaos/Private/Chaos/HeightField.cpp
+++ b/Engine/Source/Runtime/Experimental/Chaos/Private/Chaos/HeightField.cpp
@@ -1409,8 +1409,12 @@ namespace Chaos
 		FVec2 ClippedStart;
 		FVec2 ClippedEnd;
 
-		const FVec3 InflateStartPoint = StartPoint - Dir * InHalfExtents;
-		const FVec3 InflateEndPoint = StartPoint + Dir * (InHalfExtents + Length);
+		FVec3 SignDir(0);
+		SignDir[0] = FMath::Sign(Dir[0]);
+		SignDir[1] = FMath::Sign(Dir[1]);
+		const FVec3 InflateStartEnd = SignDir * InHalfExtents;
+		const FVec3 InflateStartPoint = StartPoint - InflateStartEnd;
+		const FVec3 InflateEndPoint = StartPoint + Dir * Length + InflateStartEnd;
 
 		if (InflatedBounds.ClipLine(InflateStartPoint, InflateEndPoint, ClippedStart, ClippedEnd))
 		{
