Chaos - big island manager refactor

- Islands are now first class members of the constraint graph and are persistent
- We no longer need to visit sleeping nodes in island merge and split
- Isolated particles (with no constraints) are not represented in the graph, massively reducing the initial and peak island counts in demo scenes
- Nodes, edges and islands are allocated from a pool and have persistent pointers
- Sleep logic modified to early-out as quickly as possible
- More opportunites for parallelism (not done yet)

diff --git a/Engine/Source/Programs/HeadlessChaos/Private/Character/TestCharacterGroundConstraintContainer.cpp b/Engine/Source/Programs/HeadlessChaos/Private/Character/TestCharacterGroundConstraintContainer.cpp
--- a/Engine/Source/Programs/HeadlessChaos/Private/Character/TestCharacterGroundConstraintContainer.cpp
+++ b/Engine/Source/Programs/HeadlessChaos/Private/Character/TestCharacterGroundConstraintContainer.cpp
@@ -181,7 +181,7 @@ namespace ChaosTest
 		}
 
 		// Initialize graph
-		Private::FPBDIslandManager Graph;
+		Private::FPBDIslandManager Graph(SOAs);
 		Container.SetContainerId(0); // Usually set by the evolution the container is registered with
 		Graph.AddConstraintContainer(Container);
 
@@ -193,11 +193,11 @@ namespace ChaosTest
 		{
 			Graph.AddParticle(Particle);
 		}
-		Graph.InitializeGraph(SOAs.GetNonDisabledDynamicView());
+		Graph.UpdateParticles();
 
 		// Add constraints and update graph
 		Container.AddConstraintsToGraph(Graph);
-		Graph.UpdateIslands(SOAs);
+		Graph.UpdateIslands();
 
 		auto Constraints = Container.GetConstraints();
 		for (FCharacterGroundConstraintHandle* Constraint : Constraints)
@@ -205,18 +205,18 @@ namespace ChaosTest
 			EXPECT_TRUE(Constraint->IsInConstraintGraph());
 		}
 
-		TArray<TSet<FGeometryParticleHandle*>> ExpectedParticlesInIsland =
+		TArray<TArray<FGeometryParticleHandle*>> ExpectedParticlesInIsland =
 		{
 			{ DynamicParticles[0], DynamicParticles[1], DynamicParticles[9] },
 			{ DynamicParticles[2] },
 			{ DynamicParticles[3], DynamicParticles[10] },
 			{ DynamicParticles[4], DynamicParticles[11], DynamicParticles[12] },
 			{ DynamicParticles[5], DynamicParticles[13] },
-			{ DynamicParticles[6], DynamicParticles[14] },
-			{ DynamicParticles[7] },
-			{ DynamicParticles[8] },
+			{ DynamicParticles[6], DynamicParticles[14], StaticParticles[0] },
+			{ DynamicParticles[7], StaticParticles[1] },
+			{ DynamicParticles[8], StaticParticles[1] },
 		};
-		TArray<TSet<FCharacterGroundConstraintHandle*>> ExpectedConstraintsInIsland =
+		TArray<TArray<FCharacterGroundConstraintHandle*>> ExpectedConstraintsInIsland =
 		{
 			{ Constraints[0], Constraints[1] },
 			{ Constraints[2] },
@@ -228,33 +228,27 @@ namespace ChaosTest
 			{ Constraints[11] },
 		};
 		const int ExpectedNumIslands = ExpectedParticlesInIsland.Num();
-		TArray<int32> IslandIndices =
-		{
-			DynamicParticles[0]->CastToRigidParticle()->IslandIndex(),
-			DynamicParticles[2]->CastToRigidParticle()->IslandIndex(),
-			DynamicParticles[3]->CastToRigidParticle()->IslandIndex(),
-			DynamicParticles[4]->CastToRigidParticle()->IslandIndex(),
-			DynamicParticles[5]->CastToRigidParticle()->IslandIndex(),
-			DynamicParticles[6]->CastToRigidParticle()->IslandIndex(),
-			DynamicParticles[7]->CastToRigidParticle()->IslandIndex(),
-			DynamicParticles[8]->CastToRigidParticle()->IslandIndex(),
-		};
 
 		for (int32 IslandIdx = 0; IslandIdx < ExpectedNumIslands; ++IslandIdx)
 		{
-			Private::FPBDIsland* Island = Graph.GetIsland(IslandIdx);
-			TArrayView<Private::FPBDIslandParticle> IslandParticles = Island->GetParticles();
+			// Find all the islands that contain one of these particles (should only be one because they are all dynamic)
+			const TArray<const Private::FPBDIsland*> Islands = Graph.FindParticleIslands(ExpectedParticlesInIsland[IslandIdx][0]);
+			EXPECT_EQ(Islands.Num(), 1);
+
+			// Find all the particles in the island and make sure they match the expected set
+			const TArray<const FGeometryParticleHandle*> IslandParticles = Graph.FindParticlesInIslands(Islands);
 			EXPECT_EQ(IslandParticles.Num(), ExpectedParticlesInIsland[IslandIdx].Num());
-			for (Private::FPBDIslandParticle& IslandParticle : IslandParticles)
+			for (const FGeometryParticleHandle* ExpectedIslandParticle : ExpectedParticlesInIsland[IslandIdx])
 			{
-				EXPECT_TRUE(ExpectedParticlesInIsland[IslandIdx].Contains(IslandParticle.GetParticle()));
+				EXPECT_TRUE(IslandParticles.Contains(ExpectedIslandParticle));
 			}
 
-			TArrayView<Private::FPBDIslandConstraint> IslandConstraints = Island->GetConstraints(Container.GetContainerId());
+			// Find all the constraints in the island and make sure they match the expected set
+			const TArray<const FConstraintHandle*> IslandConstraints = Graph.FindConstraintsInIslands(Islands, Container.GetContainerId());
 			EXPECT_EQ(IslandConstraints.Num(), ExpectedConstraintsInIsland[IslandIdx].Num());
-			for (Private::FPBDIslandConstraint& IslandConstraint : IslandConstraints)
+			for (const FConstraintHandle* ExpectedIslandConstraint : ExpectedConstraintsInIsland[IslandIdx])
 			{
-				EXPECT_TRUE(ExpectedConstraintsInIsland[IslandIdx].Contains(IslandConstraint.GetConstraint()->As<FCharacterGroundConstraintHandle>()));
+				EXPECT_TRUE(IslandConstraints.Contains(ExpectedIslandConstraint));
 			}
 		}
 
@@ -288,16 +282,16 @@ namespace ChaosTest
 		Constraints.Add(Container.AddConstraint(Settings[0], Data[0], Particles[1], Particles[3]));
 		Constraints.Add(Container.AddConstraint(Settings[0], Data[0], Particles[2]));
 
-		Private::FPBDIslandManager Graph;
+		Private::FPBDIslandManager Graph(SOAs);
 		Container.SetContainerId(0);
 		Graph.AddConstraintContainer(Container);
 		for (auto Particle : Particles)
 		{
 			Graph.AddParticle(Particle);
 		}
-		Graph.InitializeGraph(SOAs.GetNonDisabledDynamicView());
+		Graph.UpdateParticles();
 		Container.AddConstraintsToGraph(Graph);
-		Graph.UpdateIslands(SOAs);
+		Graph.UpdateIslands();
 
 		// First use non-island version
 		Solver->AddConstraints();
diff --git a/Engine/Source/Programs/HeadlessChaos/Private/HeadlessChaosTestClustering.cpp b/Engine/Source/Programs/HeadlessChaos/Private/HeadlessChaosTestClustering.cpp
--- a/Engine/Source/Programs/HeadlessChaos/Private/HeadlessChaosTestClustering.cpp
+++ b/Engine/Source/Programs/HeadlessChaos/Private/HeadlessChaosTestClustering.cpp
@@ -147,7 +147,7 @@ namespace ChaosTest {
 				EXPECT_NE(Particle.Handle(), ClusteredParticles.Handle(BoxID));	//make sure boxes are not in non disabled array
 			}
 
-			EXPECT_TRUE(Evolution.GetIslandManager().DebugCheckParticleNotInGraph(ClusteredParticles.Handle(BoxID)));
+			EXPECT_FALSE(ClusteredParticles.Handle(BoxID)->IsInConstraintGraph());
 		}
 
 		for (Chaos::FPBDRigidParticleHandle* ClusterHandle : ClusterHandles)
@@ -173,7 +173,7 @@ namespace ChaosTest {
 				EXPECT_NE(Particle.Handle(), ClusterHandle);	//make sure boxes are not in non disabled array
 			}
 
-			EXPECT_TRUE(Evolution.GetIslandManager().DebugCheckParticleNotInGraph(ClusterHandle));
+			EXPECT_FALSE(ClusterHandle->IsInConstraintGraph());
 		}
 
 		EXPECT_EQ(Particles.GetNonDisabledView().Num(), NumBoxes);
@@ -272,7 +272,7 @@ namespace ChaosTest {
 		{
 			EXPECT_TRUE(ClusteredParticles.Disabled(BoxID));	//no boxes should be active yet
 			EXPECT_TRUE(Evolution.GetRigidClustering().GetTopLevelClusterParents().Contains(ClusteredParticles.Handle(BoxID)) == false);
-			EXPECT_TRUE(Evolution.GetIslandManager().DebugCheckParticleNotInGraph(ClusteredParticles.Handle(BoxID)));
+			EXPECT_FALSE(ClusteredParticles.Handle(BoxID)->IsInConstraintGraph());
 		}
 
 		SolverStrainArray[NumBoxes + NumBoxes / 4 + 1] = (FReal)1;
@@ -289,7 +289,7 @@ namespace ChaosTest {
 		{
 			EXPECT_TRUE(ClusteredParticles.Disabled(BoxID));	//no boxes should be active yet
 			EXPECT_TRUE(Evolution.GetRigidClustering().GetTopLevelClusterParents().Contains(ClusteredParticles.Handle(BoxID)) == false);
-			EXPECT_TRUE(Evolution.GetIslandManager().DebugCheckParticleNotInGraph(ClusteredParticles.Handle(BoxID)));
+			EXPECT_FALSE(ClusteredParticles.Handle(BoxID)->IsInConstraintGraph());
 		}
 		
 	}
diff --git a/Engine/Source/Programs/HeadlessChaos/Private/HeadlessChaosTestGraphEvolution.cpp b/Engine/Source/Programs/HeadlessChaos/Private/HeadlessChaosTestGraphEvolution.cpp
--- a/Engine/Source/Programs/HeadlessChaos/Private/HeadlessChaosTestGraphEvolution.cpp
+++ b/Engine/Source/Programs/HeadlessChaos/Private/HeadlessChaosTestGraphEvolution.cpp
@@ -77,9 +77,12 @@ namespace ChaosTest
 				Advance();
 
 				bIsSleeping = true;
-				for (int32 IslandIndex = 0; IslandIndex < IslandManager->NumIslands(); ++IslandIndex)
+				for (FPBDRigidParticleHandle* ParticleHandle : ParticleHandles)
 				{
-					bIsSleeping = bIsSleeping && IslandManager->GetIsland(IslandIndex)->IsSleeping();
+					if (ParticleHandle->IsDynamic() && !ParticleHandle->IsSleeping())
+					{
+						bIsSleeping = false;
+					}
 				}
 			}
 
@@ -108,13 +111,85 @@ namespace ChaosTest
 
 		Test.Advance();
 
-		EXPECT_EQ(Test.IslandManager->NumIslands(), 1);
+		EXPECT_EQ(Test.IslandManager->GetNumIslands(), 1);
+	}
+
+	// Start with a kinematic connected to a dynamic. Verify that removing the
+	// constraint removed both particles.
+	// This version explicitly removes the constraint from the graph.
+	// 
+	//		(d=dynamic, s=sleeping, k=kinematic)
+	//		Ak - Bd
+	// =>	{}
+	//
+	GTEST_TEST(GraphEvolutionTests, TestConstraintGraph_KinematicDynamic_Remove)
+	{
+		FGraphEvolutionTest Test(2);
+		Test.MakeChain();
+
+		Test.Evolution.SetParticleObjectState(Test.ParticleHandles[0], EObjectStateType::Kinematic);
+
+		Test.Advance();
+
+		// Should have 1 island
+		EXPECT_EQ(Test.IslandManager->GetNumIslands(), 1);
+		EXPECT_TRUE(Test.ParticleHandles[0]->IsInConstraintGraph());
+		EXPECT_TRUE(Test.ParticleHandles[1]->IsInConstraintGraph());
+		EXPECT_TRUE(Test.ConstraintHandles[0]->IsInConstraintGraph());
+
+		// Disable the constraint and remove it from the graph
+		Test.ConstraintHandles[0]->SetEnabled(false);
+		Test.IslandManager->RemoveConstraint(Test.ConstraintHandles[0]);
+
+		Test.Advance();
+
+		// Should have no islands and all particles should have been removed
+		EXPECT_EQ(Test.IslandManager->GetNumIslands(), 0);
+		EXPECT_FALSE(Test.ParticleHandles[0]->IsInConstraintGraph());
+		EXPECT_FALSE(Test.ParticleHandles[1]->IsInConstraintGraph());
+		EXPECT_FALSE(Test.ConstraintHandles[0]->IsInConstraintGraph());
+	}
+
+	// Start with a kinematic connected to a dynamic. Verify that removing the
+	// constraint removed both particles.
+	// This version has the constraint removed by making all particles kinematic
+	//		(d=dynamic, s=sleeping, k=kinematic)
+	//		Ak - Bd
+	// =>	Ak - Bk
+	// =>	{}
+	//
+	GTEST_TEST(GraphEvolutionTests, TestConstraintGraph_KinematicDynamic_Remove2)
+	{
+		FGraphEvolutionTest Test(2);
+		Test.MakeChain();
+
+		Test.Evolution.SetParticleObjectState(Test.ParticleHandles[0], EObjectStateType::Kinematic);
+
+		Test.Advance();
+
+		// Should have 1 island
+		EXPECT_EQ(Test.IslandManager->GetNumIslands(), 1);
+		EXPECT_TRUE(Test.ParticleHandles[0]->IsInConstraintGraph());
+		EXPECT_TRUE(Test.ParticleHandles[1]->IsInConstraintGraph());
+		EXPECT_TRUE(Test.ConstraintHandles[0]->IsInConstraintGraph());
+
+		// Make the other particle kinematic
+		Test.Evolution.SetParticleObjectState(Test.ParticleHandles[1], EObjectStateType::Kinematic);
+
+		Test.Advance();
+
+		// Should have no islands and all particles should have been removed
+		EXPECT_EQ(Test.IslandManager->GetNumIslands(), 0);
+		EXPECT_FALSE(Test.ParticleHandles[0]->IsInConstraintGraph());
+		EXPECT_FALSE(Test.ParticleHandles[1]->IsInConstraintGraph());
+		EXPECT_FALSE(Test.ConstraintHandles[0]->IsInConstraintGraph());
 	}
 
 	// Start with an island containing 4 particle connected in a chain, then make the second one kinematic.
 	// Check that the island splits.
-	//		A-B-C-D 
-	// =>	A-B, B-C-D
+	//		(d=dynamic, s=sleeping, k=kinematic)
+	//		Ad - Bd - Cd - Dd 
+	// =>	Ad - Bk   Bk - Cd - Dd
 	//
 	GTEST_TEST(GraphEvolutionTests, TestConstraintGraph_ToKinematic)
 	{
@@ -123,7 +198,7 @@ namespace ChaosTest
 
 		Test.Advance();
 
-		EXPECT_EQ(Test.IslandManager->NumIslands(), 1);
+		EXPECT_EQ(Test.IslandManager->GetNumIslands(), 1);
 
 		// Convert particle B to kinematic to split the islands: A-B, B-C-D
 		Test.Evolution.SetParticleObjectState(Test.ParticleHandles[1], EObjectStateType::Kinematic);
@@ -131,7 +206,30 @@ namespace ChaosTest
 		Test.Advance();
 
 		// Should now have 2 islands
-		EXPECT_EQ(Test.IslandManager->NumIslands(), 2);
+		EXPECT_EQ(Test.IslandManager->GetNumIslands(), 2);
+
+		// A should be in its own island
+		EXPECT_EQ(Test.IslandManager->GetParticleIsland(Test.ParticleHandles[0])->GetNumParticles(), 1);
+
+		// C and D should be in same island
+		EXPECT_EQ(Test.IslandManager->GetParticleIsland(Test.ParticleHandles[2]), Test.IslandManager->GetParticleIsland(Test.ParticleHandles[3]));
+
+		// B should be in 2 islands
+		EXPECT_EQ(Test.IslandManager->FindParticleIslands(Test.ParticleHandles[1]).Num(), 2);
+	}
+
+	// An isolated stationary particle with no gravity should go to sleep
+	GTEST_TEST(GraphEvolutionTests, TestConstraintGraph_ParticleSleep_Isolated)
+	{
+		FGraphEvolutionTest Test(1);
+		Test.Evolution.GetGravityForces().SetAcceleration(FVec3(0), 0);
+
+		// Make all the particles sleep
+		Test.AdvanceUntilSleeping();
+
+		// Particle should be asleep and it should have taken 21 ticks (ChaosSolverCollisionDefaultSleepCounterThreshold)
+		EXPECT_TRUE(Test.ParticleHandles[0]->IsSleeping());
+		EXPECT_EQ(Test.TickCount, 21);
 	}
 
 	// Wait for all particles to go to sleep naturally (i.e., as part of the tick and not by explicitly setting the state)
@@ -144,28 +242,24 @@ namespace ChaosTest
 		Test.Advance();
 
 		// All constraints in graph in 1 island that is awake
-		EXPECT_EQ(Test.IslandManager->NumIslands(), 1);
-		EXPECT_NE(Test.ConstraintHandles[0]->GetConstraintGraphIndex(), INDEX_NONE);
-		EXPECT_NE(Test.ConstraintHandles[1]->GetConstraintGraphIndex(), INDEX_NONE);
-		EXPECT_NE(Test.ConstraintHandles[2]->GetConstraintGraphIndex(), INDEX_NONE);
+		EXPECT_EQ(Test.IslandManager->GetNumIslands(), 1);
+		EXPECT_TRUE(Test.ConstraintHandles[0]->IsInConstraintGraph());
+		EXPECT_TRUE(Test.ConstraintHandles[1]->IsInConstraintGraph());
+		EXPECT_TRUE(Test.ConstraintHandles[2]->IsInConstraintGraph());
 		EXPECT_FALSE(Test.IslandManager->GetIsland(0)->IsSleeping());
 
 		// Make all the particles sleep
 		Test.AdvanceUntilSleeping();
 
 		// Island should be asleep
-		EXPECT_EQ(Test.IslandManager->NumIslands(), 1);
+		EXPECT_EQ(Test.IslandManager->GetNumIslands(), 1);
 		EXPECT_TRUE(Test.IslandManager->GetIsland(0)->IsSleeping());
 	}
 
 	// Sleep all particles in the scene and ensure that the island manager puts the island to sleep
 	// but retains all the constraints and particles in the island.
-	// 
-	// @todo(chaos): this test fails. Manually setting all particles in an island to the sleep state
-	// does put the island to sleep, but only after all the constraints have been removed because
-	// InitializeGraph is called before the sleep state is updated.
 	//
-	GTEST_TEST(GraphEvolutionTests, DISABLED_TestConstraintGraph_ParticleSleep_Manual)
+	GTEST_TEST(GraphEvolutionTests, TestConstraintGraph_ParticleSleep_Manual)
 	{
 		FGraphEvolutionTest Test(4);
 		Test.MakeChain();
@@ -173,10 +267,10 @@ namespace ChaosTest
 		Test.Advance();
 
 		// All constraints in graph in 1 island that is awake
-		EXPECT_EQ(Test.IslandManager->NumIslands(), 1);
-		EXPECT_NE(Test.ConstraintHandles[0]->GetConstraintGraphIndex(), INDEX_NONE);
-		EXPECT_NE(Test.ConstraintHandles[1]->GetConstraintGraphIndex(), INDEX_NONE);
-		EXPECT_NE(Test.ConstraintHandles[2]->GetConstraintGraphIndex(), INDEX_NONE);
+		EXPECT_EQ(Test.IslandManager->GetNumIslands(), 1);
+		EXPECT_TRUE(Test.ConstraintHandles[0]->IsInConstraintGraph());
+		EXPECT_TRUE(Test.ConstraintHandles[1]->IsInConstraintGraph());
+		EXPECT_TRUE(Test.ConstraintHandles[2]->IsInConstraintGraph());
 		EXPECT_FALSE(Test.IslandManager->GetIsland(0)->IsSleeping());
 
 		// Make all the particles sleep
@@ -188,13 +282,13 @@ namespace ChaosTest
 		Test.Advance();
 
 		// Island should be asleep
-		EXPECT_EQ(Test.IslandManager->NumIslands(), 1);
+		EXPECT_EQ(Test.IslandManager->GetNumIslands(), 1);
 		EXPECT_TRUE(Test.IslandManager->GetIsland(0)->IsSleeping());
 	}
 
 	// Start with an island containing 4 particle connected in a chain, then make the middle two kinematic.
 	// This makes the B-C constraint kinematic ("invalid" in the Constraint Graph terminolgy) which means
-	// it does not belong in any island.
+	// it does not belong in any island and is kicked out of the graph (the edge is deleted).
 	// Check that the island manager handles kinematic-kinematic constraints
 	//		A-B-C-D 
 	// =>	A-B  C-D
@@ -207,10 +301,10 @@ namespace ChaosTest
 		Test.Advance();
 
 		// All constraints in graph in 1 island
-		EXPECT_EQ(Test.IslandManager->NumIslands(), 1);
-		EXPECT_NE(Test.ConstraintHandles[0]->GetConstraintGraphIndex(), INDEX_NONE);
-		EXPECT_NE(Test.ConstraintHandles[1]->GetConstraintGraphIndex(), INDEX_NONE);
-		EXPECT_NE(Test.ConstraintHandles[2]->GetConstraintGraphIndex(), INDEX_NONE);
+		EXPECT_EQ(Test.IslandManager->GetNumIslands(), 1);
+		EXPECT_TRUE(Test.ConstraintHandles[0]->IsInConstraintGraph());
+		EXPECT_TRUE(Test.ConstraintHandles[1]->IsInConstraintGraph());
+		EXPECT_TRUE(Test.ConstraintHandles[2]->IsInConstraintGraph());
 
 		// Convert a particle B and C to kinematic to split the islands: A-B, C-D
 		Test.Evolution.SetParticleObjectState(Test.ParticleHandles[1], EObjectStateType::Kinematic);
@@ -218,19 +312,23 @@ namespace ChaosTest
 
 		Test.Advance();
 
-		// Constraint[1] is not in an island, but is still in the graph (we have not explicitly removed it)
-		EXPECT_NE(Test.ConstraintHandles[0]->GetConstraintGraphIndex(), INDEX_NONE);
-		EXPECT_NE(Test.ConstraintHandles[1]->GetConstraintGraphIndex(), INDEX_NONE);				// Still in graph
-		EXPECT_EQ(Test.IslandManager->GetConstraintIsland(Test.ConstraintHandles[1]), INDEX_NONE);	// Not in an island
-		EXPECT_NE(Test.ConstraintHandles[2]->GetConstraintGraphIndex(), INDEX_NONE);
+		// Constraint[1] was kicked from the graph
+		EXPECT_TRUE(Test.ConstraintHandles[0]->IsInConstraintGraph());
+		EXPECT_FALSE(Test.ConstraintHandles[1]->IsInConstraintGraph());	// Not in graph
+		EXPECT_TRUE(Test.ConstraintHandles[2]->IsInConstraintGraph());
 
 		// Should now have 2 islands
-		EXPECT_EQ(Test.IslandManager->NumIslands(), 2);
+		EXPECT_EQ(Test.IslandManager->GetNumIslands(), 2);
 	}
 
 
 	// Same as TestConstraintGraph_KinematicKinematic but islands are sleeping when the change is made.
-	// Sleeping islands are not cleared of constraints, but they should be awoken by the dynamic->kinematic conversion
+	// 
+	// (d=dynamic, s=sleeping, k=kinematic)
+	//		As - Bs - Cs - Ds
+	// =>	As - Bk   Ck - Ds
+	// =>	Ad - Bk   Ck - Ds
+	//
 	GTEST_TEST(GraphEvolutionTests, TestConstraintGraph_KinematicKinematic_Sleeping)
 	{
 		FGraphEvolutionTest Test(4);
@@ -239,47 +337,47 @@ namespace ChaosTest
 		Test.Advance();
 
 		// All constraints in graph in 1 island that is awake
-		EXPECT_EQ(Test.IslandManager->NumIslands(), 1);
-		EXPECT_NE(Test.ConstraintHandles[0]->GetConstraintGraphIndex(), INDEX_NONE);
-		EXPECT_NE(Test.ConstraintHandles[1]->GetConstraintGraphIndex(), INDEX_NONE);
-		EXPECT_NE(Test.ConstraintHandles[2]->GetConstraintGraphIndex(), INDEX_NONE);
+		EXPECT_EQ(Test.IslandManager->GetNumIslands(), 1);
 		EXPECT_FALSE(Test.IslandManager->GetIsland(0)->IsSleeping());
+		EXPECT_TRUE(Test.ConstraintHandles[0]->IsInConstraintGraph());
+		EXPECT_TRUE(Test.ConstraintHandles[1]->IsInConstraintGraph());
+		EXPECT_TRUE(Test.ConstraintHandles[2]->IsInConstraintGraph());
 
 		// Wait for the sleep state
 		Test.AdvanceUntilSleeping();
 
 		// Island should be asleep but still contain all the particles and constraints
-		EXPECT_EQ(Test.IslandManager->NumIslands(), 1);
+		EXPECT_EQ(Test.IslandManager->GetNumIslands(), 1);
 		EXPECT_TRUE(Test.IslandManager->GetIsland(0)->IsSleeping());
-		EXPECT_NE(Test.ConstraintHandles[0]->GetConstraintGraphIndex(), INDEX_NONE);
-		EXPECT_NE(Test.ConstraintHandles[1]->GetConstraintGraphIndex(), INDEX_NONE);
-		EXPECT_NE(Test.ConstraintHandles[2]->GetConstraintGraphIndex(), INDEX_NONE);
+		EXPECT_TRUE(Test.ConstraintHandles[0]->IsInConstraintGraph());
+		EXPECT_TRUE(Test.ConstraintHandles[1]->IsInConstraintGraph());
+		EXPECT_TRUE(Test.ConstraintHandles[2]->IsInConstraintGraph());
 
 		// Convert a particle B and C to kinematic to split the islands: A-B, C-D
-		// NOTE: the island won't split until the island is woken
 		Test.Evolution.SetParticleObjectState(Test.ParticleHandles[1], EObjectStateType::Kinematic);
 		Test.Evolution.SetParticleObjectState(Test.ParticleHandles[2], EObjectStateType::Kinematic);
 
 		Test.Advance();
 
-		// Island is asleep and won't have split yet
-		EXPECT_EQ(Test.IslandManager->NumIslands(), 1);
+		// The kinematic-kinematic constraint was kicked from graph
+		EXPECT_TRUE(Test.ConstraintHandles[0]->IsInConstraintGraph());
+		EXPECT_FALSE(Test.ConstraintHandles[1]->IsInConstraintGraph());
+		EXPECT_TRUE(Test.ConstraintHandles[2]->IsInConstraintGraph());
 
-		// The kinematic-kinematic constraint will still be in the graph, and still in its island (it will not be removed until the island wakes)
-		EXPECT_NE(Test.ConstraintHandles[1]->GetConstraintGraphIndex(), INDEX_NONE);
-		EXPECT_NE(Test.IslandManager->GetConstraintIsland(Test.ConstraintHandles[1]), INDEX_NONE);
+		// Island will have split into two sleeping islands
+		EXPECT_EQ(Test.IslandManager->GetNumIslands(), 2);
+		EXPECT_TRUE(Test.IslandManager->GetIsland(0)->IsSleeping());
+		EXPECT_TRUE(Test.IslandManager->GetIsland(1)->IsSleeping());
 
-		// Wake a dynamic particle the island and we should get the split at the next update
+		// Wake a dynamic particle
 		Test.Evolution.SetParticleObjectState(Test.ParticleHandles[0], EObjectStateType::Dynamic);
 
 		Test.Advance();
 
-		// Should now have 2 islands (with 2 particles and 1 constraint each though we can't check that atm)
-		EXPECT_EQ(Test.IslandManager->NumIslands(), 2);
-
-		// The kinematic-kinematic constraint will still be in the graph, but not in any island
-		EXPECT_NE(Test.ConstraintHandles[1]->GetConstraintGraphIndex(), INDEX_NONE);
-		EXPECT_EQ(Test.IslandManager->GetConstraintIsland(Test.ConstraintHandles[1]), INDEX_NONE);
+		// Should now have 2 islands and only one awake particle
+		EXPECT_EQ(Test.IslandManager->GetNumIslands(), 2);
+		EXPECT_FALSE(Test.ParticleHandles[0]->IsSleeping());
+		EXPECT_TRUE(Test.ParticleHandles[3]->IsSleeping());
 	}
 
 	// 3 objects sat on the floor awake. Make the floor dynamic.
@@ -298,7 +396,7 @@ namespace ChaosTest
 		Test.Advance();
 
 		// Each particle in its own island (kinematic will be in all 3)
-		EXPECT_EQ(Test.IslandManager->NumIslands(), 3);
+		EXPECT_EQ(Test.IslandManager->GetNumIslands(), 3);
 		EXPECT_EQ(Test.IslandManager->FindParticleIslands(Test.ParticleHandles[0]).Num(), 3);
 		EXPECT_EQ(Test.IslandManager->FindParticleIslands(Test.ParticleHandles[1]).Num(), 1);
 		EXPECT_EQ(Test.IslandManager->FindParticleIslands(Test.ParticleHandles[2]).Num(), 1);
@@ -310,7 +408,7 @@ namespace ChaosTest
 		Test.Advance();
 
 		// All particles in same island
-		EXPECT_EQ(Test.IslandManager->NumIslands(), 1);
+		EXPECT_EQ(Test.IslandManager->GetNumIslands(), 1);
 		EXPECT_EQ(Test.IslandManager->FindParticleIslands(Test.ParticleHandles[0]).Num(), 1);
 		EXPECT_EQ(Test.IslandManager->FindParticleIslands(Test.ParticleHandles[1]).Num(), 1);
 		EXPECT_EQ(Test.IslandManager->FindParticleIslands(Test.ParticleHandles[2]).Num(), 1);
@@ -325,7 +423,7 @@ namespace ChaosTest
 	//  \   |   /	   =>     \   |   /
 	//      Ak		              Ad
 	//
-	GTEST_TEST(GraphEvolutionTests, TestConstraintGraph_KinematicToDynamic_WithSleep)
+	GTEST_TEST(GraphEvolutionTests, TestConstraintGraph_KinematicToDynamic_WithSleep1)
 	{
 		FGraphEvolutionTest Test(4);
 		Test.MakeFloor();
@@ -333,7 +431,7 @@ namespace ChaosTest
 		Test.AdvanceUntilSleeping();
 
 		// Each particle in its own island (kinematic will be in all 3)
-		EXPECT_EQ(Test.IslandManager->NumIslands(), 3);
+		EXPECT_EQ(Test.IslandManager->GetNumIslands(), 3);
 		EXPECT_EQ(Test.IslandManager->FindParticleIslands(Test.ParticleHandles[0]).Num(), 3);
 		EXPECT_EQ(Test.IslandManager->FindParticleIslands(Test.ParticleHandles[1]).Num(), 1);
 		EXPECT_EQ(Test.IslandManager->FindParticleIslands(Test.ParticleHandles[2]).Num(), 1);
@@ -345,7 +443,7 @@ namespace ChaosTest
 		Test.Advance();
 
 		// All particles in one awake island
-		EXPECT_EQ(Test.IslandManager->NumIslands(), 1);
+		EXPECT_EQ(Test.IslandManager->GetNumIslands(), 1);
 		EXPECT_FALSE(Test.IslandManager->GetIsland(0)->IsSleeping());
 		EXPECT_EQ(Test.IslandManager->FindParticleIslands(Test.ParticleHandles[0]).Num(), 1);
 		EXPECT_EQ(Test.IslandManager->FindParticleIslands(Test.ParticleHandles[1]).Num(), 1);
@@ -376,7 +474,7 @@ namespace ChaosTest
 		Test.AdvanceUntilSleeping();
 
 		// Each particle in its own island (kinematic will be in all 3)
-		EXPECT_EQ(Test.IslandManager->NumIslands(), 3);
+		EXPECT_EQ(Test.IslandManager->GetNumIslands(), 3);
 		EXPECT_EQ(Test.IslandManager->FindParticleIslands(Test.ParticleHandles[0]).Num(), 3);
 		EXPECT_EQ(Test.IslandManager->FindParticleIslands(Test.ParticleHandles[1]).Num(), 1);
 		EXPECT_EQ(Test.IslandManager->FindParticleIslands(Test.ParticleHandles[2]).Num(), 1);
@@ -388,7 +486,7 @@ namespace ChaosTest
 		Test.Advance();
 
 		// All particles in one asleep island
-		EXPECT_EQ(Test.IslandManager->NumIslands(), 1);
+		EXPECT_EQ(Test.IslandManager->GetNumIslands(), 1);
 		EXPECT_TRUE(Test.IslandManager->GetIsland(0)->IsSleeping());
 		EXPECT_EQ(Test.IslandManager->FindParticleIslands(Test.ParticleHandles[0]).Num(), 1);
 		EXPECT_EQ(Test.IslandManager->FindParticleIslands(Test.ParticleHandles[1]).Num(), 1);
@@ -428,7 +526,7 @@ namespace ChaosTest
 		EXPECT_FALSE(FConstGenericParticleHandle(Test.ParticleHandles[3])->Sleeping());
 
 		// Each particle in its own island (kinematic will be in all 3)
-		EXPECT_EQ(Test.IslandManager->NumIslands(), 3);
+		EXPECT_EQ(Test.IslandManager->GetNumIslands(), 3);
 		EXPECT_EQ(Test.IslandManager->FindParticleIslands(Test.ParticleHandles[0]).Num(), 3);
 		EXPECT_EQ(Test.IslandManager->FindParticleIslands(Test.ParticleHandles[1]).Num(), 1);
 		EXPECT_EQ(Test.IslandManager->FindParticleIslands(Test.ParticleHandles[2]).Num(), 1);
@@ -440,7 +538,7 @@ namespace ChaosTest
 		Test.Advance();
 
 		// All particles in one awake island (D was awake so it would wake the island)
-		EXPECT_EQ(Test.IslandManager->NumIslands(), 1);
+		EXPECT_EQ(Test.IslandManager->GetNumIslands(), 1);
 		EXPECT_FALSE(Test.IslandManager->GetIsland(0)->IsSleeping());
 		EXPECT_EQ(Test.IslandManager->FindParticleIslands(Test.ParticleHandles[0]).Num(), 1);
 		EXPECT_EQ(Test.IslandManager->FindParticleIslands(Test.ParticleHandles[1]).Num(), 1);
@@ -481,7 +579,7 @@ namespace ChaosTest
 		EXPECT_TRUE(FConstGenericParticleHandle(Test.ParticleHandles[3])->Sleeping());
 
 		// Each particle in its own island (kinematic will be in all 3)
-		EXPECT_EQ(Test.IslandManager->NumIslands(), 3);
+		EXPECT_EQ(Test.IslandManager->GetNumIslands(), 3);
 		EXPECT_EQ(Test.IslandManager->FindParticleIslands(Test.ParticleHandles[0]).Num(), 3);
 		EXPECT_EQ(Test.IslandManager->FindParticleIslands(Test.ParticleHandles[1]).Num(), 1);
 		EXPECT_EQ(Test.IslandManager->FindParticleIslands(Test.ParticleHandles[2]).Num(), 1);
@@ -493,7 +591,7 @@ namespace ChaosTest
 		Test.Advance();
 
 		// All particles in one awake island (B was awake so it would wake the island, even though A was set to sleeping)
-		EXPECT_EQ(Test.IslandManager->NumIslands(), 1);
+		EXPECT_EQ(Test.IslandManager->GetNumIslands(), 1);
 		EXPECT_FALSE(Test.IslandManager->GetIsland(0)->IsSleeping());
 		EXPECT_EQ(Test.IslandManager->FindParticleIslands(Test.ParticleHandles[0]).Num(), 1);
 		EXPECT_EQ(Test.IslandManager->FindParticleIslands(Test.ParticleHandles[1]).Num(), 1);
@@ -526,10 +624,10 @@ namespace ChaosTest
 		EXPECT_TRUE(FConstGenericParticleHandle(Test.ParticleHandles[1])->Sleeping());
 		EXPECT_TRUE(FConstGenericParticleHandle(Test.ParticleHandles[2])->Sleeping());
 		EXPECT_TRUE(FConstGenericParticleHandle(Test.ParticleHandles[3])->Sleeping());
-		EXPECT_EQ(Test.IslandManager->NumIslands(), 1);
+		EXPECT_EQ(Test.IslandManager->GetNumIslands(), 1);
 
 		// Flag the island for wake up
-		if (Test.IslandManager->NumIslands() == 1)
+		if (Test.IslandManager->GetNumIslands() == 1)
 		{
 			EXPECT_TRUE(Test.IslandManager->GetIsland(0)->IsSleeping());
 			EXPECT_TRUE(false);	// Should be calling WakeIsland here
@@ -539,8 +637,8 @@ namespace ChaosTest
 		Test.Advance();
 
 		// Island and all particles should now be awake
-		EXPECT_EQ(Test.IslandManager->NumIslands(), 1);
-		if (Test.IslandManager->NumIslands() == 1)
+		EXPECT_EQ(Test.IslandManager->GetNumIslands(), 1);
+		if (Test.IslandManager->GetNumIslands() == 1)
 		{
 			EXPECT_FALSE(Test.IslandManager->GetIsland(0)->IsSleeping());
 		}
@@ -569,20 +667,20 @@ namespace ChaosTest
 		EXPECT_FALSE(FConstGenericParticleHandle(Test.ParticleHandles[1])->Sleeping());
 		EXPECT_FALSE(FConstGenericParticleHandle(Test.ParticleHandles[2])->Sleeping());
 		EXPECT_FALSE(FConstGenericParticleHandle(Test.ParticleHandles[3])->Sleeping());
-		EXPECT_EQ(Test.IslandManager->NumIslands(), 1);
+		EXPECT_EQ(Test.IslandManager->GetNumIslands(), 1);
 
 		// Flag the island for sleep
-		if (Test.IslandManager->NumIslands() == 1)
+		if (Test.IslandManager->GetNumIslands() == 1)
 		{
 			EXPECT_FALSE(Test.IslandManager->GetIsland(0)->IsSleeping());
-			Test.IslandManager->SleepIsland(Test.Particles, 0);
+			Test.IslandManager->SetParticleIslandIsSleeping(Test.ParticleHandles[0], true);
 		}
 
 		Test.Advance();
 
 		// Island and all particles should now be asleep
-		EXPECT_EQ(Test.IslandManager->NumIslands(), 1);
-		if (Test.IslandManager->NumIslands() == 1)
+		EXPECT_EQ(Test.IslandManager->GetNumIslands(), 1);
+		if (Test.IslandManager->GetNumIslands() == 1)
 		{
 			EXPECT_TRUE(Test.IslandManager->GetIsland(0)->IsSleeping());
 		}
@@ -655,8 +753,8 @@ namespace ChaosTest
 	// This behaviour is required for streaming to work since scene creation may
 	// be amortized over multiple frames and constraints may be made betweens
 	// sleeping particles in a later tick.
-	// In this case, A and B start in different sleeping island and get merged into
-	// a single still-sleeping island.
+	// In this case, A and B start sleeping and get merged into a single 
+	// still-sleeping island when we add the constraint.
 	//
 	//		(d=dynamic, s=sleeping, k=kinematic)
 	//		As  Bs    =>    As - Bs
@@ -670,7 +768,9 @@ namespace ChaosTest
 		Test.Evolution.SetParticleObjectState(Test.ParticleHandles[1], EObjectStateType::Sleeping);
 
 		Test.Advance();
-		EXPECT_EQ(Test.IslandManager->NumIslands(), 2);
+
+		// Particles without any constraints are not in the graph
+		EXPECT_EQ(Test.IslandManager->GetNumIslands(), 0);
 		EXPECT_TRUE(FConstGenericParticleHandle(Test.ParticleHandles[0])->Sleeping());
 		EXPECT_TRUE(FConstGenericParticleHandle(Test.ParticleHandles[1])->Sleeping());
 
@@ -679,8 +779,8 @@ namespace ChaosTest
 
 		Test.Advance();
 
-		// A and B still asleep
-		EXPECT_EQ(Test.IslandManager->NumIslands(), 1);
+		// A and B still asleep but now in an island with the constraint
+		EXPECT_EQ(Test.IslandManager->GetNumIslands(), 1);
 		EXPECT_TRUE(FConstGenericParticleHandle(Test.ParticleHandles[0])->Sleeping());
 		EXPECT_TRUE(FConstGenericParticleHandle(Test.ParticleHandles[1])->Sleeping());
 		EXPECT_TRUE(Test.ConstraintHandles[0]->IsSleeping());
@@ -698,13 +798,13 @@ namespace ChaosTest
 	{
 		FGraphEvolutionTest Test(3);
 
-		// Chains the particles and make the first one kinematix
+		// Chains the particles and make the first one kinematic
 		Test.Evolution.SetParticleObjectState(Test.ParticleHandles[0], EObjectStateType::Kinematic);
 		Test.MakeChain();
 
 		// Wait for sleep
 		Test.AdvanceUntilSleeping();
-		EXPECT_EQ(Test.IslandManager->NumIslands(), 1);
+		EXPECT_EQ(Test.IslandManager->GetNumIslands(), 1);
 		EXPECT_TRUE(FConstGenericParticleHandle(Test.ParticleHandles[1])->Sleeping());
 		EXPECT_TRUE(FConstGenericParticleHandle(Test.ParticleHandles[2])->Sleeping());
 		EXPECT_TRUE(Test.ConstraintHandles[0]->IsSleeping());
@@ -716,7 +816,7 @@ namespace ChaosTest
 		Test.Advance();
 
 		// All still asleep, including the new constraint
-		EXPECT_EQ(Test.IslandManager->NumIslands(), 1);
+		EXPECT_EQ(Test.IslandManager->GetNumIslands(), 1);
 		EXPECT_TRUE(FConstGenericParticleHandle(Test.ParticleHandles[1])->Sleeping());
 		EXPECT_TRUE(FConstGenericParticleHandle(Test.ParticleHandles[2])->Sleeping());
 		EXPECT_TRUE(Test.ConstraintHandles[0]->IsSleeping());
@@ -741,7 +841,7 @@ namespace ChaosTest
 
 		// Wait for sleep
 		Test.AdvanceUntilSleeping();
-		EXPECT_EQ(Test.IslandManager->NumIslands(), 2);
+		EXPECT_EQ(Test.IslandManager->GetNumIslands(), 2);
 		EXPECT_TRUE(FConstGenericParticleHandle(Test.ParticleHandles[0])->Sleeping());
 		EXPECT_TRUE(FConstGenericParticleHandle(Test.ParticleHandles[1])->Sleeping());
 		EXPECT_TRUE(FConstGenericParticleHandle(Test.ParticleHandles[2])->Sleeping());
@@ -755,7 +855,7 @@ namespace ChaosTest
 		Test.Advance();
 
 		// All still asleep, including the new constraint
-		EXPECT_EQ(Test.IslandManager->NumIslands(), 1);
+		EXPECT_EQ(Test.IslandManager->GetNumIslands(), 1);
 		EXPECT_TRUE(FConstGenericParticleHandle(Test.ParticleHandles[0])->Sleeping());
 		EXPECT_TRUE(FConstGenericParticleHandle(Test.ParticleHandles[1])->Sleeping());
 		EXPECT_TRUE(FConstGenericParticleHandle(Test.ParticleHandles[2])->Sleeping());
@@ -766,11 +866,16 @@ namespace ChaosTest
 	}
 
 	// Add constraints beween objects on the tick where their island goes to sleep, and one either side just to be sure.
+	// i.e., The SleepCounter does not get reset when we add a constraint between two particles.
 	//
 	//		(d=dynamic, s=sleeping, k=kinematic)
 	//		Ad   Bd  =>    As - Bs
 	//
-	GTEST_TEST(GraphEvolutionTests, TestConstraintGraph_SleepingSleepingConstraint_Timing)
+	// @todo(chaos): this test would fail because we do not transfer isolated particle
+	// sleep counts to the graph when we add a constraint to them. We could fix this
+	// but probably not worth worrying about
+	//
+	GTEST_TEST(GraphEvolutionTests, DISABLED_TestConstraintGraph_SleepingSleepingConstraint_Timing_Isolated)
 	{
 		// Count how many frames it takes the simulation to sleep
 		FGraphEvolutionTest SleepTest(2);
@@ -788,19 +893,17 @@ namespace ChaosTest
 			}
 			const bool bExpectSleep = (SleepRelativeTickCount >= 0);
 
-			// Should have 2 islands in a state that depends on whether we hit the sleep tick yet
-			EXPECT_EQ(Test.IslandManager->NumIslands(), 2);
+			// Should have no islands (because we have no constraints)
+			EXPECT_EQ(Test.IslandManager->GetNumIslands(), 0);
 			EXPECT_EQ(FConstGenericParticleHandle(Test.ParticleHandles[0])->Sleeping(), bExpectSleep);
 			EXPECT_EQ(FConstGenericParticleHandle(Test.ParticleHandles[1])->Sleeping(), bExpectSleep);
-			EXPECT_EQ(Test.IslandManager->GetIsland(0)->IsSleeping(), bExpectSleep);
-			EXPECT_EQ(Test.IslandManager->GetIsland(1)->IsSleeping(), bExpectSleep);
 
 			// Add a constraint A-B and tick
 			Test.ConstraintHandles.Add(Test.Constraints.AddConstraint({ Test.ParticleHandles[0], Test.ParticleHandles[1] }));
 			Test.Advance();
 
 			// Should now have 1 island and it should be asleep
-			EXPECT_EQ(Test.IslandManager->NumIslands(), 1);
+			EXPECT_EQ(Test.IslandManager->GetNumIslands(), 1);
 			EXPECT_TRUE(FConstGenericParticleHandle(Test.ParticleHandles[0])->Sleeping());
 			EXPECT_TRUE(FConstGenericParticleHandle(Test.ParticleHandles[1])->Sleeping());
 			EXPECT_TRUE(Test.IslandManager->GetIsland(0)->IsSleeping());
@@ -810,6 +913,69 @@ namespace ChaosTest
 		}
 	}
 
+	// Add constraints beween objects on the tick where their island goes to sleep, and one either side just to be sure.
+	// i.e., The SleepCounter does not get reset when we add a constraint between two particles.
+	//
+	//		(d=dynamic, s=sleeping, k=kinematic)
+	//		Ad - Bd   Cd - Dd  =>    As - Bs - Cs - Ds
+	//
+	// NOTE: this one works where DISABLED_TestConstraintGraph_SleepingSleepingConstraint_Timing_Isolated
+	// would fail because we retain SleepCounter when merging islands (but not when adding isloated
+	// particles that have their own sleep counter)
+	//
+	GTEST_TEST(GraphEvolutionTests, TestConstraintGraph_SleepingSleepingConstraint_Timing)
+	{
+		// Count how many frames it takes the simulation to sleep
+		int32 TicksToSleep = 0;
+		{
+			FGraphEvolutionTest SleepTest(4);
+			SleepTest.ConstraintHandles.Add(SleepTest.Constraints.AddConstraint({ SleepTest.ParticleHandles[0], SleepTest.ParticleHandles[1] }));
+			SleepTest.ConstraintHandles.Add(SleepTest.Constraints.AddConstraint({ SleepTest.ParticleHandles[2], SleepTest.ParticleHandles[3] }));
+
+			SleepTest.AdvanceUntilSleeping();
+			TicksToSleep = SleepTest.TickCount;
+		}
+
+		// Create a new simulation up to the sleep tick +/- a tick
+		// Verify that adding a constraint on that tick leaves the scene as expected
+		for (int32 SleepRelativeTickCount = -1; SleepRelativeTickCount < 2; ++SleepRelativeTickCount)
+		{
+			FGraphEvolutionTest Test(4);
+			Test.ConstraintHandles.Add(Test.Constraints.AddConstraint({ Test.ParticleHandles[0], Test.ParticleHandles[1] }));
+			Test.ConstraintHandles.Add(Test.Constraints.AddConstraint({ Test.ParticleHandles[2], Test.ParticleHandles[3] }));
+
+			for (int32 Frame = 0; Frame < TicksToSleep + SleepRelativeTickCount; ++Frame)
+			{
+				Test.Advance();
+			}
+			const bool bExpectSleep = (SleepRelativeTickCount >= 0);
+
+			// Should have no islands (because we have no constraints)
+			EXPECT_EQ(Test.IslandManager->GetNumIslands(), 2);
+			EXPECT_EQ(FConstGenericParticleHandle(Test.ParticleHandles[0])->Sleeping(), bExpectSleep);
+			EXPECT_EQ(FConstGenericParticleHandle(Test.ParticleHandles[1])->Sleeping(), bExpectSleep);
+			EXPECT_EQ(FConstGenericParticleHandle(Test.ParticleHandles[2])->Sleeping(), bExpectSleep);
+			EXPECT_EQ(FConstGenericParticleHandle(Test.ParticleHandles[3])->Sleeping(), bExpectSleep);
+
+			// Add a constraint B-C and tick
+			Test.ConstraintHandles.Add(Test.Constraints.AddConstraint({ Test.ParticleHandles[1], Test.ParticleHandles[2] }));
+			Test.Advance();
+
+			// Should now have 1 island and it should be asleep
+			EXPECT_EQ(Test.IslandManager->GetNumIslands(), 1);
+			EXPECT_TRUE(Test.IslandManager->GetIsland(0)->IsSleeping());
+			EXPECT_TRUE(FConstGenericParticleHandle(Test.ParticleHandles[0])->Sleeping());
+			EXPECT_TRUE(FConstGenericParticleHandle(Test.ParticleHandles[1])->Sleeping());
+			EXPECT_TRUE(FConstGenericParticleHandle(Test.ParticleHandles[2])->Sleeping());
+			EXPECT_TRUE(FConstGenericParticleHandle(Test.ParticleHandles[3])->Sleeping());
+
+			// Constraints should also be asleep
+			EXPECT_TRUE(Test.ConstraintHandles[0]->IsSleeping());
+			EXPECT_TRUE(Test.ConstraintHandles[1]->IsSleeping());
+			EXPECT_TRUE(Test.ConstraintHandles[2]->IsSleeping());
+		}
+	}
+
 	// Test an edge case bug that is probably easy to accidentally reintroduce. This would leave
 	// a dangling pointer in the constraint graph due to a collision constraint being deleted while
 	// in a sleeping island.
@@ -844,20 +1010,19 @@ namespace ChaosTest
 
 		Test.Advance();
 
-		EXPECT_EQ(Test.IslandManager->NumIslands(), 1);
-		EXPECT_FALSE(Test.IslandManager->GetIsland(0)->IsSleeping());
+		EXPECT_EQ(Test.IslandManager->GetNumIslands(), 0);
 
-		// Explicitly put B to sleep (its island will still be considered awake until the tick)
+		// Explicitly put B to sleep
 		Test.Evolution.SetParticleObjectState(Test.ParticleHandles[1], EObjectStateType::Sleeping);
 
-		// Add a constraint A-B. B is asleep, but its island is awake
+		// Add a constraint A-B. B is asleep
 		Test.ConstraintHandles.Add(Test.Constraints.AddConstraint({ Test.ParticleHandles[0], Test.ParticleHandles[1] }));
 
 		Test.Advance();
 
 		// Everything should be asleep
 		// The bug was that the constraint was still flagged as awake, but in a sleeping island.
-		EXPECT_EQ(Test.IslandManager->NumIslands(), 1);
+		EXPECT_EQ(Test.IslandManager->GetNumIslands(), 1);
 		EXPECT_TRUE(Test.IslandManager->GetIsland(0)->IsSleeping());
 		EXPECT_TRUE(FConstGenericParticleHandle(Test.ParticleHandles[1])->Sleeping());
 		EXPECT_TRUE(Test.ConstraintHandles[0]->IsSleeping());
@@ -885,7 +1050,7 @@ namespace ChaosTest
 
 		Test.Advance();
 
-		EXPECT_EQ(Test.IslandManager->NumIslands(), 1);
+		EXPECT_EQ(Test.IslandManager->GetNumIslands(), 1);
 		EXPECT_FALSE(Test.IslandManager->GetIsland(0)->IsSleeping());
 
 		// Explicitly put both particles (and therefore their island) to sleep
@@ -896,14 +1061,14 @@ namespace ChaosTest
 
 		// Everything should be asleep
 		// The bug was that the constraint was still flagged as awake, but in a sleeping island.
-		EXPECT_EQ(Test.IslandManager->NumIslands(), 1);
+		EXPECT_EQ(Test.IslandManager->GetNumIslands(), 1);
 		EXPECT_TRUE(Test.IslandManager->GetIsland(0)->IsSleeping());
 		EXPECT_TRUE(FConstGenericParticleHandle(Test.ParticleHandles[0])->Sleeping());
 		EXPECT_TRUE(FConstGenericParticleHandle(Test.ParticleHandles[1])->Sleeping());
 		EXPECT_TRUE(Test.ConstraintHandles[0]->IsSleeping());
 	}
 
-	// Test isolated kinematic particles are removed from the island manager if they have no constraints
+	// Test isolated particles are not present in the graph
 	GTEST_TEST(GraphEvolutionTests, TestConstraintGraph_KinematicRemoveFromGraph)
 	{
 		// Create a scene with 3 dynamic particles
@@ -911,39 +1076,22 @@ namespace ChaosTest
 
 		Test.Advance();
 
-		// Initially all particles are dynamic and should be in the graph
-		EXPECT_EQ(Test.IslandManager->NumIslands(), 3);
-		EXPECT_EQ(Test.IslandManager->GetIslandGraph()->ItemNodes.Num(), 3);
-		EXPECT_TRUE(Test.ParticleHandles[0]->IsInConstraintGraph());	// Dynamic
-		EXPECT_TRUE(Test.ParticleHandles[1]->IsInConstraintGraph());	// Dynamic
-		EXPECT_TRUE(Test.ParticleHandles[2]->IsInConstraintGraph());	// Dynamic
-		EXPECT_TRUE(Test.IslandManager->GetIslandGraph()->GraphNodes.IsValidIndex(0));
-		EXPECT_TRUE(Test.IslandManager->GetIslandGraph()->GraphNodes.IsValidIndex(1));
-		EXPECT_TRUE(Test.IslandManager->GetIslandGraph()->GraphNodes.IsValidIndex(2));
-
-		// Change a particle to kinematic. Since it has no collisions or joints it should
-		// be removed from the constraint graph at the next update.
+		EXPECT_EQ(Test.IslandManager->GetNumIslands(), 0);
+		EXPECT_EQ(Test.IslandManager->GetNumParticles(), 0);
+
+		// Change a particle to kinematic
 		Test.Evolution.SetParticleObjectState(Test.ParticleHandles[1], EObjectStateType::Kinematic);
 
 		Test.Advance();
 
-		// Should now have only 2 nodes in the graph, each in their own island
-		EXPECT_EQ(Test.IslandManager->NumIslands(), 2);
-		EXPECT_EQ(Test.IslandManager->GetIslandGraph()->ItemNodes.Num(), 2);
-
-		// Kinematic was removed from the graph, but dynamics are still there
-		EXPECT_TRUE(Test.ParticleHandles[0]->IsInConstraintGraph());	// Dynamic
-		EXPECT_FALSE(Test.ParticleHandles[1]->IsInConstraintGraph());	// Kinematic
-		EXPECT_TRUE(Test.ParticleHandles[2]->IsInConstraintGraph());	// Dynamic
+		EXPECT_EQ(Test.IslandManager->GetNumIslands(), 0);
+		EXPECT_EQ(Test.IslandManager->GetNumParticles(), 0);
 
 		Test.Advance();
 
 		// State should not have changed with a second tick
-		EXPECT_EQ(Test.IslandManager->NumIslands(), 2);
-		EXPECT_EQ(Test.IslandManager->GetIslandGraph()->ItemNodes.Num(), 2);
-		EXPECT_TRUE(Test.ParticleHandles[0]->IsInConstraintGraph());	// Dynamic
-		EXPECT_FALSE(Test.ParticleHandles[1]->IsInConstraintGraph());	// Kinematic
-		EXPECT_TRUE(Test.ParticleHandles[2]->IsInConstraintGraph());	// Dynamic
+		EXPECT_EQ(Test.IslandManager->GetNumIslands(), 0);
+		EXPECT_EQ(Test.IslandManager->GetNumParticles(), 0);
 	}
 
 	// Test the conditions for a kinematic particle waking an island
@@ -965,10 +1113,10 @@ namespace ChaosTest
 		Test.Advance();
 
 		// Expect one sleeping island
-		EXPECT_EQ(Test.IslandManager->NumIslands(), 1);
-		EXPECT_NE(Test.ConstraintHandles[0]->GetConstraintGraphIndex(), INDEX_NONE);
-		EXPECT_NE(Test.ConstraintHandles[1]->GetConstraintGraphIndex(), INDEX_NONE);
-		EXPECT_NE(Test.ConstraintHandles[2]->GetConstraintGraphIndex(), INDEX_NONE);
+		EXPECT_EQ(Test.IslandManager->GetNumIslands(), 1);
+		EXPECT_NE(Test.ConstraintHandles[0]->GetConstraintGraphEdge(), nullptr);
+		EXPECT_NE(Test.ConstraintHandles[1]->GetConstraintGraphEdge(), nullptr);
+		EXPECT_NE(Test.ConstraintHandles[2]->GetConstraintGraphEdge(), nullptr);
 		EXPECT_TRUE(Test.IslandManager->GetIsland(0)->IsSleeping());
 
 		// Set to velocity mode and animate
@@ -979,7 +1127,7 @@ namespace ChaosTest
 		Test.Advance();
 
 		// Expect one sleeping island as the velocity of the kinematic particle is still zero
-		EXPECT_EQ(Test.IslandManager->NumIslands(), 1);
+		EXPECT_EQ(Test.IslandManager->GetNumIslands(), 1);
 		EXPECT_TRUE(Test.IslandManager->GetIsland(0)->IsSleeping());
 
 		KinematicParticle->SetV(FVec3(10.0f, 0.0f, 0.0f));
@@ -987,7 +1135,7 @@ namespace ChaosTest
 		Test.Advance();
 
 		// Expect one awake island
-		EXPECT_EQ(Test.IslandManager->NumIslands(), 1);
+		EXPECT_EQ(Test.IslandManager->GetNumIslands(), 1);
 		EXPECT_FALSE(Test.IslandManager->GetIsland(0)->IsSleeping());
 
 		// Put all particles back to sleep and now set angular velocity
@@ -1000,7 +1148,7 @@ namespace ChaosTest
 		Test.Advance();
 
 		// Check we've put the island back to sleep
-		EXPECT_EQ(Test.IslandManager->NumIslands(), 1);
+		EXPECT_EQ(Test.IslandManager->GetNumIslands(), 1);
 		EXPECT_TRUE(Test.IslandManager->GetIsland(0)->IsSleeping());
 
 		// Now set angular velocity. Island should wake
@@ -1008,7 +1156,7 @@ namespace ChaosTest
 
 		Test.Advance();
 
-		EXPECT_EQ(Test.IslandManager->NumIslands(), 1);
+		EXPECT_EQ(Test.IslandManager->GetNumIslands(), 1);
 		EXPECT_FALSE(Test.IslandManager->GetIsland(0)->IsSleeping());
 
 		// Put all particles back to sleep
@@ -1021,7 +1169,7 @@ namespace ChaosTest
 		Test.Advance();
 
 		// Check we've put the island back to sleep
-		EXPECT_EQ(Test.IslandManager->NumIslands(), 1);
+		EXPECT_EQ(Test.IslandManager->GetNumIslands(), 1);
 		EXPECT_TRUE(Test.IslandManager->GetIsland(0)->IsSleeping());
 
 		// Now set to position mode. Initially the island should stay sleeping as the target
@@ -1032,7 +1180,7 @@ namespace ChaosTest
 
 		Test.Advance();
 
-		EXPECT_EQ(Test.IslandManager->NumIslands(), 1);
+		EXPECT_EQ(Test.IslandManager->GetNumIslands(), 1);
 		EXPECT_TRUE(Test.IslandManager->GetIsland(0)->IsSleeping());
 
 		// Set a non-zero position target. Should cause the island to wake
@@ -1041,7 +1189,7 @@ namespace ChaosTest
 
 		Test.Advance();
 
-		EXPECT_EQ(Test.IslandManager->NumIslands(), 1);
+		EXPECT_EQ(Test.IslandManager->GetNumIslands(), 1);
 		EXPECT_FALSE(Test.IslandManager->GetIsland(0)->IsSleeping());
 
 		// Put all particles back to sleep
@@ -1055,7 +1203,7 @@ namespace ChaosTest
 		Test.Advance();
 
 		// Check we've put the island back to sleep
-		EXPECT_EQ(Test.IslandManager->NumIslands(), 1);
+		EXPECT_EQ(Test.IslandManager->GetNumIslands(), 1);
 		EXPECT_TRUE(Test.IslandManager->GetIsland(0)->IsSleeping());
 
 		// Set a non-identity rotation target. Should cause the island to wake
@@ -1065,7 +1213,7 @@ namespace ChaosTest
 
 		Test.Advance();
 
-		EXPECT_EQ(Test.IslandManager->NumIslands(), 1);
+		EXPECT_EQ(Test.IslandManager->GetNumIslands(), 1);
 		EXPECT_FALSE(Test.IslandManager->GetIsland(0)->IsSleeping());
 
 	}
diff --git a/Engine/Source/Programs/HeadlessChaos/Private/HeadlessChaosTestJoints.cpp b/Engine/Source/Programs/HeadlessChaos/Private/HeadlessChaosTestJoints.cpp
--- a/Engine/Source/Programs/HeadlessChaos/Private/HeadlessChaosTestJoints.cpp
+++ b/Engine/Source/Programs/HeadlessChaos/Private/HeadlessChaosTestJoints.cpp
@@ -527,6 +527,7 @@ namespace ChaosTest {
 		PhysicalMaterial->SleepingAngularThreshold = 0;
 		PhysicalMaterial->DisabledLinearThreshold = 0;
 		PhysicalMaterial->DisabledAngularThreshold = 0;
+		PhysicalMaterial->SleepCounterThreshold = 20;
 
 		{
 			FParticleUniqueIndicesMultithreaded UniqueIndices;
diff --git a/Engine/Source/Programs/HeadlessChaos/Private/HeadlessChaosTestSim.cpp b/Engine/Source/Programs/HeadlessChaos/Private/HeadlessChaosTestSim.cpp
--- a/Engine/Source/Programs/HeadlessChaos/Private/HeadlessChaosTestSim.cpp
+++ b/Engine/Source/Programs/HeadlessChaos/Private/HeadlessChaosTestSim.cpp
@@ -319,6 +319,7 @@ namespace ChaosTest {
 		Module->DestroySolver(Solver);
 	}
 
+	// Test that isolated stationary particles without gravity sleep according to their sleep settings.
 	GTEST_TEST(AllEvolutions, SimTests_SleepAndWakeSimTest3)
 	{
 		FParticleUniqueIndicesMultithreaded UniqueIndices;
diff --git a/Engine/Source/Runtime/Experimental/Chaos/Private/Chaos/ChaosDebugDraw.cpp b/Engine/Source/Runtime/Experimental/Chaos/Private/Chaos/ChaosDebugDraw.cpp
--- a/Engine/Source/Runtime/Experimental/Chaos/Private/Chaos/ChaosDebugDraw.cpp
+++ b/Engine/Source/Runtime/Experimental/Chaos/Private/Chaos/ChaosDebugDraw.cpp
@@ -19,8 +19,6 @@
 #include "Chaos/Particle/ParticleUtilities.h"
 #include "Chaos/ParticleHandle.h"
 #include "Chaos/PBDCollisionConstraints.h"
-#include "Chaos/PBDConstraintColor.h"
-#include "Chaos/PBDConstraintGraph.h"
 #include "Chaos/PBDJointConstraints.h"
 #include "Chaos/PBDRigidClustering.h"
 #include "Chaos/PBDRigidParticles.h"
@@ -196,19 +194,35 @@ namespace Chaos
 
 		FColor GetIslandColor(const int32 IslandIndex, const bool bIsAwake)
 		{
-			static FColor AwakeColors[] =
+			static const FColor AwakeColors[] =
 			{
 				FColor::Red,
 				FColor::Orange,
 				FColor::Yellow,
 				FColor::Green,
+				FColor::Emerald,
+				FColor::Cyan,
+				FColor::Turquoise,
 				FColor::Blue,
 				FColor::Magenta,
+				FColor::Purple,
 			};
 			const int32 NumAwakeColors = UE_ARRAY_COUNT(AwakeColors);
-			static FColor SleepingColor = FColor::Black;
 
-			return bIsAwake ? AwakeColors[IslandIndex % NumAwakeColors] : SleepingColor;
+			static const FColor SleepingColor = FColor::Black;
+			static const FColor NullColor = FColor::White;
+
+			if (IslandIndex == INDEX_NONE)
+			{
+				return NullColor;
+			}
+
+			if (!bIsAwake)
+			{
+				return SleepingColor;
+			}
+
+			return AwakeColors[IslandIndex % NumAwakeColors];
 		};
 
 
@@ -617,9 +631,10 @@ namespace Chaos
 				{
 					ShapeColor = Settings.ShapesColorsPerShapeType.GetColorFromShapeType(InnerType);
 				}
-				if (bChaosDebugDebugDrawColorShapesByIsland && (FConstGenericParticleHandle(Particle)->IslandIndex() != INDEX_NONE))
+				if (bChaosDebugDebugDrawColorShapesByIsland)
 				{
-					ShapeColor = GetIslandColor(FConstGenericParticleHandle(Particle)->IslandIndex(), true);
+					const int32 IslandId = (Particle->GetConstraintGraphNode() != nullptr) ? Particle->GetConstraintGraphNode()->GetIslandId() : INDEX_NONE;
+					ShapeColor = GetIslandColor(IslandId, true);
 				}
 				if (bChaosDebugDebugDrawColorShapesByInternalCluster)
 				{
@@ -1343,33 +1358,41 @@ namespace Chaos
 
 				if (bChaosDebugDebugDrawContactGraph)
 				{
-					FDebugDrawQueue::GetInstance().DrawDebugLine(Transform0.GetLocation(), ContactPos, FColor::Red, false, UE_KINDA_SMALL_NUMBER, Settings.DrawPriority, Settings.LineThickness);
-					FDebugDrawQueue::GetInstance().DrawDebugLine(Transform1.GetLocation(), ContactPos, FColor::White, false, UE_KINDA_SMALL_NUMBER, Settings.DrawPriority, Settings.LineThickness);
-					FDebugDrawQueue::GetInstance().DrawDebugString(ContactPos, FString::Format(TEXT("{0}-{1}-{2}"), { LevelIndex, ColorIndex, OrderIndex }), nullptr, FColor::Yellow, UE_KINDA_SMALL_NUMBER, false, Settings.FontScale);
+					if (FConstGenericParticleHandle(Constraint->GetParticle0())->IsDynamic())
+					{
+						FDebugDrawQueue::GetInstance().DrawDebugLine(Transform0.GetLocation(), ContactPos, FColor::Red, false, UE_KINDA_SMALL_NUMBER, Settings.DrawPriority, Settings.LineThickness);
+					}
+					if (FConstGenericParticleHandle(Constraint->GetParticle1())->IsDynamic())
+					{
+						FDebugDrawQueue::GetInstance().DrawDebugLine(Transform1.GetLocation(), ContactPos, FColor::White, false, UE_KINDA_SMALL_NUMBER, Settings.DrawPriority, Settings.LineThickness);
+					}
+					if (ColorIndex >= 0)
+					{
+						FDebugDrawQueue::GetInstance().DrawDebugString(ContactPos, FString::Format(TEXT("{0} L{1} C{2}"), { OrderIndex, LevelIndex, ColorIndex }), nullptr, FColor::Yellow, UE_KINDA_SMALL_NUMBER, false, Settings.FontScale);
+					}
+					else if (LevelIndex >= 0)
+					{
+						FDebugDrawQueue::GetInstance().DrawDebugString(ContactPos, FString::Format(TEXT("{0} L{1}"), { OrderIndex, LevelIndex }), nullptr, FColor::Yellow, UE_KINDA_SMALL_NUMBER, false, Settings.FontScale);
+					}
+					else
+					{
+						FDebugDrawQueue::GetInstance().DrawDebugString(ContactPos, FString::Format(TEXT("{0}"), { OrderIndex }), nullptr, FColor::Yellow, UE_KINDA_SMALL_NUMBER, false, Settings.FontScale);
+					}
 				}
 			};
 
 			if (bChaosDebugDebugDrawIslands)
 			{
-				TArray<FAABB3> IslandBounds;
-				IslandBounds.SetNum(Graph.NumIslands());
-
-				const typename Private::FPBDIslandManager::GraphType* IslandGraph = Graph.GetIslandGraph();
-				for (const auto& GraphNode : IslandGraph->GraphNodes)
+				for (int32 IslandIndex = 0; IslandIndex < Graph.GetNumIslands(); ++IslandIndex)
 				{
-					FConstGenericParticleHandle Particle = GraphNode.NodeItem;
-					if (Particle->IsDynamic() && Particle->HasBounds())
+					const Private::FPBDIsland* Island = Graph.GetIsland(IslandIndex);
+					FAABB3 IslandAABB = FAABB3::EmptyAABB();
+
+					for (int32 NodeIndex = 0; NodeIndex < Island->GetNumParticles(); ++NodeIndex)
 					{
-						for (int32 ParticleIsland : Graph.FindParticleIslands(Particle->Handle()))
-						{
-							IslandBounds[ParticleIsland].GrowToInclude(Particle->BoundingBox());
-						}
+						FConstGenericParticleHandle Particle = Island->GetNode(NodeIndex)->GetParticle();
+						IslandAABB.GrowToInclude(Particle->BoundingBox());
 					}
-				}
-			
-				for (int32 IslandIndex = 0; IslandIndex < IslandBounds.Num(); ++IslandIndex)
-				{
-					FAABB3 IslandAABB = IslandBounds[IslandIndex];
 
 					const FColor IslandColor = GetIslandColor(IslandIndex, !Graph.GetIsland(IslandIndex)->IsSleeping());
 					const FAABB3 Bounds = IslandAABB.TransformedAABB(SpaceTransform);
@@ -1379,20 +1402,22 @@ namespace Chaos
 
 			if (bChaosDebugDebugDrawContactGraph || bChaosDebugDebugDrawContactGraphUnused || bChaosDebugDebugDrawContactGraphUsed)
 			{
-				const typename Private::FPBDIslandManager::GraphType* IslandGraph = Graph.GetIslandGraph();
-				for (const auto& GraphEdge : IslandGraph->GraphEdges)
+				for (int32 ContainerIndex = 0; ContainerIndex < Graph.GetNumConstraintContainers(); ++ContainerIndex)
 				{
-					const FConstraintHandle* Constraint = GraphEdge.EdgeItem;
-					if (const FPBDCollisionConstraintHandle* Collision = Constraint->As<FPBDCollisionConstraintHandle>())
-					{
-						const int32 IslandIndex = GraphEdge.IslandIndex;
-						// @chaos(todo): would be nice to have this data retained for debug draw...
-						const int32 LevelIndex = INDEX_NONE;
-						const int32 ColorIndex = INDEX_NONE;
-						const int32 OrderIndex = INDEX_NONE;
-						const bool bIsUsed = !Collision->GetConstraint()->AccumulatedImpulse.IsNearlyZero();
-						DrawGraphCollision(SpaceTransform, Collision->GetConstraint(), IslandIndex, LevelIndex, ColorIndex, OrderIndex, bIsUsed, Settings);
-					}
+					Graph.VisitConstConstraints(ContainerIndex,
+						[&Graph, &SpaceTransform, &Settings, &DrawGraphCollision](const Private::FPBDIslandConstraint* IslandConstraint)
+						{
+							const FConstraintHandle* Constraint = IslandConstraint->GetConstraint();
+							if (const FPBDCollisionConstraintHandle* Collision = Constraint->As<FPBDCollisionConstraintHandle>())
+							{
+								const int32 IslandIndex = Graph.GetIslandIndex(Graph.GetConstraintIsland(Collision));
+								const int32 LevelIndex = Graph.GetConstraintLevel(IslandConstraint);
+								const int32 ColorIndex = Graph.GetConstraintColor(IslandConstraint);
+								const int32 OrderIndex = Graph.GetIslandArrayIndex(IslandConstraint);
+								const bool bIsUsed = !Collision->GetConstraint()->AccumulatedImpulse.IsNearlyZero();
+								DrawGraphCollision(SpaceTransform, Collision->GetConstraint(), IslandIndex, LevelIndex, ColorIndex, OrderIndex, bIsUsed, Settings);
+							}
+						});
 				}
 			}
 		}
diff --git a/Engine/Source/Runtime/Experimental/Chaos/Private/Chaos/Character/CharacterGroundConstraintContainer.cpp b/Engine/Source/Runtime/Experimental/Chaos/Private/Chaos/Character/CharacterGroundConstraintContainer.cpp
--- a/Engine/Source/Runtime/Experimental/Chaos/Private/Chaos/Character/CharacterGroundConstraintContainer.cpp
+++ b/Engine/Source/Runtime/Experimental/Chaos/Private/Chaos/Character/CharacterGroundConstraintContainer.cpp
@@ -105,7 +105,7 @@ namespace Chaos
 			}
 			else if (bIsInGraph && !bShouldBeInGraph)
 			{
-				IslandManager.RemoveConstraint(ContainerId, Constraint);
+				IslandManager.RemoveConstraint(Constraint);
 			}
 		}
 	}
diff --git a/Engine/Source/Runtime/Experimental/Chaos/Private/Chaos/Character/CharacterGroundConstraintContainerSolver.cpp b/Engine/Source/Runtime/Experimental/Chaos/Private/Chaos/Character/CharacterGroundConstraintContainerSolver.cpp
--- a/Engine/Source/Runtime/Experimental/Chaos/Private/Chaos/Character/CharacterGroundConstraintContainerSolver.cpp
+++ b/Engine/Source/Runtime/Experimental/Chaos/Private/Chaos/Character/CharacterGroundConstraintContainerSolver.cpp
@@ -1,6 +1,7 @@
 // Copyright Epic Games, Inc. All Rights Reserved.
 
 #include "CharacterGroundConstraintContainerSolver.h"
+#include "Chaos/Island/IslandManager.h"
 
 namespace Chaos
 {
@@ -34,12 +35,12 @@ namespace Chaos
 			}
 		}
 
-		void FCharacterGroundConstraintContainerSolver::AddConstraints(const TArrayView<Private::FPBDIslandConstraint>& IslandConstraints)
+		void FCharacterGroundConstraintContainerSolver::AddConstraints(const TArrayView<Private::FPBDIslandConstraint*>& IslandConstraints)
 		{
-			for (Private::FPBDIslandConstraint& IslandConstraint : IslandConstraints)
+			for (Private::FPBDIslandConstraint* IslandConstraint : IslandConstraints)
 			{
 				// We will only ever be given constraints from our container (asserts in non-shipping)
-				FCharacterGroundConstraintHandle* Constraint = IslandConstraint.GetConstraint()->AsUnsafe<FCharacterGroundConstraintHandle>();
+				FCharacterGroundConstraintHandle* Constraint = IslandConstraint->GetConstraint()->AsUnsafe<FCharacterGroundConstraintHandle>();
 
 				AddConstraint(Constraint);
 			}
diff --git a/Engine/Source/Runtime/Experimental/Chaos/Private/Chaos/Character/CharacterGroundConstraintContainerSolver.h b/Engine/Source/Runtime/Experimental/Chaos/Private/Chaos/Character/CharacterGroundConstraintContainerSolver.h
--- a/Engine/Source/Runtime/Experimental/Chaos/Private/Chaos/Character/CharacterGroundConstraintContainerSolver.h
+++ b/Engine/Source/Runtime/Experimental/Chaos/Private/Chaos/Character/CharacterGroundConstraintContainerSolver.h
@@ -26,7 +26,7 @@ namespace Chaos
 			virtual void AddConstraints() override final;
 
 			/// Add all constraints from an island (guaranteed to be from the same constraint container)
-			virtual void AddConstraints(const TArrayView<Private::FPBDIslandConstraint>& IslandConstraints) override final;
+			virtual void AddConstraints(const TArrayView<Private::FPBDIslandConstraint*>& IslandConstraints) override final;
 
 			/// Set the bodies on the constraint from the constraint data solver
 			virtual void AddBodies(FSolverBodyContainer& SolverBodyContainer) override final;
diff --git a/Engine/Source/Runtime/Experimental/Chaos/Private/Chaos/Collision/PBDCollisionConstraint.cpp b/Engine/Source/Runtime/Experimental/Chaos/Private/Chaos/Collision/PBDCollisionConstraint.cpp
--- a/Engine/Source/Runtime/Experimental/Chaos/Private/Chaos/Collision/PBDCollisionConstraint.cpp
+++ b/Engine/Source/Runtime/Experimental/Chaos/Private/Chaos/Collision/PBDCollisionConstraint.cpp
@@ -154,7 +154,7 @@ namespace Chaos
 		Constraint.GetContainerCookie().ClearContainerData();
 
 		// We are not in the constraint graph, even if Source was
-		Constraint.SetConstraintGraphIndex(INDEX_NONE);
+		Constraint.SetConstraintGraphEdge(nullptr);
 
 		return Constraint;
 	}
@@ -165,14 +165,14 @@ namespace Chaos
 		// We do not want to overwrite these properties because they are managed by the systems
 		// where the constraint is registsred (allocator, container, graph).
 		const FPBDCollisionConstraintContainerCookie Cookie = ContainerCookie;
-		const int32 ConstraintGraphIndex = GetConstraintGraphIndex();
+		Private::FPBDIslandConstraint* ConstraintGraphEdge = GetConstraintGraphEdge();
 
 		// Copy everything
 		*this = Source;
 
 		// Restore the system state
 		ContainerCookie = Cookie;
-		SetConstraintGraphIndex(ConstraintGraphIndex);
+		SetConstraintGraphEdge(ConstraintGraphEdge);
 	}
 
 	FPBDCollisionConstraint::FPBDCollisionConstraint()
@@ -259,7 +259,7 @@ namespace Chaos
 	{
 #if !UE_BUILD_TEST && !UE_BUILD_SHIPPING
 		// Make sure we have been dropped by the graph
-		ensure(GetConstraintGraphIndex() == INDEX_NONE);
+		ensure(GetConstraintGraphEdge() == nullptr);
 #endif
 	}
 
diff --git a/Engine/Source/Runtime/Experimental/Chaos/Private/Chaos/Collision/PBDCollisionContainerSolver.cpp b/Engine/Source/Runtime/Experimental/Chaos/Private/Chaos/Collision/PBDCollisionContainerSolver.cpp
--- a/Engine/Source/Runtime/Experimental/Chaos/Private/Chaos/Collision/PBDCollisionContainerSolver.cpp
+++ b/Engine/Source/Runtime/Experimental/Chaos/Private/Chaos/Collision/PBDCollisionContainerSolver.cpp
@@ -342,12 +342,12 @@ namespace Chaos
 		}
 	}
 
-	void FPBDCollisionContainerSolver::AddConstraints(const TArrayView<Private::FPBDIslandConstraint>& IslandConstraints)
+	void FPBDCollisionContainerSolver::AddConstraints(const TArrayView<Private::FPBDIslandConstraint*>& IslandConstraints)
 	{
-		for (Private::FPBDIslandConstraint& IslandConstraint : IslandConstraints)
+		for (Private::FPBDIslandConstraint* IslandConstraint : IslandConstraints)
 		{
 			// We will only ever be given constraints from our container (asserts in non-shipping)
-			FPBDCollisionConstraint& Constraint = IslandConstraint.GetConstraint()->AsUnsafe<FPBDCollisionConstraintHandle>()->GetContact();
+			FPBDCollisionConstraint& Constraint = IslandConstraint->GetConstraint()->AsUnsafe<FPBDCollisionConstraintHandle>()->GetContact();
 
 			AddConstraint(Constraint);
 		}
diff --git a/Engine/Source/Runtime/Experimental/Chaos/Private/Chaos/Collision/PBDCollisionContainerSolverJacobi.cpp b/Engine/Source/Runtime/Experimental/Chaos/Private/Chaos/Collision/PBDCollisionContainerSolverJacobi.cpp
--- a/Engine/Source/Runtime/Experimental/Chaos/Private/Chaos/Collision/PBDCollisionContainerSolverJacobi.cpp
+++ b/Engine/Source/Runtime/Experimental/Chaos/Private/Chaos/Collision/PBDCollisionContainerSolverJacobi.cpp
@@ -304,12 +304,12 @@ namespace Chaos
 			}
 		}
 
-		void FPBDCollisionContainerSolverJacobi::AddConstraints(const TArrayView<Private::FPBDIslandConstraint>& IslandConstraints)
+		void FPBDCollisionContainerSolverJacobi::AddConstraints(const TArrayView<Private::FPBDIslandConstraint*>& IslandConstraints)
 		{
-			for (Private::FPBDIslandConstraint& IslandConstraint : IslandConstraints)
+			for (Private::FPBDIslandConstraint* IslandConstraint : IslandConstraints)
 			{
 				// We will only ever be given constraints from our container (asserts in non-shipping)
-				FPBDCollisionConstraint& Constraint = IslandConstraint.GetConstraint()->AsUnsafe<FPBDCollisionConstraintHandle>()->GetContact();
+				FPBDCollisionConstraint& Constraint = IslandConstraint->GetConstraint()->AsUnsafe<FPBDCollisionConstraintHandle>()->GetContact();
 
 				AddConstraint(Constraint);
 			}
diff --git a/Engine/Source/Runtime/Experimental/Chaos/Private/Chaos/Collision/PBDCollisionContainerSolverSimd.cpp b/Engine/Source/Runtime/Experimental/Chaos/Private/Chaos/Collision/PBDCollisionContainerSolverSimd.cpp
--- a/Engine/Source/Runtime/Experimental/Chaos/Private/Chaos/Collision/PBDCollisionContainerSolverSimd.cpp
+++ b/Engine/Source/Runtime/Experimental/Chaos/Private/Chaos/Collision/PBDCollisionContainerSolverSimd.cpp
@@ -348,7 +348,7 @@ namespace Chaos
 			check(false);
 		}
 
-		void FPBDCollisionContainerSolverSimd::AddConstraints(const TArrayView<Private::FPBDIslandConstraint>& IslandConstraints)
+		void FPBDCollisionContainerSolverSimd::AddConstraints(const TArrayView<Private::FPBDIslandConstraint*>& IslandConstraints)
 		{
 			// Decide what lane this island goes into: Find the lane with the least constraints in it
 			int32 IslandLaneIndex = 0;
@@ -369,10 +369,10 @@ namespace Chaos
 			}
 
 			// Add all the constraints in the island to the selected lane
-			for (Private::FPBDIslandConstraint& IslandConstraint : IslandConstraints)
+			for (Private::FPBDIslandConstraint* IslandConstraint : IslandConstraints)
 			{
 				// NOTE: We will only ever be given constraints from our container (asserts in non-shipping)
-				FPBDCollisionConstraint& Constraint = IslandConstraint.GetConstraint()->AsUnsafe<FPBDCollisionConstraintHandle>()->GetContact();
+				FPBDCollisionConstraint& Constraint = IslandConstraint->GetConstraint()->AsUnsafe<FPBDCollisionConstraintHandle>()->GetContact();
 
 				const int32 SolverIndex = SimdData.SimdNumConstraints.GetValue(IslandLaneIndex);
 				SimdData.SimdNumConstraints.SetValue(IslandLaneIndex, SolverIndex + 1);
diff --git a/Engine/Source/Runtime/Experimental/Chaos/Private/Chaos/Collision/ParticlePairMidPhase.cpp b/Engine/Source/Runtime/Experimental/Chaos/Private/Chaos/Collision/ParticlePairMidPhase.cpp
--- a/Engine/Source/Runtime/Experimental/Chaos/Private/Chaos/Collision/ParticlePairMidPhase.cpp
+++ b/Engine/Source/Runtime/Experimental/Chaos/Private/Chaos/Collision/ParticlePairMidPhase.cpp
@@ -527,11 +527,6 @@ namespace Chaos
 	{
 		if (Constraint.IsValid() && IsUsedSince(SleepEpoch))
 		{
-#if CHAOS_CONSTRAINTHANDLE_DEBUG_ENABLED
-			// If a particle changed to kinematic, all constraints without a dynamic should have been removed
-			CheckConstraintContainsDynamic(Constraint.Get());
-#endif
-
 			// We need to refresh the epoch so that the constraint state will be used as the previous
 			// state if the pair is still colliding in the next tick. However, we are not in the active 
 			// collisions array so reset the indices.
@@ -1784,7 +1779,7 @@ namespace Chaos
 		bool bInGraph = false;
 		VisitConstCollisions([&bInGraph](const FPBDCollisionConstraint& Constraint)
 		{
-			if (Constraint.GetConstraintGraphIndex() != INDEX_NONE)
+			if (Constraint.GetConstraintGraphEdge() != nullptr)
 			{
 				bInGraph = true;
 				return ECollisionVisitorResult::Stop;
diff --git a/Engine/Source/Runtime/Experimental/Chaos/Private/Chaos/Island/IslandGroup.cpp b/Engine/Source/Runtime/Experimental/Chaos/Private/Chaos/Island/IslandGroup.cpp
--- a/Engine/Source/Runtime/Experimental/Chaos/Private/Chaos/Island/IslandGroup.cpp
+++ b/Engine/Source/Runtime/Experimental/Chaos/Private/Chaos/Island/IslandGroup.cpp
@@ -3,7 +3,6 @@
 
 #include "Chaos/ConstraintHandle.h"
 #include "Chaos/Island/IslandManager.h"
-#include "Chaos/Island/SolverIsland.h"
 #include "Chaos/PBDConstraintContainer.h"
 
 namespace Chaos
@@ -48,7 +47,7 @@ namespace Chaos
 
 				for (int32 ContainerIndex = 0; ContainerIndex < ConstraintContainerSolvers.Num(); ++ContainerIndex)
 				{
-					const int32 NumIslandConstraints = Island->GetConstraints(ContainerIndex).Num();
+					const int32 NumIslandConstraints = Island->GetNumContainerConstraints(ContainerIndex);
 					NumContainerConstraints[ContainerIndex] += NumIslandConstraints;
 					NumConstraints += NumIslandConstraints;
 				}
@@ -71,7 +70,7 @@ namespace Chaos
 
 					for (FPBDIsland* Island : Islands)
 					{
-						TArrayView<FPBDIslandConstraint> IslandConstraints = MakeArrayView(Island->GetConstraints(ContainerIndex));
+						TArrayView<FPBDIslandConstraint*> IslandConstraints = Island->GetConstraints(ContainerIndex);
 						ConstraintContainerSolvers[ContainerIndex]->AddConstraints(IslandConstraints);
 					}
 				}
@@ -93,10 +92,10 @@ namespace Chaos
 		void FPBDIslandConstraintGroupSolver::SetIterationSettings(const FIterationSettings& InDefaultIterations)
 		{
 			Iterations = InDefaultIterations;
-			for (FPBDIsland* Island : Islands)
-			{
-				Iterations = FIterationSettings::Merge(Iterations, Island->GetIterationSettings());
-			}
+			//for (FPBDIsland* Island : Islands)
+			//{
+			//	Iterations = FIterationSettings::Merge(Iterations, Island->GetIterationSettings());
+			//}
 		}
 
 	}	// namespace Private
diff --git a/Engine/Source/Runtime/Experimental/Chaos/Private/Chaos/Island/IslandGroupManager.cpp b/Engine/Source/Runtime/Experimental/Chaos/Private/Chaos/Island/IslandGroupManager.cpp
--- a/Engine/Source/Runtime/Experimental/Chaos/Private/Chaos/Island/IslandGroupManager.cpp
+++ b/Engine/Source/Runtime/Experimental/Chaos/Private/Chaos/Island/IslandGroupManager.cpp
@@ -3,7 +3,6 @@
 
 #include "Chaos/ConstraintHandle.h"
 #include "Chaos/Island/IslandManager.h"
-#include "Chaos/Island/SolverIsland.h"
 #include "Chaos/PBDConstraintContainer.h"
 #include "ChaosStats.h"
 #include "ProfilingDebugging/ScopedTimers.h"
@@ -116,7 +115,7 @@ namespace Chaos
 			// NOTE: We do not add islands with no constraints even if they have particles. There is no need since the Particles'
 			// predicted positions/rotations will not be changed by the solver. (Zero-constraint islands may contain one isolated Particle).
 			TArray<FPBDIsland*> Islands;
-			for (int32 IslandIndex = 0; IslandIndex < IslandManager.NumIslands(); ++IslandIndex)
+			for (int32 IslandIndex = 0; IslandIndex < IslandManager.GetNumIslands(); ++IslandIndex)
 			{
 				FPBDIsland* Island = IslandManager.GetIsland(IslandIndex); 
 
@@ -125,11 +124,11 @@ namespace Chaos
 				{
 					if(bIsResimming)
 					{ 
-						UE_LOG(LogChaos, Log, TEXT("Chaos Island[%d] needs resim = %d"), IslandIndex, IslandManager.IslandNeedsResim(IslandIndex));
+						UE_LOG(LogChaos, Log, TEXT("Chaos Island[%d] needs resim = %d"), IslandIndex, Island->NeedsResim());
 					}
 				}
 #endif
-				if (!Island->IsSleeping() && (Island->GetNumConstraints() > 0) && (!bIsResimming || (bIsResimming && IslandManager.IslandNeedsResim(IslandIndex))))
+				if (!Island->IsSleeping() && (Island->GetNumConstraints() > 0) && (!bIsResimming || (bIsResimming && Island->NeedsResim())))
 				{
 					Islands.Add(Island);
 				}
