Chaos - fix IslandManager bug when removing the last constraint from an island that has been queued to be merged

diff --git a/Engine/Source/Programs/HeadlessChaos/Private/HeadlessChaosTestGraphEvolution.cpp b/Engine/Source/Programs/HeadlessChaos/Private/HeadlessChaosTestGraphEvolution.cpp
--- a/Engine/Source/Programs/HeadlessChaos/Private/HeadlessChaosTestGraphEvolution.cpp
+++ b/Engine/Source/Programs/HeadlessChaos/Private/HeadlessChaosTestGraphEvolution.cpp
@@ -1253,6 +1253,55 @@ namespace ChaosTest
 
 	}
 
+	// Test that island merging works if we remove the last constraint
+	// in an island that was already queued for merge (PLAY-6440).
+	// 
+	// (d=dynamic, s=sleeping, k=kinematic)
+	//		Ad - Bd   Cd - Dd
+	// =>	Ad - Bd - Cd   Dd
+	//
+	GTEST_TEST(GraphEvolutionTests, TestConstraintGraph_IslandMerge_EnableDisable)
+	{
+		FGraphEvolutionTest Test(4);
+
+		// Create constraints in a chain but disable the middle constraint so we have two islands {A-B} and {C-D}
+		Test.MakeChain();
+		Test.ConstraintHandles[1]->SetEnabled(false);
+
+		Test.Advance();
+
+		EXPECT_EQ(Test.IslandManager->GetNumIslands(), 2);
+		EXPECT_FALSE(Test.ConstraintHandles[1]->IsInConstraintGraph());
+
+		// Enable the constraint B-C.
+		// NOTE: In the implementation, the enable will add the constraint to one of the islands
+		// and queue the two islands to be merged, but the actual merging happens in Advance().
+		Test.ConstraintHandles[1]->SetEnabled(true);
+
+		// Disable the constraint C-D.
+		// NOTE: This will leave the second island without any constraints, but it is
+		// not destroyed immediately because it is queued for merging and, even though
+		// it has no constraints, it still contains particle C which needs to be
+		// copied to the new merged island. Particle D would have been removed because
+		// it does not have any constraints.
+		// Issue (PLAY-6440) was caused by the island being destroyed because it was
+		// empty, but it was still queued to be merged.
+		Test.ConstraintHandles[2]->SetEnabled(false);
+
+		Test.Advance();
+
+		// We should now only have 1 island and D should not be in the graph
+		EXPECT_EQ(Test.IslandManager->GetNumIslands(), 1);
+		EXPECT_TRUE(Test.ConstraintHandles[0]->IsInConstraintGraph());
+		EXPECT_TRUE(Test.ConstraintHandles[1]->IsInConstraintGraph());
+		EXPECT_FALSE(Test.ConstraintHandles[2]->IsInConstraintGraph());	// Not in graph
+		EXPECT_TRUE(Test.ParticleHandles[0]->IsInConstraintGraph());
+		EXPECT_TRUE(Test.ParticleHandles[1]->IsInConstraintGraph());
+		EXPECT_TRUE(Test.ParticleHandles[2]->IsInConstraintGraph());
+		EXPECT_FALSE(Test.ParticleHandles[3]->IsInConstraintGraph());	// Not in graph
+	}
+
+
 	// Test the sparse array repeatable index assignment. See TSparseArray::SortFreeList()
 	GTEST_TEST(SparseArrayTests, TestSortFreeList)
 	{
diff --git a/Engine/Source/Runtime/Experimental/Chaos/Private/Chaos/Island/IslandManager.cpp b/Engine/Source/Runtime/Experimental/Chaos/Private/Chaos/Island/IslandManager.cpp
--- a/Engine/Source/Runtime/Experimental/Chaos/Private/Chaos/Island/IslandManager.cpp
+++ b/Engine/Source/Runtime/Experimental/Chaos/Private/Chaos/Island/IslandManager.cpp
@@ -731,13 +731,24 @@ namespace Chaos::Private
 			FPBDIslandParticle* Node1 = Edge->Nodes[1];
 			FPBDIsland* Island = Edge->Island;
 
-			// Remove edge from the island, and possibly destroy the island (if empty)
+			// Remove edge from the island
 			RemoveEdgeFromIsland(Edge);
 
-			// Destroy the edge
+			// Destroy the edge (also disconnects it from its nodes)
 			DestroyGraphEdge(Edge);
 
-			// We also destroy the nodes if they have no edges left
+			// Destroy the island if it is now empty of constraints
+			// NOTE: If the island is queued for merging we won't destroy it now. 
+			// It will get destroyed later anyway, but more importantly it may 
+			// contain particles that need to be copied to the merged island.
+			if ((Island != nullptr) && (Island->NumEdges == 0) && (Island->MergeSetIslandIndex == INDEX_NONE))
+			{
+				RemoveNodeFromIsland(Node0);
+				RemoveNodeFromIsland(Node1);
+				DestroyIsland(Island);
+			}
+
+			// Destroy the nodes if they have no constraints
 			if ((Node0 != nullptr) && Node0->Edges.IsEmpty())
 			{
 				RemoveNodeFromIsland(Node0);
@@ -748,12 +759,6 @@ namespace Chaos::Private
 				RemoveNodeFromIsland(Node1);
 				DestroyGraphNode(Node1);
 			}
-
-			// And destroy the island if it is now empty
-			if ((Island != nullptr) && (Island->NumEdges == 0))
-			{
-				DestroyIsland(Island);
-			}
 		}
 	}
 
@@ -1185,6 +1190,7 @@ namespace Chaos::Private
 		{
 			check(Island->Nodes.IsEmpty());
 			check(Island->NumEdges == 0);
+			check(Island->MergeSetIslandIndex == INDEX_NONE);
 
 			Islands.Free(Island);
 		}
@@ -1541,12 +1547,13 @@ namespace Chaos::Private
 				// Merge each child island into the parent island
 				for (FPBDIsland* ChildIsland : MergeSet->Islands)
 				{
+					// Reset the merge set ready for next tick. We only really need to do this on the island
+					// that survives the merge, but it helps with error tracking to do it for all (see DestroyIsland)
+					ChildIsland->MergeSet = nullptr;
+					ChildIsland->MergeSetIslandIndex = INDEX_NONE;
+
 					ProcessIslandMerge(ParentIsland, ChildIsland);
 				}
-
-				// Reset the parent island's merge set ready for next tick (the children will have been destroyed)
-				ParentIsland->MergeSet = nullptr;
-				ParentIsland->MergeSetIslandIndex = INDEX_NONE;
 			}
 		}
 
diff --git a/Engine/Source/Runtime/Experimental/Chaos/Public/Chaos/Island/IslandManager.h b/Engine/Source/Runtime/Experimental/Chaos/Public/Chaos/Island/IslandManager.h
--- a/Engine/Source/Runtime/Experimental/Chaos/Public/Chaos/Island/IslandManager.h
+++ b/Engine/Source/Runtime/Experimental/Chaos/Public/Chaos/Island/IslandManager.h
@@ -581,9 +581,18 @@ namespace Chaos::Private
 		const int32 ContainerId = ConstraintContainer.GetContainerId();
 		for (FConstraintHandle* ConstraintHandle : ConstraintContainer.GetConstraintHandles())
 		{
-			if ((ConstraintHandle != nullptr) && ConstraintHandle->IsEnabled())
+			if (ConstraintHandle != nullptr)
 			{
-				AddConstraint(ContainerId, ConstraintHandle, ConstraintHandle->GetConstrainedParticles());
+				const bool bIsInGraph = ConstraintHandle->IsInConstraintGraph();
+				const bool bShouldBeInGraph = ConstraintHandle->IsEnabled();
+				if (bShouldBeInGraph && !bIsInGraph)
+				{
+					AddConstraint(ContainerId, ConstraintHandle, ConstraintHandle->GetConstrainedParticles());
+				}
+				else if (!bShouldBeInGraph && bIsInGraph)
+				{
+					RemoveConstraint(ConstraintHandle);
+				}
 			}
 		}
 	}
