ManagedArray allow circular dependencies for GroupIndexDependency (default is to continue to disallow, but can now optionally allow). Set default values of index arrays to INDEX_NONE when resizing.

Added test to HeadlessChaos that tests adding and removing elements when there are attributes which have a dependency on their own group.

Updated the test in HeadlessChaos which tests circular dependencies now that they're allowed.

diff --git a/Engine/Source/Programs/HeadlessChaos/Private/GeometryCollection/GeometryCollectionTestCreation.cpp b/Engine/Source/Programs/HeadlessChaos/Private/GeometryCollection/GeometryCollectionTestCreation.cpp
--- a/Engine/Source/Programs/HeadlessChaos/Private/GeometryCollection/GeometryCollectionTestCreation.cpp
+++ b/Engine/Source/Programs/HeadlessChaos/Private/GeometryCollection/GeometryCollectionTestCreation.cpp
@@ -1275,9 +1275,9 @@ namespace GeometryCollectionTest
 		Collection->AddExternalAttribute<FTransform>("AttributeC", Group3, Transform, Group2);
 		Collection->AddExternalAttribute<FTransform>("AttributeD", Group4, Transform, Group1);
 
-		// Force a circular group dependency - from G1 to G3
-		// can't figure how to trap the assert that is deliberately fired off during this call
-		//Collection->SetDependency()<FTransform>("AttributeD", Group1, Transform, Group3);
+		// Force a circular group dependency - from G1 to G3 (these are now allowed!)
+		constexpr bool bAllowCircularDependency = true;
+		Collection->SetDependency("AttributeA", Group1, Group3, bAllowCircularDependency);
 
 		delete Collection;
 	}
@@ -1463,6 +1463,182 @@ namespace GeometryCollectionTest
 		delete Collection;
 	}
 
+	void IntListSelfDependencyTest()
+	{
+		//--gtest_filter=GeometryCollection_CreationTest.IntListSelfDependencyTest
+		FManagedArrayCollection Collection;
+
+		constexpr bool bSaved = true;
+		constexpr bool bAllowCircularDependency = true;
+		FManagedArrayCollection::FConstructionParameters DependencyOnA("GroupA", bSaved, bAllowCircularDependency);
+		Collection.AddGroup("GroupA");
+
+		TManagedArray<float>& Values = Collection.AddAttribute<float>("Values", "GroupA");
+
+		//
+		// Build two arrays:
+		//   - One array is an index. (dependent on  A)
+		//   - The other array stores a local copy of the values. 
+		// As elements are deleted from GroupA, the indices will be 
+		// shifted, and missing references will be 
+		// updated to INDEX_NONE. 
+		// 
+
+		// int32
+		TManagedArray<int32>& Int = Collection.AddAttribute<int32>("Int", "GroupA", DependencyOnA);
+		TManagedArray<float>& IntValues = Collection.AddAttribute<float>("IntValues", "GroupA");
+		// TArray<int32>
+		TManagedArray<TArray<int32>>& Array = Collection.AddAttribute<TArray<int32>>("Array", "GroupA", DependencyOnA);
+		TManagedArray<TArray<float>>& ArrayValues = Collection.AddAttribute<TArray<float>>("ArrayValues", "GroupA");
+		//FIntVector
+		TManagedArray<FIntVector>& Vec = Collection.AddAttribute<FIntVector>("IntVec", "GroupA", DependencyOnA);
+		TManagedArray<FVector3f>& VecValues = Collection.AddAttribute<FVector3f>("IntVecValues", "GroupA");
+		//FIntVector2
+		TManagedArray<FIntVector2>& Vec2 = Collection.AddAttribute<FIntVector2>("IntVec2", "GroupA", DependencyOnA);
+		TManagedArray<FVector2f>& Vec2Values = Collection.AddAttribute<FVector2f>("IntVec2Values", "GroupA");
+		// TArray<FIntVector2>
+		TManagedArray<TArray<FIntVector2>>& ArrayVec2 = Collection.AddAttribute<TArray<FIntVector2>>("ArrayIntVec2", "GroupA", DependencyOnA);
+		TManagedArray<TArray<FVector2f>>& ArrayVec2Values = Collection.AddAttribute<TArray<FVector2f>>("ArrayIntVec2Values", "GroupA");
+		//FIntVector4
+		TManagedArray<FIntVector4>& Vec4 = Collection.AddAttribute<FIntVector4>("IntVec4", "GroupA", DependencyOnA);
+		TManagedArray<FVector4f>& Vec4Values = Collection.AddAttribute<FVector4f>("IntVec4Values", "GroupA");
+
+		auto CheckIndexNoneRange = [&](int32 StartIndex, int32 EndIndex)
+		{
+			for (int32 Index = StartIndex; Index < EndIndex; ++Index)
+			{
+				EXPECT_NEAR(Values[Index], 0.f, FLT_EPSILON);
+				EXPECT_EQ(Int[Index], INDEX_NONE);
+				EXPECT_EQ(Array[Index].Num(), 0);
+				EXPECT_EQ(Vec[Index], FIntVector(INDEX_NONE));
+				EXPECT_EQ(Vec2[Index], FIntVector2(INDEX_NONE));
+				EXPECT_EQ(ArrayVec2[Index].Num(), 0);
+				EXPECT_EQ(Vec4[Index], FIntVector4(INDEX_NONE));
+			}
+		};
+
+		// Add elements and check indices are initialized
+		Collection.AddElements(10, "GroupA");
+		CheckIndexNoneRange(0, 10);
+
+		// Set initial values
+		for (int i = 0; i < 10; i++)
+		{
+			Values[i] = 1.f + i * 0.1f;
+		}
+
+		for (int i = 0; i < 10; i++)
+		{
+			// int32
+			Int[i] = i;
+			IntValues[i] = Values[Int[i]];
+
+			// TArray<int32>
+			Array[i].SetNum(i);
+			ArrayValues[i].SetNum(i);
+			for (int j = 0; j < i; ++j)
+			{
+				Array[i][j] = (j + i) % 10;
+				ArrayValues[i][j] = Values[Array[i][j]];
+			}
+
+			//FIntVector
+			Vec[i] = FIntVector(i, (i + 1) % 10, (i + 2) % 10);
+			VecValues[i] = FVector3f(Values[Vec[i][0]], Values[Vec[i][1]], Values[Vec[i][2]]);
+
+			//FIntVector2
+			Vec2[i] = FIntVector2((i + 3) % 10, (i + 4) % 10);
+			Vec2Values[i] = FVector2f(Values[Vec2[i][0]], Values[Vec2[i][1]]);
+
+			// TArray<FIntVector2>
+			ArrayVec2[i].SetNum(i);
+			ArrayVec2Values[i].SetNum(i);
+			for (int j = 0; j < i; ++j)
+			{
+				ArrayVec2[i][j] = FIntVector2((j + i + 3) % 10, (j + i + 4) % 10);
+				ArrayVec2Values[i][j] = FVector2f(Values[ArrayVec2[i][j][0]], Values[ArrayVec2[i][j][1]]);
+			}
+
+			//FIntVector4
+			Vec4[i] = FIntVector4((i + 5) % 10, (i + 6) % 10, (i + 7) % 10, (i + 8) % 10);
+			Vec4Values[i] = FVector4f(Values[Vec4[i][0]], Values[Vec4[i][1]], Values[Vec4[i][2]], Values[Vec4[i][3]]);
+		}
+
+		auto CheckValues = [&](int32 Size)
+		{
+			for (int i = 0; i < Size; i++)
+			{
+				// int32
+				if (Int[i] != INDEX_NONE)
+				{
+					EXPECT_NEAR(IntValues[i], Values[Int[i]], FLT_EPSILON);
+				}
+
+				// TArray<int32>
+				for (int j = 0; j < Array[i].Num(); ++j)
+				{
+					if (Array[i][j] != INDEX_NONE)
+					{
+						EXPECT_NEAR(ArrayValues[i][j], Values[Array[i][j]], FLT_EPSILON);
+					}
+				}
+
+				//FIntVector
+				for (int j = 0; j < 3; ++j)
+				{
+					if (Vec[i][j] != INDEX_NONE)
+					{
+						EXPECT_NEAR(VecValues[i][j], Values[Vec[i][j]], FLT_EPSILON);
+					}
+				}
+
+				//FIntVector2
+				for (int j = 0; j < 2; ++j)
+				{
+					if (Vec2[i][j] != INDEX_NONE)
+					{
+						EXPECT_NEAR(Vec2Values[i][j], Values[Vec2[i][j]], FLT_EPSILON);
+					}
+				}
+
+				// TArray<FIntVector2>
+				for (int j = 0; j < ArrayVec2[i].Num(); ++j)
+				{
+					for (int k = 0; k < 2; ++k)
+					{
+						if (ArrayVec2[i][j][k] != INDEX_NONE)
+						{
+							EXPECT_NEAR(ArrayVec2Values[i][j][k], Values[ArrayVec2[i][j][k]], FLT_EPSILON);
+						}
+					}
+				}
+
+				//FIntVector4
+				for (int j = 0; j < 4; ++j)
+				{
+					if (Vec4[i][j] != INDEX_NONE)
+					{
+						EXPECT_NEAR(Vec4Values[i][j], Values[Vec4[i][j]], FLT_EPSILON);
+					}
+				}
+			}
+		};
+		CheckValues(10);
+
+		// Insert some elements
+		Collection.InsertElements(5, 2, "GroupA");
+		CheckIndexNoneRange(5, 7);
+		CheckValues(Collection.NumElements("GroupA"));
+
+		// Remove some elements
+		Collection.RemoveElements("GroupA", 3, 2);
+		CheckValues(Collection.NumElements("GroupA"));
+
+		// Remove some more elements
+		Collection.RemoveElements("GroupA", { 1,3,5 });
+		CheckValues(Collection.NumElements("GroupA"));
+	}
+
 
 	void AppendManagedArrayCollectionTest()
 	{
@@ -1696,24 +1872,45 @@ namespace GeometryCollectionTest
 		Collection.AddAttribute<int32>("C", "C"); // None
 		Collection.AddAttribute<int32>("D", "D"); // None
 
-		// A cycle would be "C depending on A", which would fail attribute creation.
-		EXPECT_FALSE(Collection.IsConnected("C", "A"));
-
-		// So if we tried to create a dependency on A from C, we would first want to
-		// know if A already connects to C, and throw an error if so. 
-		EXPECT_TRUE(Collection.IsConnected("A","C"));
-
 		EXPECT_TRUE(Collection.IsConnected("A", "B"));
+		EXPECT_TRUE(Collection.IsConnected("A", "C"));
 		EXPECT_TRUE(Collection.IsConnected("A", "D"));
 		EXPECT_TRUE(Collection.IsConnected("B", "C"));
 
 		EXPECT_FALSE(Collection.IsConnected("A", "A"));
 		EXPECT_FALSE(Collection.IsConnected("B", "A"));
-		EXPECT_FALSE(Collection.IsConnected("B","D"));
-		EXPECT_FALSE(Collection.IsConnected("D", "B"));
+		EXPECT_FALSE(Collection.IsConnected("B", "B"));
 		EXPECT_FALSE(Collection.IsConnected("B", "D"));
+		EXPECT_FALSE(Collection.IsConnected("C", "A"));
+		EXPECT_FALSE(Collection.IsConnected("C", "B"));
+		EXPECT_FALSE(Collection.IsConnected("C", "C"));
+		EXPECT_FALSE(Collection.IsConnected("C", "D"));
+		EXPECT_FALSE(Collection.IsConnected("D", "A"));
 		EXPECT_FALSE(Collection.IsConnected("D", "B"));
+		EXPECT_FALSE(Collection.IsConnected("D", "C"));
+		EXPECT_FALSE(Collection.IsConnected("D", "D"));
+
+		// Now create a cycle by connecting C back to A
+		constexpr bool bAllowCircularDependency = true;
+		Collection.SetDependency("C", "C", "A", bAllowCircularDependency);
+
+		EXPECT_TRUE(Collection.IsConnected("A", "A"));
+		EXPECT_TRUE(Collection.IsConnected("A", "B"));
+		EXPECT_TRUE(Collection.IsConnected("A", "C"));
+		EXPECT_TRUE(Collection.IsConnected("A", "D"));
+		EXPECT_TRUE(Collection.IsConnected("B", "A"));
+		EXPECT_TRUE(Collection.IsConnected("B", "B"));
+		EXPECT_TRUE(Collection.IsConnected("B", "C"));
+		EXPECT_TRUE(Collection.IsConnected("B", "D"));
+		EXPECT_TRUE(Collection.IsConnected("C", "A"));
+		EXPECT_TRUE(Collection.IsConnected("C", "B"));
+		EXPECT_TRUE(Collection.IsConnected("C", "C"));
+		EXPECT_TRUE(Collection.IsConnected("C", "D"));
+
 		EXPECT_FALSE(Collection.IsConnected("D", "A"));
+		EXPECT_FALSE(Collection.IsConnected("D", "B"));
+		EXPECT_FALSE(Collection.IsConnected("D", "C"));
+		EXPECT_FALSE(Collection.IsConnected("D", "D"));
 	}
 
 }
diff --git a/Engine/Source/Programs/HeadlessChaos/Private/HeadlessChaosApp.cpp b/Engine/Source/Programs/HeadlessChaos/Private/HeadlessChaosApp.cpp
--- a/Engine/Source/Programs/HeadlessChaos/Private/HeadlessChaosApp.cpp
+++ b/Engine/Source/Programs/HeadlessChaos/Private/HeadlessChaosApp.cpp
@@ -341,6 +341,7 @@ TEST(GeometryCollection_CreationTest,ReindexMaterialsTest) { GeometryCollectionT
 TEST(GeometryCollection_CreationTest,ContiguousElementsTest) { GeometryCollectionTest::ContiguousElementsTest(); SUCCEED(); }
 TEST(GeometryCollection_CreationTest, AttributeDependencyTest) { GeometryCollectionTest::AttributeDependencyTest(); SUCCEED(); }
 TEST(GeometryCollection_CreationTest, IntListReindexOnDeletionTest) { GeometryCollectionTest::IntListReindexOnDeletionTest(); SUCCEED(); }
+TEST(GeometryCollection_CreationTest, IntListSelfDependencyTest) { GeometryCollectionTest::IntListSelfDependencyTest(); SUCCEED(); }
 TEST(GeometryCollection_CreationTest, AppendManagedArrayCollectionTest) { GeometryCollectionTest::AppendManagedArrayCollectionTest(); SUCCEED(); }
 TEST(GeometryCollection_CreationTest, AppendTransformCollectionTest) { GeometryCollectionTest::AppendTransformCollectionTest(); SUCCEED(); }
 TEST(GeometryCollection_CreationTest, CollectionCycleTest) { GeometryCollectionTest::CollectionCycleTest(); SUCCEED(); }
diff --git a/Engine/Source/Programs/HeadlessChaos/Public/GeometryCollection/GeometryCollectionTestCreation.h b/Engine/Source/Programs/HeadlessChaos/Public/GeometryCollection/GeometryCollectionTestCreation.h
--- a/Engine/Source/Programs/HeadlessChaos/Public/GeometryCollection/GeometryCollectionTestCreation.h
+++ b/Engine/Source/Programs/HeadlessChaos/Public/GeometryCollection/GeometryCollectionTestCreation.h
@@ -39,6 +39,8 @@ namespace GeometryCollectionTest
 
 	void IntListReindexOnDeletionTest();
 
+	void IntListSelfDependencyTest();
+
 	void AppendManagedArrayCollectionTest();
 
 	void AppendTransformCollectionTest();
diff --git a/Engine/Source/Runtime/Experimental/Chaos/Private/GeometryCollection/GeometryCollection.cpp b/Engine/Source/Runtime/Experimental/Chaos/Private/GeometryCollection/GeometryCollection.cpp
--- a/Engine/Source/Runtime/Experimental/Chaos/Private/GeometryCollection/GeometryCollection.cpp
+++ b/Engine/Source/Runtime/Experimental/Chaos/Private/GeometryCollection/GeometryCollection.cpp
@@ -149,6 +149,8 @@ void FGeometryCollection::Construct()
 
 void FGeometryCollection::SetDefaults(FName Group, uint32 StartSize, uint32 NumElements)
 {
+	Super::SetDefaults(Group, StartSize, NumElements);
+
 	if (Group == FTransformCollection::TransformGroup)
 	{
 		for (uint32 Idx = StartSize; Idx < StartSize + NumElements; ++Idx)
@@ -1591,6 +1593,7 @@ void FGeometryCollection::Init(FGeometryCollection* Collection, const TArray<flo
 		TManagedArray<int32>& MaterialIndex = Collection->MaterialIndex;
 		TManagedArray<bool>& Internal = Collection->Internal;
 		TManagedArray<FTransform>& Transform = Collection->Transform;
+		TManagedArray<int32>& BoneMap = Collection->BoneMap;
 		
 		Collection->SetNumUVLayers(1);
 
@@ -1604,6 +1607,7 @@ void FGeometryCollection::Init(FGeometryCollection* Collection, const TArray<flo
 			(*UV0)[Idx] = FVector2f::ZeroVector;
 
 			Colors[Idx] = FLinearColor::White;
+			BoneMap[Idx] = 0;
 		}
 
 		
diff --git a/Engine/Source/Runtime/Experimental/Chaos/Private/GeometryCollection/GeometryCollectionUtility.cpp b/Engine/Source/Runtime/Experimental/Chaos/Private/GeometryCollection/GeometryCollectionUtility.cpp
--- a/Engine/Source/Runtime/Experimental/Chaos/Private/GeometryCollection/GeometryCollectionUtility.cpp
+++ b/Engine/Source/Runtime/Experimental/Chaos/Private/GeometryCollection/GeometryCollectionUtility.cpp
@@ -41,6 +41,7 @@ namespace GeometryCollection
 		TManagedArray<bool>& Internal = RestCollection->Internal;
 		TManagedArray<FTransform>&  Transform = RestCollection->Transform;
 		TManagedArray<int32>& SimType = RestCollection->SimulationType;
+		TManagedArray<int32>& BoneMap = RestCollection->BoneMap;
 
 		// set the particle information
 		Transform[0] = center;
@@ -86,6 +87,15 @@ namespace GeometryCollection
 		Colors[6] = FLinearColor::White;
 		Colors[7] = FLinearColor::White;
 
+		BoneMap[0] = 0;
+		BoneMap[1] = 0;
+		BoneMap[2] = 0;
+		BoneMap[3] = 0;
+		BoneMap[4] = 0;
+		BoneMap[5] = 0;
+		BoneMap[6] = 0;
+		BoneMap[7] = 0;
+
 
 		// set the index information
 
diff --git a/Engine/Source/Runtime/Experimental/Chaos/Private/GeometryCollection/ManagedArrayCollection.cpp b/Engine/Source/Runtime/Experimental/Chaos/Private/GeometryCollection/ManagedArrayCollection.cpp
--- a/Engine/Source/Runtime/Experimental/Chaos/Private/GeometryCollection/ManagedArrayCollection.cpp
+++ b/Engine/Source/Runtime/Experimental/Chaos/Private/GeometryCollection/ManagedArrayCollection.cpp
@@ -409,12 +409,12 @@ void FManagedArrayCollection::ReorderElements(FName Group, const TArray<int32>&
 	}
 }
 
-void FManagedArrayCollection::SetDependency(FName Name, FName Group, FName DependencyGroup)
+void FManagedArrayCollection::SetDependency(FName Name, FName Group, FName DependencyGroup, bool bAllowCircularDependency)
 {
 	ensure(HasAttribute(Name, Group));
-	if (ensure(!IsConnected(DependencyGroup, Group)))
+	FKeyType Key = FManagedArrayCollection::MakeMapKey(Name, Group);
+	if (ensure(bAllowCircularDependency || !IsConnected(DependencyGroup, Group)))
 	{
-		FKeyType Key = FManagedArrayCollection::MakeMapKey(Name, Group);
 		Map[Key].GroupIndexDependency = DependencyGroup;
 	}
 }
@@ -431,6 +431,17 @@ void FManagedArrayCollection::RemoveDependencyFor(FName Group)
 	}
 }
 
+void FManagedArrayCollection::SetDefaults(FName Group, uint32 StartSize, uint32 NumElements)
+{
+	for (TTuple<FKeyType, FValueType>& Entry : Map)
+	{
+		if (Entry.Key.Get<1>() == Group)
+		{
+			Entry.Value.Value->SetDefaults(StartSize, NumElements, Entry.Value.GroupIndexDependency != "");
+		}
+	}
+}
+
 void FManagedArrayCollection::SyncGroupSizeFrom(const FManagedArrayCollection& InCollection, FName Group)
 {
 	if (!HasGroup(Group))
diff --git a/Engine/Source/Runtime/Experimental/Chaos/Public/GeometryCollection/GeometryCollection.h b/Engine/Source/Runtime/Experimental/Chaos/Public/GeometryCollection/GeometryCollection.h
--- a/Engine/Source/Runtime/Experimental/Chaos/Public/GeometryCollection/GeometryCollection.h
+++ b/Engine/Source/Runtime/Experimental/Chaos/Public/GeometryCollection/GeometryCollection.h
@@ -273,7 +273,7 @@ public:
 	//
 	//
 
-	void SetDefaults(FName Group, uint32 StartSize, uint32 NumElements);
+	virtual void SetDefaults(FName Group, uint32 StartSize, uint32 NumElements) override;
 
 	// Transform Group
 	TManagedArray<int32>		TransformToGeometryIndex;
diff --git a/Engine/Source/Runtime/Experimental/Chaos/Public/GeometryCollection/GeometryCollectionUtility.h b/Engine/Source/Runtime/Experimental/Chaos/Public/GeometryCollection/GeometryCollectionUtility.h
--- a/Engine/Source/Runtime/Experimental/Chaos/Public/GeometryCollection/GeometryCollectionUtility.h
+++ b/Engine/Source/Runtime/Experimental/Chaos/Public/GeometryCollection/GeometryCollectionUtility.h
@@ -48,6 +48,7 @@ namespace GeometryCollection
 		TManagedArray<bool>& Internal = RestCollection->Internal;
 		TManagedArray<FTransform>& Transform = RestCollection->Transform;
 		TManagedArray<int32>& SimulationType = RestCollection->SimulationType;
+		TManagedArray<int32>& BoneMap = RestCollection->BoneMap;
 
 		// Set particle info
 		Transform[0] = Xf;
@@ -64,6 +65,7 @@ namespace GeometryCollection
 			UV0[Idx] = FVector2f(UV);
 			
 			Colors[Idx] = FLinearColor::White;
+			BoneMap[Idx] = 0;
 		}
 
 		// Set face info
diff --git a/Engine/Source/Runtime/Experimental/Chaos/Public/GeometryCollection/ManagedArray.h b/Engine/Source/Runtime/Experimental/Chaos/Public/GeometryCollection/ManagedArray.h
--- a/Engine/Source/Runtime/Experimental/Chaos/Public/GeometryCollection/ManagedArray.h
+++ b/Engine/Source/Runtime/Experimental/Chaos/Public/GeometryCollection/ManagedArray.h
@@ -177,6 +177,10 @@ protected:
 	*/
 	virtual void CopyRange(const FManagedArrayBase& ConstArray, int32 Start, int32 Stop, int32 Offset = 0) {};
 
+	/**
+	* Set default values. 
+	*/
+	virtual void SetDefaults(uint32 StartSize, uint32 NumElements, bool bHasGroupIndexDependency) {};
 
 	/**
 	* Get allocated memory 
@@ -1113,6 +1117,17 @@ public:
 			}
 		}
 	}
+
+	virtual void SetDefaults(uint32 StartSize, uint32 NumElements, bool bHasGroupIndexDependency) override
+	{
+		if (bHasGroupIndexDependency)
+		{
+			for (uint32 Index = StartSize; Index < StartSize + NumElements; ++Index)
+			{
+				this->operator[](Index) = INDEX_NONE;
+			}
+		}
+	}
 };
 
 
@@ -1253,6 +1268,17 @@ public:
 			}
 		}
 	}
+
+	virtual void SetDefaults(uint32 StartSize, uint32 NumElements, bool bHasGroupIndexDependency) override
+	{
+		if (bHasGroupIndexDependency)
+		{
+			for (uint32 Index = StartSize; Index < StartSize + NumElements; ++Index)
+			{
+				this->operator[](Index) = FIntVector(INDEX_NONE);
+			}
+		}
+	}
 };
 
 template<>
@@ -1319,6 +1345,17 @@ public:
 			}
 		}
 	}
+
+	virtual void SetDefaults(uint32 StartSize, uint32 NumElements, bool bHasGroupIndexDependency) override
+	{
+		if (bHasGroupIndexDependency)
+		{
+			for (uint32 Index = StartSize; Index < StartSize + NumElements; ++Index)
+			{
+				this->operator[](Index) = FIntVector2(INDEX_NONE);
+			}
+		}
+	}
 };
 
 template<>
@@ -1459,6 +1496,17 @@ public:
 			}
 		}
 	}
+
+	virtual void SetDefaults(uint32 StartSize, uint32 NumElements, bool bHasGroupIndexDependency) override
+	{
+		if (bHasGroupIndexDependency)
+		{
+			for (uint32 Index = StartSize; Index < StartSize + NumElements; ++Index)
+			{
+				this->operator[](Index) = FIntVector4(INDEX_NONE);
+			}
+		}
+	}
 };
 
 template<>
diff --git a/Engine/Source/Runtime/Experimental/Chaos/Public/GeometryCollection/ManagedArrayCollection.h b/Engine/Source/Runtime/Experimental/Chaos/Public/GeometryCollection/ManagedArrayCollection.h
--- a/Engine/Source/Runtime/Experimental/Chaos/Public/GeometryCollection/ManagedArrayCollection.h
+++ b/Engine/Source/Runtime/Experimental/Chaos/Public/GeometryCollection/ManagedArrayCollection.h
@@ -65,13 +65,15 @@ public:
 	struct FConstructionParameters {
 
 		FConstructionParameters(FName GroupIndexDependencyIn = ""
-			, bool SavedIn = true)
+			, bool SavedIn = true, bool bInAllowCircularDependency = false)
 			: GroupIndexDependency(GroupIndexDependencyIn)
 			, Saved(SavedIn)
+			, bAllowCircularDependency(bInAllowCircularDependency)
 		{}
 
 		FName GroupIndexDependency;
 		bool Saved;
+		bool bAllowCircularDependency;
 	};
 
 	struct FProcessingParameters
@@ -133,7 +135,7 @@ public:
 			FValueType Value(ManagedArrayType<T>(), *(new TManagedArray<T>()));
 			Value.Value->Resize(NumElements(Group));
 			Value.Saved = Parameters.Saved;
-			if (ensure(!IsConnected(Parameters.GroupIndexDependency,Group)))
+			if (ensure(Parameters.bAllowCircularDependency || !IsConnected(Parameters.GroupIndexDependency, Group)))
 			{
 				Value.GroupIndexDependency = Parameters.GroupIndexDependency;
 			}
@@ -212,7 +214,7 @@ public:
 		FValueType Value(ManagedArrayType<T>(), ValueIn);
 		Value.Value->Resize(NumElements(Group));
 		Value.Saved = Parameters.Saved;
-		if (ensure(!IsConnected(Parameters.GroupIndexDependency, Group)))
+		if (ensure(Parameters.bAllowCircularDependency || !IsConnected(Parameters.GroupIndexDependency, Group)))
 		{
 			Value.GroupIndexDependency = Parameters.GroupIndexDependency;
 		}
@@ -451,7 +453,7 @@ public:
 	/**
 	*
 	*/
-	void SetDependency(FName Name, FName Group, FName DependencyGroup);
+	void SetDependency(FName Name, FName Group, FName DependencyGroup, bool bAllowCircularDependency = false);
 
 	/**
 	*
@@ -660,7 +662,6 @@ private:
 
 	};
 
-	virtual void SetDefaults(FName Group, uint32 StartSize, uint32 NumElements) {};
 
 	TMap< FKeyType, FValueType> Map;	//data is owned by the map explicitly
 	TMap< FName, FGroupInfo> GroupInfo;
@@ -668,6 +669,8 @@ private:
 
 protected:
 
+	virtual void SetDefaults(FName Group, uint32 StartSize, uint32 NumElements);
+
 	/**
 	 * Virtual helper function called by CopyMatchingAttributesFrom; adds attributes 'default, but optional' attributes that are present in InCollection
 	 * This is used by FGeometryCollection to make sure all UV layers are copied over by CopyMatchingAttributesFrom()
