Chaos - speculative fixes for race condition in collision detection caused by a dynamic particle being leftin the moving kinematic particle list

- remove particles from the MovingKinematics list when disabled
- don't allow calls to SetKinematicTarget on dynamic particles

diff --git a/Engine/Source/Programs/HeadlessChaos/Private/HeadlessChaosTestPhysicsScene.cpp b/Engine/Source/Programs/HeadlessChaos/Private/HeadlessChaosTestPhysicsScene.cpp
--- a/Engine/Source/Programs/HeadlessChaos/Private/HeadlessChaosTestPhysicsScene.cpp
+++ b/Engine/Source/Programs/HeadlessChaos/Private/HeadlessChaosTestPhysicsScene.cpp
@@ -1926,4 +1926,100 @@ namespace ChaosTest {
 
 		EXPECT_TRUE(HitBuffer.HasBlockingHit());
 	}
+
+	// Disable a moving kinematic particle and switch it to dynamic while disabled. Verify that when
+	// disabled it is not in any active lists, and that when re-enabled it is not duplicated
+	//
+	// There was a (benign) bug where particles were not removed from the MovingKinematics list until the
+	// next call to ApplyKinematicTargets, which means they can be included in collision detection.
+	//
+	GTEST_TEST(EvolutionTests, TestDisableKinematicEnableDynamic)
+	{
+		const FReal Dt = FReal(1.0 / 60.0);
+		FParticleUniqueIndicesMultithreaded UniqueIndices;
+		FPBDRigidsSOAs Particles(UniqueIndices);
+		THandleArray<FChaosPhysicsMaterial> PhysicalMaterials;
+		FPBDRigidsEvolutionGBF Evolution(Particles, PhysicalMaterials);
+
+		Evolution.GetGravityForces().SetAcceleration(FVec3(0), 0);
+
+		// Create a moving kinematic particle
+		TArray<FPBDRigidParticleHandle*> ParticleHandles = Evolution.CreateDynamicParticles(1);
+		Evolution.EnableParticle(ParticleHandles[0]);
+		Evolution.SetParticleObjectState(ParticleHandles[0], EObjectStateType::Kinematic);
+		Evolution.SetParticleKinematicTarget(ParticleHandles[0], FKinematicTarget::MakePositionTarget(FVec3(10,0,0), FRotation3::FromIdentity()));
+
+		Evolution.AdvanceOneTimeStep(Dt);
+
+		EXPECT_FALSE(ParticleHandles[0]->IsDynamic());
+		EXPECT_TRUE(ParticleHandles[0]->IsKinematic());
+
+		// Check that it is in the moving kinematics list
+		{
+			const TParticleView<TPBDRigidParticles<FReal, 3>>& DynamicSleepingView = Particles.GetNonDisabledDynamicView();
+			const TParticleView<TPBDRigidParticles<FReal, 3>>& DynamicMovingKinematicView = Particles.GetActiveDynamicMovingKinematicParticlesView();
+			EXPECT_EQ(DynamicSleepingView.Num(), 0);
+			EXPECT_EQ(DynamicMovingKinematicView.Num(), 1);
+		}
+
+		// Disable the kinematic
+		Evolution.DisableParticle(ParticleHandles[0]);
+
+		// It should not be in any active views now
+		{
+			const TParticleView<TPBDRigidParticles<FReal, 3>>& DynamicSleepingView = Particles.GetNonDisabledDynamicView();
+			const TParticleView<TPBDRigidParticles<FReal, 3>>& DynamicMovingKinematicView = Particles.GetActiveDynamicMovingKinematicParticlesView();
+			EXPECT_EQ(DynamicSleepingView.Num(), 0);
+			EXPECT_EQ(DynamicMovingKinematicView.Num(), 0);
+		}
+
+		// Make the particle dynamic and enable it
+		Evolution.SetParticleObjectState(ParticleHandles[0], EObjectStateType::Dynamic);
+		Evolution.EnableParticle(ParticleHandles[0]);
+
+		// Check that it is in the active views, but not duplicated in either
+		{
+			const TParticleView<TPBDRigidParticles<FReal, 3>>& DynamicSleepingView = Particles.GetNonDisabledDynamicView();
+			const TParticleView<TPBDRigidParticles<FReal, 3>>& DynamicMovingKinematicView = Particles.GetActiveDynamicMovingKinematicParticlesView();
+			EXPECT_EQ(DynamicSleepingView.Num(), 1);
+			EXPECT_EQ(DynamicMovingKinematicView.Num(), 1);
+		}
+	}
+
+	// Check that we cannot set a kinematic target on a dynamic particle.
+	//
+	// This would cause the dynamic particle to be added to the MovingKinematics
+	// list which means it would appear twice in the GetActiveDynamicMovingKinematicParticlesView
+	// which can result in a race condition in collision detection as a particle pair will be
+	// considered twice and possibly on different threads.
+	GTEST_TEST(EvolutionTests, TestKinematicTargetOnDynamic)
+	{
+		const FReal Dt = FReal(1.0 / 60.0);
+		FParticleUniqueIndicesMultithreaded UniqueIndices;
+		FPBDRigidsSOAs Particles(UniqueIndices);
+		THandleArray<FChaosPhysicsMaterial> PhysicalMaterials;
+		FPBDRigidsEvolutionGBF Evolution(Particles, PhysicalMaterials);
+
+		Evolution.GetGravityForces().SetAcceleration(FVec3(0), 0);
+
+		// Create a dynamic particle
+		TArray<FPBDRigidParticleHandle*> ParticleHandles = Evolution.CreateDynamicParticles(1);
+		Evolution.EnableParticle(ParticleHandles[0]);
+
+		// Set the kinematic target
+		Evolution.SetParticleKinematicTarget(ParticleHandles[0], FKinematicTarget::MakePositionTarget(FVec3(10, 0, 0), FRotation3::FromIdentity()));
+
+		// We should not have a kinematic target
+		EXPECT_FALSE(ParticleHandles[0]->KinematicTarget().IsSet());
+
+		Evolution.AdvanceOneTimeStep(Dt);
+
+		// Check that it is in the active views, but not duplicated in either
+		{
+			const TParticleView<TPBDRigidParticles<FReal, 3>>& DynamicSleepingView = Particles.GetNonDisabledDynamicView();
+			const TParticleView<TPBDRigidParticles<FReal, 3>>& DynamicMovingKinematicView = Particles.GetActiveDynamicMovingKinematicParticlesView();
+			EXPECT_EQ(DynamicSleepingView.Num(), 1);
+			EXPECT_EQ(DynamicMovingKinematicView.Num(), 1);
+		}
+	}
 }
diff --git a/Engine/Source/Runtime/Experimental/Chaos/Private/PBDRigidsSolver.cpp b/Engine/Source/Runtime/Experimental/Chaos/Private/PBDRigidsSolver.cpp
--- a/Engine/Source/Runtime/Experimental/Chaos/Private/PBDRigidsSolver.cpp
+++ b/Engine/Source/Runtime/Experimental/Chaos/Private/PBDRigidsSolver.cpp
@@ -353,6 +353,13 @@ namespace Chaos
 
 		int32 ChaosVisualDebuggerEnable = 1;
 		FAutoConsoleVariableRef CVarChaosVisualDebuggerEnable(TEXT("p.Chaos.VisualDebuggerEnable"), ChaosVisualDebuggerEnable, TEXT("Enable/Disable pushing/saving data to the visual debugger"));
+
+
+		// Enable a couple bug fixes with temporary roll-back just in case
+		bool bRemoveParticleFromMovingKinematicsOnDisable = true;
+		bool bDisallowSetKinematicTargetOnDynamics = true;
+		FAutoConsoleVariableRef CVarChaosRemoveParticleFromMovingKinematicsOnDisable(TEXT("p.Chaos.RemoveParticleFromMovingKinematicsOnDisable"), bRemoveParticleFromMovingKinematicsOnDisable, TEXT(""));
+		FAutoConsoleVariableRef CVarChaosDisallowSetKinematicTargetOnDynamics(TEXT("p.Chaos.DisallowSetKinematicTargetOnDynamics"), bDisallowSetKinematicTargetOnDynamics, TEXT(""));
 	}
 }
 
diff --git a/Engine/Source/Runtime/Experimental/Chaos/Public/Chaos/PBDRigidsEvolution.h b/Engine/Source/Runtime/Experimental/Chaos/Public/Chaos/PBDRigidsEvolution.h
--- a/Engine/Source/Runtime/Experimental/Chaos/Public/Chaos/PBDRigidsEvolution.h
+++ b/Engine/Source/Runtime/Experimental/Chaos/Public/Chaos/PBDRigidsEvolution.h
@@ -66,6 +66,11 @@ struct FBroadPhaseConfig
 
 extern CHAOS_API FBroadPhaseConfig BroadPhaseConfig;
 
+namespace CVars
+{
+	extern bool bDisallowSetKinematicTargetOnDynamics;
+}
+
 namespace Collisions
 {
 	void CHAOS_API ResetChaosCollisionCounters();
@@ -412,7 +417,16 @@ public:
 	*/
 	void SetParticleKinematicTarget(FKinematicGeometryParticleHandle* KinematicHandle, const FKinematicTarget& NewKinematicTarget)
 	{
-		if (KinematicHandle)
+		if (KinematicHandle == nullptr)
+		{
+			return;
+		}
+
+		// NOTE: We ignore SetKinematicTarget if called on a dynamic body since this is not supported
+		// and if we callMarkMovingKinematic(KinematicHandle) on a dynamic the particle will end up
+		// in two mutally exlcusivbe particle lists, leading to a collision detection race condition.
+		// @todo(chaos): maybe we should ensure that this is not called for dynamics
+		if ((KinematicHandle->ObjectState() == EObjectStateType::Kinematic) || !CVars::bDisallowSetKinematicTargetOnDynamics)
 		{
 			// optimization : we keep track of moving kinematic targets ( list gets clear every frame )
 			if (NewKinematicTarget.GetMode() != EKinematicTargetMode::None)
diff --git a/Engine/Source/Runtime/Experimental/Chaos/Public/Chaos/PBDRigidsSOAs.h b/Engine/Source/Runtime/Experimental/Chaos/Public/Chaos/PBDRigidsSOAs.h
--- a/Engine/Source/Runtime/Experimental/Chaos/Public/Chaos/PBDRigidsSOAs.h
+++ b/Engine/Source/Runtime/Experimental/Chaos/Public/Chaos/PBDRigidsSOAs.h
@@ -7,6 +7,11 @@
 
 namespace Chaos
 {
+namespace CVars
+{
+	extern bool bRemoveParticleFromMovingKinematicsOnDisable;
+}
+
 class IParticleUniqueIndices
 {
 public:
@@ -438,6 +443,10 @@ public:
 			// All active particles RIGID particles
 			{
 				RemoveFromActiveArray(PBDRigid, /*bStillDirty=*/false);
+				if (CVars::bRemoveParticleFromMovingKinematicsOnDisable)
+				{
+					MovingKinematicsMapArray.Remove(PBDRigid);
+				}
 			}
 		}
 		else if (Particle->CastToKinematicParticle())
@@ -534,7 +543,8 @@ public:
 		{
 			ActivateParticle(Particle, true);
 		}
-			UpdateViews();
+		
+		UpdateViews();
 	}
 
 	/**
@@ -602,6 +612,10 @@ public:
 		{
 			Particle->MoveToSOA(*DynamicDisabledParticles);
 			RemoveFromActiveArray(Particle->CastToRigidParticle(), /*bStillDirty=*/ false);
+			if (CVars::bRemoveParticleFromMovingKinematicsOnDisable)
+			{
+				MovingKinematicsMapArray.Remove(Particle);
+			}
 		}
 		else
 		{
