Chaos - Fix character ground constraint handling of disabled and destroyed particles

diff --git a/Engine/Source/Programs/HeadlessChaos/Private/Character/TestCharacterGroundConstraintProxy.cpp b/Engine/Source/Programs/HeadlessChaos/Private/Character/TestCharacterGroundConstraintProxy.cpp
--- a/Engine/Source/Programs/HeadlessChaos/Private/Character/TestCharacterGroundConstraintProxy.cpp
+++ b/Engine/Source/Programs/HeadlessChaos/Private/Character/TestCharacterGroundConstraintProxy.cpp
@@ -91,6 +91,13 @@ namespace ChaosTest
 
 		EXPECT_EQ(Solver->GetCharacterGroundConstraints().GetNumConstraints(), 1);
 
+		// Check that the constraint is set on the constrained particles
+		ASSERT_EQ(Solver->GetParticles().GetDynamicParticles().ParticleConstraints(0).Num(), 1);
+		EXPECT_EQ(Solver->GetParticles().GetDynamicParticles().ParticleConstraints(0)[0], Constraint_PT);
+
+		ASSERT_EQ(Solver->GetParticles().GetDynamicParticles().ParticleConstraints(1).Num(), 1);
+		EXPECT_EQ(Solver->GetParticles().GetDynamicParticles().ParticleConstraints(1)[0], Constraint_PT);
+
 		// Change something on the game side and check that it is synced on the physics side
 		RemoteData.Clear(Manager, DataIdx);
 
@@ -101,6 +108,19 @@ namespace ChaosTest
 
 		EXPECT_FLOAT_EQ(Constraint_PT->GetData().GroundDistance, 1234.5f);
 
+		// Change something on the physics thread and check that it's synced back to the game thread
+		FCharacterGroundConstraintDynamicData DynData;
+		DynData.GroundDistance = 1.2f;
+		DynData.GroundNormal = FVector(1.0f, 0.0f, 0.0f);
+		Constraint_PT->SetData(DynData);
+
+		Chaos::FDirtyCharacterGroundConstraintData DirtyConstraintData;
+		Proxy->BufferPhysicsResults(DirtyConstraintData);
+		Proxy->PullFromPhysicsState(DirtyConstraintData, 0);
+
+		EXPECT_FLOAT_EQ(Constraint_GT->GetGroundDistance(), DynData.GroundDistance);
+		EXPECT_VECTOR_FLOAT_EQ(Constraint_GT->GetGroundNormal(), DynData.GroundNormal);
+
 		// Remove the constraint and particles from the solver
 		Solver->UnregisterObject(Constraint_GT);
 		Solver->UnregisterObject(CharacterProxy);
@@ -110,4 +130,79 @@ namespace ChaosTest
 
 		// Physics thread constraint is deleted in a callback set on the marshalling manager
 	}
+
+	TEST_F(CharacterGroundConstraintsProxyTest, TestEnableDisable)
+	{
+		// Create a game thread constraint and register it with the solver
+		FCharacterGroundConstraint* Constraint_GT = new FCharacterGroundConstraint();
+		FSingleParticlePhysicsProxy* CharacterProxy = CreateParticle();
+		FSingleParticlePhysicsProxy* GroundProxy = CreateParticle();
+		FSingleParticlePhysicsProxy* SecondGroundProxy = CreateParticle();
+		Constraint_GT->Init(CharacterProxy);
+		Constraint_GT->SetGroundParticleProxy(GroundProxy);
+
+		Solver->RegisterObject(CharacterProxy);
+		Solver->RegisterObject(GroundProxy);
+		Solver->RegisterObject(Constraint_GT);
+
+		FCharacterGroundConstraintProxy* Proxy = Constraint_GT->GetProxy<FCharacterGroundConstraintProxy>();
+
+		const int32 DataIdx = 0;
+		FDirtyProxiesBucketInfo BucketInfo;
+		BucketInfo.Num[(uint32)EPhysicsProxyType::CharacterGroundConstraintType] = 1;
+		FDirtyPropertiesManager Manager;
+		Manager.PrepareBuckets(BucketInfo);
+
+		FDirtyChaosProperties RemoteData;
+
+		Proxy->PushStateOnGameThread(Manager, 0, RemoteData);
+		Proxy->InitializeOnPhysicsThread(Solver, Manager, DataIdx, RemoteData);
+
+		FCharacterGroundConstraintHandle* Constraint_PT = Proxy->GetPhysicsThreadAPI();
+
+		// Check enabling/disabling
+		EXPECT_TRUE(Constraint_PT->IsEnabled());
+		Constraint_PT->SetEnabled(false);
+		EXPECT_FALSE(Constraint_PT->IsEnabled());
+		Constraint_PT->SetEnabled(true);
+		EXPECT_TRUE(Constraint_PT->IsEnabled());
+
+		// Check that the constraint is set on the constrained particles
+		ASSERT_EQ(Solver->GetParticles().GetDynamicParticles().ParticleConstraints(0).Num(), 1);
+		EXPECT_EQ(Solver->GetParticles().GetDynamicParticles().ParticleConstraints(0)[0], Constraint_PT);
+
+		ASSERT_EQ(Solver->GetParticles().GetDynamicParticles().ParticleConstraints(1).Num(), 1);
+		EXPECT_EQ(Solver->GetParticles().GetDynamicParticles().ParticleConstraints(1)[0], Constraint_PT);
+
+		ASSERT_EQ(Solver->GetParticles().GetDynamicParticles().ParticleConstraints(2).Num(), 0);
+
+		// Change the ground particle and check that the particle constraints are updated
+
+		Constraint_PT->SetGroundParticle(SecondGroundProxy->GetHandle_LowLevel());
+
+		ASSERT_EQ(Solver->GetParticles().GetDynamicParticles().ParticleConstraints(1).Num(), 0);
+
+		ASSERT_EQ(Solver->GetParticles().GetDynamicParticles().ParticleConstraints(2).Num(), 1);
+		EXPECT_EQ(Solver->GetParticles().GetDynamicParticles().ParticleConstraints(2)[0], Constraint_PT);
+
+		// Destroy the constraint on the physics thread. Ensure that the particle constraints are unregistered
+		Proxy->DestroyOnPhysicsThread(Solver);
+
+		ASSERT_EQ(Solver->GetParticles().GetDynamicParticles().ParticleConstraints(0).Num(), 0);
+		ASSERT_EQ(Solver->GetParticles().GetDynamicParticles().ParticleConstraints(1).Num(), 0);
+		ASSERT_EQ(Solver->GetParticles().GetDynamicParticles().ParticleConstraints(2).Num(), 0);
+
+		// Reinitialize on the physics thread
+		Proxy->InitializeOnPhysicsThread(Solver, Manager, DataIdx, RemoteData);
+		ASSERT_EQ(Solver->GetParticles().GetDynamicParticles().ParticleConstraints(0).Num(), 1);
+		ASSERT_EQ(Solver->GetParticles().GetDynamicParticles().ParticleConstraints(1).Num(), 1);
+		ASSERT_EQ(Solver->GetParticles().GetDynamicParticles().ParticleConstraints(2).Num(), 0);
+
+		// Disable the character particle - constraint should be disabled
+		Solver->GetEvolution()->DisableParticle(CharacterProxy->GetHandle_LowLevel());
+		ASSERT_FALSE(Constraint_PT->IsEnabled());
+
+		Solver->GetEvolution()->EnableParticle(CharacterProxy->GetHandle_LowLevel());
+		ASSERT_TRUE(Constraint_PT->IsEnabled());
+	}
 }
\ No newline at end of file
diff --git a/Engine/Source/Runtime/Experimental/Chaos/Private/Chaos/Character/CharacterGroundConstraintContainer.cpp b/Engine/Source/Runtime/Experimental/Chaos/Private/Chaos/Character/CharacterGroundConstraintContainer.cpp
--- a/Engine/Source/Runtime/Experimental/Chaos/Private/Chaos/Character/CharacterGroundConstraintContainer.cpp
+++ b/Engine/Source/Runtime/Experimental/Chaos/Private/Chaos/Character/CharacterGroundConstraintContainer.cpp
@@ -153,8 +153,6 @@ namespace Chaos
 					}
 				}
 			}
-
-			RemovedParticle->ParticleConstraints().Reset();
 		}
 	}
 
diff --git a/Engine/Source/Runtime/Experimental/Chaos/Private/Chaos/PBDJointConstraints.cpp b/Engine/Source/Runtime/Experimental/Chaos/Private/Chaos/PBDJointConstraints.cpp
--- a/Engine/Source/Runtime/Experimental/Chaos/Private/Chaos/PBDJointConstraints.cpp
+++ b/Engine/Source/Runtime/Experimental/Chaos/Private/Chaos/PBDJointConstraints.cpp
@@ -542,8 +542,6 @@ namespace Chaos
 					}
 				}
 			}
-
-			RemovedParticle->ParticleConstraints().Reset();
 		}
 	}
 
diff --git a/Engine/Source/Runtime/Experimental/Chaos/Private/PhysicsProxy/CharacterGroundConstraintProxy.cpp b/Engine/Source/Runtime/Experimental/Chaos/Private/PhysicsProxy/CharacterGroundConstraintProxy.cpp
--- a/Engine/Source/Runtime/Experimental/Chaos/Private/PhysicsProxy/CharacterGroundConstraintProxy.cpp
+++ b/Engine/Source/Runtime/Experimental/Chaos/Private/PhysicsProxy/CharacterGroundConstraintProxy.cpp
@@ -68,13 +68,6 @@ namespace Chaos
 				if (ConstraintSettings && ConstraintData)
 				{
 					Constraint_PT = ConstraintContainer.AddConstraint(*ConstraintSettings, *ConstraintData, CharacterHandle, GroundHandle);
-
-					CharacterHandle->AddConstraintHandle(Constraint_PT);
-
-					if (GroundHandle)
-					{
-						GroundHandle->AddConstraintHandle(Constraint_PT);
-					}
 				}
 			}
 		}
diff --git a/Engine/Source/Runtime/Experimental/Chaos/Public/Chaos/Character/CharacterGroundConstraintContainer.h b/Engine/Source/Runtime/Experimental/Chaos/Public/Chaos/Character/CharacterGroundConstraintContainer.h
--- a/Engine/Source/Runtime/Experimental/Chaos/Public/Chaos/Character/CharacterGroundConstraintContainer.h
+++ b/Engine/Source/Runtime/Experimental/Chaos/Public/Chaos/Character/CharacterGroundConstraintContainer.h
@@ -71,7 +71,12 @@ namespace Chaos
 		{
 			if (InGroundParticle != GroundParticle)
 			{
+				if (GroundParticle)
+				{
+					GroundParticle->RemoveConstraintHandle(this);
+				}
 				GroundParticle = InGroundParticle;
+				GroundParticle->AddConstraintHandle(this);
 				bGroundParticleChanged = true;
 			}
 		}
diff --git a/Engine/Source/Runtime/Experimental/Chaos/Public/Chaos/PBDRigidsEvolution.h b/Engine/Source/Runtime/Experimental/Chaos/Public/Chaos/PBDRigidsEvolution.h
--- a/Engine/Source/Runtime/Experimental/Chaos/Public/Chaos/PBDRigidsEvolution.h
+++ b/Engine/Source/Runtime/Experimental/Chaos/Public/Chaos/PBDRigidsEvolution.h
@@ -636,14 +636,15 @@ public:
 	*/
 	void DisconnectConstraints(const TSet<FGeometryParticleHandle*>& RemovedParticles)
 	{
-		for (FGeometryParticleHandle* ParticleHandle : RemovedParticles)
+		for (FPBDConstraintContainer* Container : ConstraintContainers)
 		{
-			RemoveConstraintsFromConstraintGraph(ParticleHandle->ParticleConstraints());
+			Container->DisconnectConstraints(RemovedParticles);
 		}
 
-		for (FPBDConstraintContainer* Container : ConstraintContainers)
+		for (FGeometryParticleHandle* ParticleHandle : RemovedParticles)
 		{
-			Container->DisconnectConstraints(RemovedParticles);
+			RemoveConstraintsFromConstraintGraph(ParticleHandle->ParticleConstraints());
+			ParticleHandle->ParticleConstraints().Reset();
 		}
 	}
 
