Add contact modification per particle

diff --git a/Engine/Source/Programs/HeadlessChaos/Private/HeadlessChaosTestContactModification.cpp b/Engine/Source/Programs/HeadlessChaos/Private/HeadlessChaosTestContactModification.cpp
--- a/Engine/Source/Programs/HeadlessChaos/Private/HeadlessChaosTestContactModification.cpp
+++ b/Engine/Source/Programs/HeadlessChaos/Private/HeadlessChaosTestContactModification.cpp
@@ -1083,4 +1083,86 @@ namespace ChaosTest
 		Module->DestroySolver(Solver);
 	}
 
+	GTEST_TEST(AllTraits, ContactModification_SelectByParticle)
+	{
+		FChaosSolversModule* Module = FChaosSolversModule::GetModule();
+		auto* Solver = Module->CreateSolver(nullptr, /*AsyncDt=*/-1);
+		InitSolverSettings(Solver);
+		Solver->SetThreadingMode_External(EThreadingModeTemp::SingleThread);
+
+		// Similar to the "Disable" test:
+		// - Create a static floor and two boxes falling onto it.
+		// - One box has contacts disabled and should fall through, one should collide.
+		// - Rather than loop over all contacts, get contacts for a particular particle proxy
+
+		// simulated cube with downward velocity,should collide with floor and not fall through.
+		FSingleParticlePhysicsProxy* CollidingCubeProxy = FSingleParticlePhysicsProxy::Create(Chaos::FPBDRigidParticle::CreateParticle());
+		auto& CollidingCubeParticle = CollidingCubeProxy->GetGameThreadAPI();
+		auto CollidingCubeGeom = TSharedPtr<FImplicitObject, ESPMode::ThreadSafe>(new TBox<FReal, 3>(FVec3(-100), FVec3(100)));
+		CollidingCubeParticle.SetGeometry(CollidingCubeGeom);
+		Solver->RegisterObject(CollidingCubeProxy);
+		CollidingCubeParticle.SetGravityEnabled(false);
+		CollidingCubeParticle.SetV(FVec3(0, 0, -100));
+		CollidingCubeParticle.SetX(FVec3(200, 0, 500));
+		SetCubeInertiaTensor(CollidingCubeParticle, /*Dimension=*/200, /*Mass=*/1);
+		ChaosTest::SetParticleSimDataToCollide({ CollidingCubeProxy->GetParticle_LowLevel() });
+
+		// Simulated cube with downawrd velocity, contact modification disables collision with floor, should fall through.
+		FSingleParticlePhysicsProxy* ModifiedCubeProxy = FSingleParticlePhysicsProxy::Create(Chaos::FPBDRigidParticle::CreateParticle());
+		auto& ModifiedCubeParticle = ModifiedCubeProxy->GetGameThreadAPI();
+		auto ModifiedCubeGeom = TSharedPtr<FImplicitObject, ESPMode::ThreadSafe>(new TBox<FReal, 3>(FVec3(-100), FVec3(100)));
+		ModifiedCubeParticle.SetGeometry(ModifiedCubeGeom);
+		Solver->RegisterObject(ModifiedCubeProxy);
+		ModifiedCubeParticle.SetGravityEnabled(false);
+		ModifiedCubeParticle.SetV(FVec3(0, 0, -100));
+		ModifiedCubeParticle.SetX(FVec3(-200, 0, 500));
+		SetCubeInertiaTensor(ModifiedCubeParticle, /*Dimension=*/200, /*Mass=*/1);
+		ChaosTest::SetParticleSimDataToCollide({ ModifiedCubeProxy->GetParticle_LowLevel() });
+
+		// static floor at origin, occupying Z = [-100,0]
+		FSingleParticlePhysicsProxy* FloorProxy = FSingleParticlePhysicsProxy::Create(Chaos::FGeometryParticle::CreateParticle());
+		auto& FloorParticle = FloorProxy->GetGameThreadAPI();
+		auto FloorGeom = TSharedPtr<FImplicitObject, ESPMode::ThreadSafe>(new TBox<FReal, 3>(FVec3(-500, -500, -100), FVec3(500, 500, 0)));
+		FloorParticle.SetGeometry(FloorGeom);
+		Solver->RegisterObject(FloorProxy);
+		FloorParticle.SetX(FVec3(0, 0, 0));
+		ChaosTest::SetParticleSimDataToCollide({ FloorProxy->GetParticle_LowLevel() });
+
+		// Save Unique indices of floor and modified cube to disable in contact mod.
+		TVec2<FUniqueIdx> UniqueIndices({ ModifiedCubeParticle.UniqueIdx(), FloorParticle.UniqueIdx() });
+
+		FContactModificationTestCallback* Callback = Solver->CreateAndRegisterSimCallbackObject_External<FContactModificationTestCallback>();
+		Callback->TestLambda = [UniqueIndices, ModifiedCubeProxy](Chaos::FCollisionContactModifier& Modifier)
+		{
+			Chaos::FGeometryParticleHandle* ModifiedCubeParticle = ModifiedCubeProxy->GetHandle_LowLevel();
+			for (FContactPairModifier& PairModifier : Modifier.GetContacts(ModifiedCubeParticle))
+			{
+				PairModifier.Disable();
+			}
+		};
+
+		const float Dt = 1.0f;
+		const int32 Steps = 10;
+		for (int Step = 0; Step < Steps; ++Step)
+		{
+			Solver->AdvanceAndDispatch_External(Dt);
+			Solver->UpdateGameThreadStructures();
+		}
+
+		// Modified cube should be below floor because we disabled collision.
+		EXPECT_LT(ModifiedCubeParticle.X().Z, FloorParticle.X().Z);
+
+		// Colliding cube should be above floor due to collision.
+		EXPECT_GT(CollidingCubeParticle.X().Z, FloorParticle.X().Z);
+
+		// Floor should be at origin.
+		EXPECT_EQ(FloorParticle.X().Z, 0);
+
+		Solver->UnregisterAndFreeSimCallbackObject_External(Callback);
+		Solver->UnregisterObject(CollidingCubeProxy);
+		Solver->UnregisterObject(ModifiedCubeProxy);
+		Solver->UnregisterObject(FloorProxy);
+		Module->DestroySolver(Solver);
+	}
+
 }
diff --git a/Engine/Source/Runtime/Experimental/Chaos/Private/Chaos/ContactModification.cpp b/Engine/Source/Runtime/Experimental/Chaos/Private/Chaos/ContactModification.cpp
--- a/Engine/Source/Runtime/Experimental/Chaos/Private/Chaos/ContactModification.cpp
+++ b/Engine/Source/Runtime/Experimental/Chaos/Private/Chaos/ContactModification.cpp
@@ -515,6 +515,36 @@ namespace Chaos
 		SetToEnd();
 	}
 
+	FContactPairModifierParticleRangeIterator FContactPairModifierParticleRange::begin()
+	{
+		return FContactPairModifierParticleRangeIterator(Modifier, Constraints, 0);
+	}
+
+	FContactPairModifierParticleRangeIterator FContactPairModifierParticleRange::end()
+	{
+		return FContactPairModifierParticleRangeIterator(Modifier, Constraints, Constraints.Num());
+	}
+
+	FContactPairModifierParticleRange::FContactPairModifierParticleRange(FCollisionContactModifier* InModifier, FGeometryParticleHandle* InParticle)
+		: Modifier(InModifier)
+		, Particle(InParticle)
+	{
+		FParticleCollisions& ParticleCollisions = Particle->ParticleCollisions();
+		const int32 NumCollisions = ParticleCollisions.Num();
+		Constraints.Reserve(NumCollisions);
+		ParticleCollisions.VisitCollisions([this](FPBDCollisionConstraint& Constraint)
+		{
+			Constraints.Add(&Constraint);
+			return ECollisionVisitorResult::Continue;
+		});
+		ensureMsgf(Constraints.Num() == NumCollisions, TEXT("Number of constraints visited exceeded the number reported by the FParticleCollisions object, and therefore exceeded the reserved constraint cache."));
+	}
+
+	FContactPairModifierParticleRange FCollisionContactModifier::GetContacts(FGeometryParticleHandle* Particle)
+	{
+		return FContactPairModifierParticleRange(this, Particle);
+	}
+
 	TArrayView<FPBDCollisionConstraint* const>& FCollisionContactModifier::GetConstraints()
 	{
 		return Constraints;
diff --git a/Engine/Source/Runtime/Experimental/Chaos/Public/Chaos/ContactModification.h b/Engine/Source/Runtime/Experimental/Chaos/Public/Chaos/ContactModification.h
--- a/Engine/Source/Runtime/Experimental/Chaos/Public/Chaos/ContactModification.h
+++ b/Engine/Source/Runtime/Experimental/Chaos/Public/Chaos/ContactModification.h
@@ -323,6 +323,99 @@ namespace Chaos
 		FCollisionContactModifier* Modifier;
 		FContactPairModifier PairModifier;
 	};
+
+	/*
+	* Iterator for FContactPairModifierParticleRange
+	*/
+	class FContactPairModifierParticleRangeIterator
+	{
+	public:
+
+		FContactPairModifier& operator*()
+		{
+			return PairModifier;
+		}
+
+		FContactPairModifier* operator->()
+		{
+			return &PairModifier;
+		}
+
+		FContactPairModifierParticleRangeIterator& operator++()
+		{
+			++Index;
+
+			PairModifier
+				= IsValid()
+				? FContactPairModifier(Constraints[Index], *Modifier)
+				: FContactPairModifier();
+
+			return *this;
+		}
+
+		bool operator==(const FContactPairModifierParticleRangeIterator& Other)
+		{
+			return
+				&Constraints == &Other.Constraints &&
+				Index == Other.Index;
+		}
+
+		bool operator!=(const FContactPairModifierParticleRangeIterator& Other)
+		{
+			return !(*this == Other);
+		}
+
+		bool IsValid() const
+		{
+			return 0 <= Index && Index < Constraints.Num();
+		}
+
+	private:
+		FContactPairModifierParticleRangeIterator(FCollisionContactModifier* InModifier, TArray<FPBDCollisionConstraint*>& InConstraints, int32 InIndex)
+			: Modifier(InModifier)
+			, Constraints(InConstraints)
+			, Index(InIndex)
+		{
+			PairModifier
+				= IsValid()
+				? FContactPairModifier(Constraints[Index], *Modifier)
+				: FContactPairModifier();
+		}
+
+		FCollisionContactModifier* Modifier;
+		FContactPairModifier PairModifier;
+
+		TArray<FPBDCollisionConstraint*>& Constraints;
+		int32 Index;
+
+		// Befriend the range object so that it can create iterators
+		friend class FContactPairModifierParticleRange;
+	};
+
+	/*
+	* Interface for iterating over a range of contacts for a particular particle
+	*/
+	class FContactPairModifierParticleRange
+	{
+	public:
+		CHAOS_API FContactPairModifierParticleRangeIterator begin();
+		CHAOS_API FContactPairModifierParticleRangeIterator end();
+
+		FContactPairModifierParticleRange(FContactPairModifierParticleRange&& Other)
+			: Modifier(Other.Modifier)
+			, Particle(Other.Particle)
+			, Constraints(MoveTemp(Other.Constraints))
+		{ }
+
+	private:
+		FContactPairModifierParticleRange(FCollisionContactModifier* InModifier, FGeometryParticleHandle* InParticle);
+		FCollisionContactModifier* Modifier;
+		FGeometryParticleHandle* Particle;
+		TArray<FPBDCollisionConstraint*> Constraints;
+
+		// Befriend the modifier so that it can create ranges
+		friend class FCollisionContactModifier;
+	};
 	
 	/*
 	*  Provides interface for iterating over modifiable contact pairs
@@ -346,6 +439,8 @@ namespace Chaos
 		FContactPairModifierIterator End() const { return FContactPairModifierIterator(); }
 		FContactPairModifierIterator end() const { return FContactPairModifierIterator(); }
 
+		CHAOS_API FContactPairModifierParticleRange GetContacts(FGeometryParticleHandle* Particle);
+
 	private:
 		CHAOS_API TArrayView<FPBDCollisionConstraint* const>& GetConstraints();
 		CHAOS_API void DisableConstraint(FPBDCollisionConstraint& Constraint);
