Fix Headless Chaos compilation error

diff --git a/Engine/Source/Programs/HeadlessChaos/Private/GeometryCollection/GeometryCollectionTestEvents.cpp b/Engine/Source/Programs/HeadlessChaos/Private/GeometryCollection/GeometryCollectionTestEvents.cpp
--- a/Engine/Source/Programs/HeadlessChaos/Private/GeometryCollection/GeometryCollectionTestEvents.cpp
+++ b/Engine/Source/Programs/HeadlessChaos/Private/GeometryCollection/GeometryCollectionTestEvents.cpp
@@ -26,6 +26,10 @@ namespace GeometryCollectionTest
 			return (Data1 == Other.Data1) && (Data2 == Other.Data2);
 		}
 
+		void Reset()
+		{
+		}
+
 		int Data1;
 		FVector Data2;
 	};
