Chaos - fix capsule-convex collision detection

When the convex hits an end cap and we generate an additional contact on the capsule's cylinder, it would create an additional contact even if the convex plane was nearly opposite the capsule axis, leading to an unnecessary point and a potential NaN when exact opposite.

diff --git a/Engine/Source/Programs/HeadlessChaos/Private/HeadlessChaosTestDetectCollision.cpp b/Engine/Source/Programs/HeadlessChaos/Private/HeadlessChaosTestDetectCollision.cpp
--- a/Engine/Source/Programs/HeadlessChaos/Private/HeadlessChaosTestDetectCollision.cpp
+++ b/Engine/Source/Programs/HeadlessChaos/Private/HeadlessChaosTestDetectCollision.cpp
@@ -6,6 +6,7 @@
 #include "Chaos/Collision/CapsuleConvexContactPoint.h"
 #include "Chaos/Collision/ContactPoint.h"
 #include "Chaos/Collision/SphereConvexContactPoint.h"
+#include "Chaos/CollisionOneShotManifolds.h"
 #include "Chaos/GJK.h"
 #include "Chaos/Sphere.h"
 #include "Chaos/Utilities.h"
@@ -322,4 +323,59 @@ namespace ChaosTest
 		}
 	}
 
+	// Capsule-Convex OneShot Manifold when the convex lands on the end of the capsule and the convex face points exactly down the axis
+	// This exposed a bug in ConstructCapsuleConvexOneShotManifold where it would try to add a cylinder face contact resulting in a NaN 
+	// from the call to GetUnsafeNormal.
+	GTEST_TEST(DetectCollisionTests, TestCapsuleConvex_Manifold_OpposingFaceNormalAxis)
+	{
+		const FRealSingle CapsuleRadius = 50.0f;
+		const FRealSingle CapsuleSegmentLength = 100.0f;
+		const FVec3f BoxSize = FVec3(100, 100, 100);
+		const FRealSingle BoxMargin = 0.0f;
+		const FRealSingle CullDistance = 10.0f;
+
+		const FImplicitCapsule3 Capsule(FVec3(0.0f, 0.0f, -0.5f * CapsuleSegmentLength), FVec3(0.0f, 0.0f, 0.5f * CapsuleSegmentLength), CapsuleRadius);
+		const FImplicitConvex3 Convex = CreateConvexBox(-0.5f * BoxSize, 0.5f * BoxSize, BoxMargin);
+
+		{
+			const FVec3f CapsulePos = FVec3(0, 0, -50);
+			const FVec3f ConvexPos = FVec3(0, 0, 100);
+			const FRigidTransform3 CapsuleTransform = FRigidTransform3(CapsulePos, FRotation3::FromIdentity());
+			const FRigidTransform3 ConvexTransform = FRigidTransform3(ConvexPos, FRotation3::FromIdentity());
+
+			FContactPointManifold ManifoldPoints;
+			Collisions::ConstructCapsuleConvexOneShotManifold(Capsule, CapsuleTransform, Convex, ConvexTransform, CullDistance, ManifoldPoints);
+
+			EXPECT_EQ(ManifoldPoints.Num(), 1);
+			if (ManifoldPoints.Num() == 1)
+			{
+				EXPECT_EQ(ManifoldPoints[0].ContactType, EContactPointType::VertexPlane);
+				EXPECT_NEAR(ManifoldPoints[0].ShapeContactNormal.Z, -1.0f, UE_KINDA_SMALL_NUMBER);
+				EXPECT_NEAR(ManifoldPoints[0].ShapeContactPoints[0].Z, 100.0f, UE_KINDA_SMALL_NUMBER);
+				EXPECT_NEAR(ManifoldPoints[0].ShapeContactPoints[1].Z, -50.0f, UE_KINDA_SMALL_NUMBER);
+				EXPECT_NEAR(ManifoldPoints[0].Phi, 0.0f, UE_KINDA_SMALL_NUMBER);
+			}
+		}
+
+		{
+			const FVec3f CapsulePos = FVec3(0, 0, 50);
+			const FVec3f ConvexPos = FVec3(0, 0, -100);
+			const FRigidTransform3 CapsuleTransform = FRigidTransform3(CapsulePos, FRotation3::FromIdentity());
+			const FRigidTransform3 ConvexTransform = FRigidTransform3(ConvexPos, FRotation3::FromIdentity());
+
+			FContactPointManifold ManifoldPoints;
+			Collisions::ConstructCapsuleConvexOneShotManifold(Capsule, CapsuleTransform, Convex, ConvexTransform, CullDistance, ManifoldPoints);
+
+			EXPECT_EQ(ManifoldPoints.Num(), 1);
+			if (ManifoldPoints.Num() == 1)
+			{
+				EXPECT_EQ(ManifoldPoints[0].ContactType, EContactPointType::VertexPlane);
+				EXPECT_NEAR(ManifoldPoints[0].ShapeContactNormal.Z, 1.0f, UE_KINDA_SMALL_NUMBER);
+				EXPECT_NEAR(ManifoldPoints[0].ShapeContactPoints[0].Z, -100.0f, UE_KINDA_SMALL_NUMBER);
+				EXPECT_NEAR(ManifoldPoints[0].ShapeContactPoints[1].Z, 50.0f, UE_KINDA_SMALL_NUMBER);
+				EXPECT_NEAR(ManifoldPoints[0].Phi, 0.0f, UE_KINDA_SMALL_NUMBER);
+			}
+		}
+	}
+
 }
\ No newline at end of file
diff --git a/Engine/Source/Runtime/Experimental/Chaos/Private/Chaos/CollisionOneShotManifolds.cpp b/Engine/Source/Runtime/Experimental/Chaos/Private/Chaos/CollisionOneShotManifolds.cpp
--- a/Engine/Source/Runtime/Experimental/Chaos/Private/Chaos/CollisionOneShotManifolds.cpp
+++ b/Engine/Source/Runtime/Experimental/Chaos/Private/Chaos/CollisionOneShotManifolds.cpp
@@ -1162,7 +1162,7 @@ namespace Chaos
 				// but only if the cylinder is not standing vertically on the face
 				const FReal CapsuleCylinderTolerance = FReal(0.707);	// about 45 deg
 				const FReal CapsuleAxisDotConvexNormal = FVec3::DotProduct(ConvexPlaneNormal, CapsuleAxisSegment.GetAxis());
-				if (CapsuleAxisDotConvexNormal < CapsuleCylinderTolerance)
+				if (FMath::Abs(CapsuleAxisDotConvexNormal) < CapsuleCylinderTolerance)
 				{
 					// The line segment on the surface of the cylinder closest to the contact
 					const FVec3 CapsuleCylinderNormal = (ConvexPlaneNormal - CapsuleAxisDotConvexNormal * CapsuleAxisSegment.GetAxis()).GetUnsafeNormal();
