Chaos - fix headless chaos collision test

It was testing that a box gets pushed out of a plane, but the defaulrt depenetration velocity is now zero

diff --git a/Engine/Source/Programs/HeadlessChaos/Private/HeadlessChaosTestCollisions.cpp b/Engine/Source/Programs/HeadlessChaos/Private/HeadlessChaosTestCollisions.cpp
--- a/Engine/Source/Programs/HeadlessChaos/Private/HeadlessChaosTestCollisions.cpp
+++ b/Engine/Source/Programs/HeadlessChaos/Private/HeadlessChaosTestCollisions.cpp
@@ -155,6 +155,8 @@ namespace ChaosTest {
 		Box->AuxilaryValue(PhysicsMaterials) = MakeSerializable(PhysicsMaterial);
 		Box->UpdateWorldSpaceState(FRigidTransform3(Box->P(), Box->Q()), FVec3(0));
 
+		const FReal InitialBoxZ = Box->X().Z;
+
 		FPBDCollisionConstraintAccessor Collisions(Particles, Collided, PhysicsMaterials, PerParticlePhysicsMaterials, 2, 5);
 
 		Collisions.ComputeConstraints(Dt);
@@ -199,7 +201,9 @@ namespace ChaosTest {
 
 		Collisions.ScatterOutput(Dt);
 
-		EXPECT_NEAR(Box->P().Z, 0.5f, 1.e-2f);
+		// Box will not move because the default depentration velocity is zero
+		const FReal ExpectedBoxZ = InitialBoxZ + Collisions.CollisionConstraints.GetSolverSettings().DepenetrationVelocity * Dt;
+		EXPECT_NEAR(Box->P().Z, ExpectedBoxZ, 1.e-2f);
 
 		// Velocity is below the restitution threshold, so expecting 0 velocity despite the fact that restitution is 1
 		EXPECT_TRUE(Box->V().Equals(FVec3(0)));
@@ -307,6 +311,8 @@ namespace ChaosTest {
 		Box->AuxilaryValue(PhysicsMaterials) = MakeSerializable(PhysicsMaterial);
 		Box->UpdateWorldSpaceState(FRigidTransform3(Box->P(), Box->Q()), FVec3(0));
 
+		FReal InitialBoxZ = Box->X().Z;
+
 		FPBDCollisionConstraintAccessor Collisions(Particles, Collided, PhysicsMaterials, PerParticlePhysicsMaterials, 2, 5);
 
 		Collisions.ComputeConstraints(Dt);
@@ -344,7 +350,9 @@ namespace ChaosTest {
 		EXPECT_TRUE(Box->V().Equals(FVec3(0)));
 		EXPECT_TRUE(Box->W().Equals(FVec3(0)));
 
-		EXPECT_TRUE(FVec3::IsNearlyEqual(Box->P(), FVector(0.f, 1.f, 0.5f), 1.e-2f));
+		// Box will not move because the default depentration velocity is zero
+		const FReal ExpectedBoxZ = InitialBoxZ + Collisions.CollisionConstraints.GetSolverSettings().DepenetrationVelocity * Dt;
+		EXPECT_TRUE(FMath::IsNearlyEqual(Box->P().Z, ExpectedBoxZ, 1.e-2));
 	}
 
 	void CollisionBoxPlaneRestitution()
