Improve Geometry Collection Transform API
Transform attribute is no longer public, methods allow accessibility

diff --git a/Engine/Source/Programs/HeadlessChaos/Private/GeometryCollection/GeometryCollectionTestClustering.cpp b/Engine/Source/Programs/HeadlessChaos/Private/GeometryCollection/GeometryCollectionTestClustering.cpp
--- a/Engine/Source/Programs/HeadlessChaos/Private/GeometryCollection/GeometryCollectionTestClustering.cpp
+++ b/Engine/Source/Programs/HeadlessChaos/Private/GeometryCollection/GeometryCollectionTestClustering.cpp
@@ -92,8 +92,7 @@ namespace GeometryCollectionTest
 		UnitTest.AddSimulationObject(Collection);
 		UnitTest.Initialize();
 
-		TManagedArray<FTransform3f>& Transform = Collection->DynamicCollection->Transform;
-		FReal StartingRigidDistance = (Transform[1].GetTranslation() - Transform[0].GetTranslation()).Size(), CurrentRigidDistance = 0.f;
+		FReal StartingRigidDistance = (Collection->DynamicCollection->GetTransform(1).GetTranslation() - Collection->DynamicCollection->GetTransform(0).GetTranslation()).Size(), CurrentRigidDistance = 0.f;
 
 		TManagedArray<bool>& Active = Collection->DynamicCollection->Active;
 
@@ -119,9 +118,9 @@ namespace GeometryCollectionTest
 			EXPECT_FALSE(Active[1]);
 			EXPECT_TRUE(Active[2]);
 
-			CurrentRigidDistance = (Transform[1].GetTranslation() - Transform[0].GetTranslation()).Size();
+			CurrentRigidDistance = (Collection->DynamicCollection->GetTransform(1).GetTranslation() - Collection->DynamicCollection->GetTransform(0).GetTranslation()).Size();
 			EXPECT_LT(FMath::Abs(CurrentRigidDistance - StartingRigidDistance), SMALL_NUMBER); // two bodies under cluster maintain distance
-			EXPECT_LT(Collection->DynamicCollection->Transform[2].GetTranslation().Z, InitialZ); // body should be falling and decreasing in Z			
+			EXPECT_LT(Collection->DynamicCollection->GetTransform(2).GetTranslation().Z, InitialZ); // body should be falling and decreasing in Z			
 		}
 
 		EXPECT_TRUE(ClusterMapContains(ClusterMap, Collection->PhysObject->GetSolverClusterHandles()[0],
@@ -167,8 +166,7 @@ namespace GeometryCollectionTest
 		UnitTest.AddSimulationObject(Collection);
 		UnitTest.Initialize();		
 
-		TManagedArray<FTransform3f>& Transform = Collection->DynamicCollection->Transform;
-		FReal StartingRigidDistance = (Transform[1].GetTranslation() - Transform[0].GetTranslation()).Size(), CurrentRigidDistance = 0.f;
+		FReal StartingRigidDistance = (Collection->DynamicCollection->GetTransform(1).GetTranslation() - Collection->DynamicCollection->GetTransform(0).GetTranslation()).Size(), CurrentRigidDistance = 0.f;
 
 		UnitTest.Advance();
 
@@ -282,8 +280,7 @@ namespace GeometryCollectionTest
 		UnitTest.AddSimulationObject(Collection);
 		UnitTest.Initialize();
 
-		TManagedArray<FTransform3f>& Transform = Collection->DynamicCollection->Transform;
-		FReal StartingRigidDistance = (Transform[1].GetTranslation() - Transform[0].GetTranslation()).Size(), CurrentRigidDistance = 0.f;
+		FReal StartingRigidDistance = (Collection->DynamicCollection->GetTransform(1).GetTranslation() - Collection->DynamicCollection->GetTransform(0).GetTranslation()).Size(), CurrentRigidDistance = 0.f;
 
 		UnitTest.Advance();
 
@@ -406,8 +403,7 @@ namespace GeometryCollectionTest
 
 		Collection->PhysObject->SetCollisionParticlesPerObjectFraction(1.0);
 
-		TManagedArray<FTransform3f>& Transform = Collection->DynamicCollection->Transform;
-		FReal StartingRigidDistance = (Transform[1].GetTranslation() - Transform[0].GetTranslation()).Size(), CurrentRigidDistance = 0.f;
+		FReal StartingRigidDistance = (Collection->DynamicCollection->GetTransform(1).GetTranslation() - Collection->DynamicCollection->GetTransform(0).GetTranslation()).Size(), CurrentRigidDistance = 0.f;
 
 		UnitTest.Advance();
 
@@ -429,7 +425,7 @@ namespace GeometryCollectionTest
 		{
 			UnitTest.Advance();
 
-			CurrentRigidDistance = (Transform[1].GetTranslation() - Transform[0].GetTranslation()).Size();
+			CurrentRigidDistance = (Collection->DynamicCollection->GetTransform(1).GetTranslation() - Collection->DynamicCollection->GetTransform(0).GetTranslation()).Size();
 
 			if ((BrokenFrame == INDEX_NONE) && !ParticleHandles[2]->Disabled())
 			{
@@ -507,8 +503,7 @@ namespace GeometryCollectionTest
 		UnitTest.AddSimulationObject(Collection);
 		UnitTest.Initialize();
 		
-		TManagedArray<FTransform3f>& Transform = Collection->DynamicCollection->Transform;
-		FReal StartingRigidDistance = (Transform[1].GetTranslation() - Transform[0].GetTranslation()).Size(), CurrentRigidDistance = 0.f;
+		FReal StartingRigidDistance = (Collection->DynamicCollection->GetTransform(1).GetTranslation() - Collection->DynamicCollection->GetTransform(0).GetTranslation()).Size(), CurrentRigidDistance = 0.f;
 
 		UnitTest.Advance();
 
@@ -531,7 +526,7 @@ namespace GeometryCollectionTest
 		{
 			UnitTest.Advance();
 
-			CurrentRigidDistance = (Transform[1].GetTranslation() - Transform[0].GetTranslation()).Size();
+			CurrentRigidDistance = (Collection->DynamicCollection->GetTransform(1).GetTranslation() - Collection->DynamicCollection->GetTransform(0).GetTranslation()).Size();
 
 			if (Conditions[0]==false)
 			{
@@ -635,7 +630,8 @@ namespace GeometryCollectionTest
 		EXPECT_EQ(ClusterMap.Num(), 2);
 		EXPECT_TRUE(ClusterMapContains(ClusterMap, Collection1Handles[2], { Collection1Handles[1],Collection1Handles[0] }));
 		EXPECT_TRUE(ClusterMapContains(ClusterMap, Collection1Handles[3], { Collection1Handles[2] }));
-		TArray<FTransform> Collection1_PreReleaseTM; GeometryCollectionAlgo::GlobalMatrices(DynamicCollection1->Transform, DynamicCollection1->Parent, Collection1_PreReleaseTM);
+		const TManagedArray<FTransform3f> Transforms1(DynamicCollection1->GetTransforms());
+		TArray<FTransform> Collection1_PreReleaseTM; GeometryCollectionAlgo::GlobalMatrices(Transforms1, DynamicCollection1->Parent, Collection1_PreReleaseTM);
 		for (int Idx = 0; Idx < Collection1_PreReleaseTM.Num(); Idx++) {
 			EXPECT_TRUE( (Collection1_PreReleaseTM[Idx].GetTranslation()-Collection1_InitialTM[Idx].GetTranslation()).Size()<KINDA_SMALL_NUMBER);
 		}
@@ -646,7 +642,8 @@ namespace GeometryCollectionTest
 		EXPECT_EQ(SovlerParticleHandles.Size(), 4);
 		EXPECT_EQ(ClusterMap.Num(), 1);
 		EXPECT_TRUE(ClusterMapContains(ClusterMap, Collection1Handles[2], { Collection1Handles[1],Collection1Handles[0] }));
-		TArray<FTransform> Collection1_PostReleaseTM; GeometryCollectionAlgo::GlobalMatrices(DynamicCollection1->Transform, DynamicCollection1->Parent, Collection1_PostReleaseTM);
+		const TManagedArray<FTransform3f> Transforms2(DynamicCollection1->GetTransforms());
+		TArray<FTransform> Collection1_PostReleaseTM; GeometryCollectionAlgo::GlobalMatrices(Transforms2, DynamicCollection1->Parent, Collection1_PostReleaseTM);
 		for (int Idx = 0; Idx < Collection1_PostReleaseTM.Num(); Idx++) {
 			EXPECT_TRUE((Collection1_PostReleaseTM[Idx].GetTranslation() - Collection1_InitialTM[Idx].GetTranslation()).Size() < KINDA_SMALL_NUMBER);
 		}
@@ -656,7 +653,8 @@ namespace GeometryCollectionTest
 
 		EXPECT_EQ(SovlerParticleHandles.Size(), 4);
 		EXPECT_EQ(ClusterMap.Num(), 0);
-		TArray<FTransform> Collection1_PostRelease2TM; GeometryCollectionAlgo::GlobalMatrices(DynamicCollection1->Transform, DynamicCollection1->Parent, Collection1_PostRelease2TM);
+		const TManagedArray<FTransform3f> Transforms3(DynamicCollection1->GetTransforms());
+		TArray<FTransform> Collection1_PostRelease2TM; GeometryCollectionAlgo::GlobalMatrices(Transforms3, DynamicCollection1->Parent, Collection1_PostRelease2TM);
 		for (int Idx = 0; Idx < Collection1_PostRelease2TM.Num(); Idx++) {
 			EXPECT_TRUE((Collection1_PostRelease2TM[Idx].GetTranslation() - Collection1_InitialTM[Idx].GetTranslation()).Size() < KINDA_SMALL_NUMBER);
 		}
@@ -699,8 +697,7 @@ namespace GeometryCollectionTest
 		UnitTest.AddSimulationObject(Collection);
 		UnitTest.Initialize();
 
-		TManagedArray<FTransform3f>& Transform = Collection->DynamicCollection->Transform;
-		FReal StartingRigidDistance = (Transform[1].GetTranslation() - Transform[0].GetTranslation()).Size(), CurrentRigidDistance = 0.f;
+		FReal StartingRigidDistance = (Collection->DynamicCollection->GetTransform(1).GetTranslation() - Collection->DynamicCollection->GetTransform(0).GetTranslation()).Size(), CurrentRigidDistance = 0.f;
 
 		TArray<bool> Conditions = { false,false,false,false };
 
@@ -721,7 +718,7 @@ namespace GeometryCollectionTest
 		{
 			UnitTest.Advance();			
 
-			CurrentRigidDistance = (Transform[1].GetTranslation() - Transform[0].GetTranslation()).Size();
+			CurrentRigidDistance = (Collection->DynamicCollection->GetTransform(1).GetTranslation() - Collection->DynamicCollection->GetTransform(0).GetTranslation()).Size();
 			
 			if (Conditions[0] == false)
 			{
@@ -860,8 +857,7 @@ namespace GeometryCollectionTest
 		UnitTest.AddSimulationObject(Collection);
 		UnitTest.Initialize();
 
-		TManagedArray<FTransform3f>& Transform = Collection->DynamicCollection->Transform;
-		FReal StartingRigidDistance = (Transform[1].GetTranslation() - Transform[0].GetTranslation()).Size();
+		FReal StartingRigidDistance = (Collection->DynamicCollection->GetTransform(1).GetTranslation() - Collection->DynamicCollection->GetTransform(0).GetTranslation()).Size();
 		FReal CurrentRigidDistance = 0;
 
 		// Staged conditions
@@ -1159,8 +1155,7 @@ namespace GeometryCollectionTest
 		UnitTest.AddSimulationObject(Collection);
 		UnitTest.Initialize();
 
-		TManagedArray<FTransform3f>& Transform = Collection->DynamicCollection->Transform;
-		FReal StartingRigidDistance = (Transform[1].GetTranslation() - Transform[0].GetTranslation()).Size();
+		FReal StartingRigidDistance = (Collection->DynamicCollection->GetTransform(1).GetTranslation() - Collection->DynamicCollection->GetTransform(0).GetTranslation()).Size();
 		FReal CurrentRigidDistance = 0.f;
 
 		// Staged conditions
@@ -1446,9 +1441,6 @@ namespace GeometryCollectionTest
 
 
 		TArray<FReal> Distances;
-		TManagedArray<FTransform3f>& Transform = DynamicCollection->Transform;
-		TManagedArray<FTransform3f>& Transform2 = DynamicCollection2->Transform;
-
 		auto& Clustering = UnitTest.Solver->GetEvolution()->GetRigidClustering();
 		const auto& ClusterMap = Clustering.GetChildrenMap();
 		const auto& ParticleHandles = UnitTest.Solver->GetParticles().GetParticleHandles();
@@ -1468,10 +1460,10 @@ namespace GeometryCollectionTest
 			if (Frame == 0)
 			{
 				TArray<FTransform> GlobalTransform;
-				GeometryCollectionAlgo::GlobalMatrices(DynamicCollection->Transform, DynamicCollection->Parent, GlobalTransform);
+				GeometryCollectionAlgo::GlobalMatrices(TManagedArray<FTransform3f>(DynamicCollection->GetTransforms()), DynamicCollection->Parent, GlobalTransform);
 
 				TArray<FTransform> GlobalTransform2;
-				GeometryCollectionAlgo::GlobalMatrices(DynamicCollection2->Transform, DynamicCollection2->Parent, GlobalTransform2);
+				GeometryCollectionAlgo::GlobalMatrices(TManagedArray<FTransform3f>(DynamicCollection2->GetTransforms()), DynamicCollection2->Parent, GlobalTransform2);
 
 				// build relative transforms distances
 				for (int32 i = 0; i < (int32)GlobalTransform.Num()-1; i++)
@@ -1490,10 +1482,10 @@ namespace GeometryCollectionTest
 
 		
 		TArray<FTransform> GlobalTransform;
-		GeometryCollectionAlgo::GlobalMatrices(DynamicCollection->Transform, DynamicCollection->Parent, GlobalTransform);
+		GeometryCollectionAlgo::GlobalMatrices(TManagedArray<FTransform3f>(DynamicCollection->GetTransforms()), DynamicCollection->Parent, GlobalTransform);
 
 		TArray<FTransform> GlobalTransform2;
-		GeometryCollectionAlgo::GlobalMatrices(DynamicCollection2->Transform, DynamicCollection2->Parent, GlobalTransform2);
+		GeometryCollectionAlgo::GlobalMatrices(TManagedArray<FTransform3f>(DynamicCollection2->GetTransforms()), DynamicCollection2->Parent, GlobalTransform2);
 
 		// build relative transforms distances
 		TArray<FReal> Distances2;
@@ -1553,9 +1545,6 @@ namespace GeometryCollectionTest
 
 
 		TArray<FReal> Distances;
-		TManagedArray<FTransform3f>& Transform = DynamicCollection->Transform;
-		TManagedArray<FTransform3f>& Transform2 = DynamicCollection2->Transform;
-
 
 		auto& Clustering = UnitTest.Solver->GetEvolution()->GetRigidClustering();
 		const auto& ClusterMap = Clustering.GetChildrenMap();
@@ -1569,10 +1558,10 @@ namespace GeometryCollectionTest
 		});
 
 		TArray<FTransform> PrevGlobalTransform;
-		GeometryCollectionAlgo::GlobalMatrices(DynamicCollection->Transform, DynamicCollection->Parent, PrevGlobalTransform);
+		GeometryCollectionAlgo::GlobalMatrices(TManagedArray<FTransform3f>(DynamicCollection->GetTransforms()), DynamicCollection->Parent, PrevGlobalTransform);
 
 		TArray<FTransform> PrevGlobalTransform2;
-		GeometryCollectionAlgo::GlobalMatrices(DynamicCollection2->Transform, DynamicCollection2->Parent, PrevGlobalTransform2);
+		GeometryCollectionAlgo::GlobalMatrices(TManagedArray<FTransform3f>(DynamicCollection2->GetTransforms()), DynamicCollection2->Parent, PrevGlobalTransform2);
 
 
 		for (int Frame = 0; Frame < 100; Frame++)
@@ -1580,10 +1569,10 @@ namespace GeometryCollectionTest
 			UnitTest.Advance();
 
 			TArray<FTransform> GlobalTransform;
-			GeometryCollectionAlgo::GlobalMatrices(DynamicCollection->Transform, DynamicCollection->Parent, GlobalTransform);
+			GeometryCollectionAlgo::GlobalMatrices(TManagedArray<FTransform3f>(DynamicCollection->GetTransforms()), DynamicCollection->Parent, GlobalTransform);
 
 			TArray<FTransform> GlobalTransform2;
-			GeometryCollectionAlgo::GlobalMatrices(DynamicCollection2->Transform, DynamicCollection2->Parent, GlobalTransform2);
+			GeometryCollectionAlgo::GlobalMatrices(TManagedArray<FTransform3f>(DynamicCollection2->GetTransforms()), DynamicCollection2->Parent, GlobalTransform2);
 
 			EXPECT_EQ(ClusterMap.Num(), 1);
 			EXPECT_TRUE(ClusterMapContains(ClusterMap, ParticleHandles.Handle(6)->CastToRigidParticle(), { ParticleHandles.Handle(1)->CastToRigidParticle(),ParticleHandles.Handle(0)->CastToRigidParticle(),ParticleHandles.Handle(3)->CastToRigidParticle(),ParticleHandles.Handle(4)->CastToRigidParticle() }));
@@ -1656,10 +1645,6 @@ namespace GeometryCollectionTest
 		TSharedPtr<FGeometryDynamicCollection> DynamicCollection = Collection->DynamicCollection;
 		TSharedPtr<FGeometryDynamicCollection> DynamicCollection2 = Collection2->DynamicCollection;
 
-		TArray<FReal> Distances;
-		TManagedArray<FTransform3f>& Transform = DynamicCollection->Transform;
-		TManagedArray<FTransform3f>& Transform2 = DynamicCollection2->Transform;
-
 		auto& Clustering = UnitTest.Solver->GetEvolution()->GetRigidClustering();
 		const auto& ClusterMap = Clustering.GetChildrenMap();
 		const auto& ParticleHandles = UnitTest.Solver->GetParticles().GetParticleHandles();
@@ -1714,10 +1699,10 @@ namespace GeometryCollectionTest
 				FTransform RootTransform(Root->R(), Root->X());
 
 				TArray<FTransform> GlobalTransform1;
-				GeometryCollectionAlgo::GlobalMatrices(DynamicCollection->Transform, DynamicCollection->Parent, GlobalTransform1);
+				GeometryCollectionAlgo::GlobalMatrices(TManagedArray<FTransform3f>(DynamicCollection->GetTransforms()), DynamicCollection->Parent, GlobalTransform1);
 
 				TArray<FTransform> GlobalTransform2;
-				GeometryCollectionAlgo::GlobalMatrices(DynamicCollection2->Transform, DynamicCollection2->Parent, GlobalTransform2);
+				GeometryCollectionAlgo::GlobalMatrices(TManagedArray<FTransform3f>(DynamicCollection2->GetTransforms()), DynamicCollection2->Parent, GlobalTransform2);
 
 				EXPECT_TRUE(!InitialRootPosition.Equals(Root->X())); // root moves
 
@@ -1790,8 +1775,8 @@ namespace GeometryCollectionTest
 		TSharedPtr<FGeometryDynamicCollection> DynamicCollection1 = Collection1->DynamicCollection;
 		TSharedPtr<FGeometryDynamicCollection> DynamicCollection2 = Collection2->DynamicCollection;
 
-		TArray<FTransform> Collection1_InitialTM; GeometryCollectionAlgo::GlobalMatrices(DynamicCollection1->Transform, DynamicCollection1->Parent, Collection1_InitialTM);
-		TArray<FTransform> Collection2_InitialTM; GeometryCollectionAlgo::GlobalMatrices(DynamicCollection2->Transform, DynamicCollection2->Parent, Collection2_InitialTM);
+		TArray<FTransform> Collection1_InitialTM; GeometryCollectionAlgo::GlobalMatrices(TManagedArray<FTransform3f>(DynamicCollection1->GetTransforms()), DynamicCollection1->Parent, Collection1_InitialTM);
+		TArray<FTransform> Collection2_InitialTM; GeometryCollectionAlgo::GlobalMatrices(TManagedArray<FTransform3f>(DynamicCollection2->GetTransforms()), DynamicCollection2->Parent, Collection2_InitialTM);
 
 		TArray<Chaos::TPBDRigidClusteredParticleHandle<FReal, 3>*>& Collection1Handles = Collection1->PhysObject->GetSolverParticleHandles();
 		TArray<Chaos::TPBDRigidClusteredParticleHandle<FReal, 3>*>& Collection2Handles = Collection2->PhysObject->GetSolverParticleHandles();
@@ -1825,8 +1810,8 @@ namespace GeometryCollectionTest
 		EXPECT_TRUE(ClusterMapContains(ClusterMap, ClusterUnion.InternalCluster, { Collection1Handles[0], Collection1Handles[1], Collection2Handles[0], Collection2Handles[1] }));
 		EXPECT_EQ(ClusterMap[ClusterUnion.InternalCluster].Num(), 4);
 
-		TArray<FTransform> Collection1_PreReleaseTM; GeometryCollectionAlgo::GlobalMatrices(DynamicCollection1->Transform, DynamicCollection1->Parent, Collection1_PreReleaseTM);
-		TArray<FTransform> Collection2_PreReleaseTM; GeometryCollectionAlgo::GlobalMatrices(DynamicCollection2->Transform, DynamicCollection2->Parent, Collection2_PreReleaseTM);
+		TArray<FTransform> Collection1_PreReleaseTM; GeometryCollectionAlgo::GlobalMatrices(TManagedArray<FTransform3f>(DynamicCollection1->GetTransforms()), DynamicCollection1->Parent, Collection1_PreReleaseTM);
+		TArray<FTransform> Collection2_PreReleaseTM; GeometryCollectionAlgo::GlobalMatrices(TManagedArray<FTransform3f>(DynamicCollection2->GetTransforms()), DynamicCollection2->Parent, Collection2_PreReleaseTM);
 		for (int Idx = 0; Idx < Collection1_PreReleaseTM.Num() - 1; Idx++) {
 			EXPECT_LT(Collection1_PreReleaseTM[Idx].GetTranslation().Z, Collection1_InitialTM[Idx].GetTranslation().Z);
 			EXPECT_LT(Collection2_PreReleaseTM[Idx].GetTranslation().Z, Collection2_InitialTM[Idx].GetTranslation().Z);
@@ -1850,8 +1835,8 @@ namespace GeometryCollectionTest
 
 		EXPECT_EQ(SovlerParticleHandles.Size(), 7);
 
-		TArray<FTransform> Collection1_PostReleaseTM; GeometryCollectionAlgo::GlobalMatrices(DynamicCollection1->Transform, DynamicCollection1->Parent, Collection1_PostReleaseTM);
-		TArray<FTransform> Collection2_PostReleaseTM; GeometryCollectionAlgo::GlobalMatrices(DynamicCollection2->Transform, DynamicCollection2->Parent, Collection2_PostReleaseTM);
+		TArray<FTransform> Collection1_PostReleaseTM; GeometryCollectionAlgo::GlobalMatrices(TManagedArray<FTransform3f>(DynamicCollection1->GetTransforms()), DynamicCollection1->Parent, Collection1_PostReleaseTM);
+		TArray<FTransform> Collection2_PostReleaseTM; GeometryCollectionAlgo::GlobalMatrices(TManagedArray<FTransform3f>(DynamicCollection2->GetTransforms()), DynamicCollection2->Parent, Collection2_PostReleaseTM);
 		for (int Idx = 0; Idx < Collection1_PostReleaseTM.Num() - 1; Idx++) {
 			EXPECT_LT(Collection1_PostReleaseTM[Idx].GetTranslation().Z, Collection1_PreReleaseTM[Idx].GetTranslation().Z);
 			EXPECT_LT(Collection2_PostReleaseTM[Idx].GetTranslation().Z, Collection2_PreReleaseTM[Idx].GetTranslation().Z);
@@ -1906,8 +1891,8 @@ namespace GeometryCollectionTest
 		Chaos::FClusterUnionManager& ClusterUnionManager = Clustering.GetClusterUnionManager();
 		const auto& ClusterMap = Clustering.GetChildrenMap();
 
-		TArray<FTransform> Collection1_InitialTM; GeometryCollectionAlgo::GlobalMatrices(DynamicCollection1->Transform, DynamicCollection1->Parent, Collection1_InitialTM);
-		TArray<FTransform> Collection2_InitialTM; GeometryCollectionAlgo::GlobalMatrices(DynamicCollection2->Transform, DynamicCollection2->Parent, Collection2_InitialTM);
+		TArray<FTransform> Collection1_InitialTM; GeometryCollectionAlgo::GlobalMatrices(TManagedArray<FTransform3f>(DynamicCollection1->GetTransforms()), DynamicCollection1->Parent, Collection1_InitialTM);
+		TArray<FTransform> Collection2_InitialTM; GeometryCollectionAlgo::GlobalMatrices(TManagedArray<FTransform3f>(DynamicCollection2->GetTransforms()), DynamicCollection2->Parent, Collection2_InitialTM);
 
 		TArray<Chaos::TPBDRigidClusteredParticleHandle<FReal, 3>*>& Collection1Handles = Collection1->PhysObject->GetSolverParticleHandles();
 		TArray<Chaos::TPBDRigidClusteredParticleHandle<FReal, 3>*>& Collection2Handles = Collection2->PhysObject->GetSolverParticleHandles();
@@ -1940,8 +1925,8 @@ namespace GeometryCollectionTest
 		EXPECT_TRUE(ClusterMapContains(ClusterMap, ClusterUnion.InternalCluster, { Collection1Handles[0], Collection1Handles[1], Collection2Handles[0], Collection2Handles[1] }));
 		EXPECT_EQ(ClusterMap[ClusterUnion.InternalCluster].Num(), 4);
 
-		TArray<FTransform> Collection1_PreReleaseTM; GeometryCollectionAlgo::GlobalMatrices(DynamicCollection1->Transform, DynamicCollection1->Parent, Collection1_PreReleaseTM);
-		TArray<FTransform> Collection2_PreReleaseTM; GeometryCollectionAlgo::GlobalMatrices(DynamicCollection2->Transform, DynamicCollection2->Parent, Collection2_PreReleaseTM);
+		TArray<FTransform> Collection1_PreReleaseTM; GeometryCollectionAlgo::GlobalMatrices(TManagedArray<FTransform3f>(DynamicCollection1->GetTransforms()), DynamicCollection1->Parent, Collection1_PreReleaseTM);
+		TArray<FTransform> Collection2_PreReleaseTM; GeometryCollectionAlgo::GlobalMatrices(TManagedArray<FTransform3f>(DynamicCollection2->GetTransforms()), DynamicCollection2->Parent, Collection2_PreReleaseTM);
 		for (int Idx = 0; Idx < Collection1_PreReleaseTM.Num() - 1; Idx++) 
 		{
 			EXPECT_EQ(Collection1_PreReleaseTM[Idx].GetTranslation().Z, Collection1_InitialTM[Idx].GetTranslation().Z);
@@ -1968,8 +1953,8 @@ namespace GeometryCollectionTest
 
 		// validate that DynamicCollection2 became dynamic and fell from the cluster. 
 
-		TArray<FTransform> Collection1_PostReleaseTM; GeometryCollectionAlgo::GlobalMatrices(DynamicCollection1->Transform, DynamicCollection1->Parent, Collection1_PostReleaseTM);
-		TArray<FTransform> Collection2_PostReleaseTM; GeometryCollectionAlgo::GlobalMatrices(DynamicCollection2->Transform, DynamicCollection2->Parent, Collection2_PostReleaseTM);
+		TArray<FTransform> Collection1_PostReleaseTM; GeometryCollectionAlgo::GlobalMatrices(TManagedArray<FTransform3f>(DynamicCollection1->GetTransforms()), DynamicCollection1->Parent, Collection1_PostReleaseTM);
+		TArray<FTransform> Collection2_PostReleaseTM; GeometryCollectionAlgo::GlobalMatrices(TManagedArray<FTransform3f>(DynamicCollection2->GetTransforms()), DynamicCollection2->Parent, Collection2_PostReleaseTM);
 		for (int Idx = 0; Idx < Collection1_PostReleaseTM.Num() - 1; Idx++) 
 		{
 			if(Idx == 1)
@@ -2017,7 +2002,7 @@ namespace GeometryCollectionTest
 
 		auto& Clustering = UnitTest.Solver->GetEvolution()->GetRigidClustering();
 		const auto& ClusterMap = Clustering.GetChildrenMap();
-		TArray<FTransform> Collection1_InitialTM; GeometryCollectionAlgo::GlobalMatrices(DynamicCollection1->Transform, DynamicCollection1->Parent, Collection1_InitialTM);
+		TArray<FTransform> Collection1_InitialTM; GeometryCollectionAlgo::GlobalMatrices(TManagedArray<FTransform3f>(DynamicCollection1->GetTransforms()), DynamicCollection1->Parent, Collection1_InitialTM);
 		TArray<Chaos::TPBDRigidClusteredParticleHandle<FReal, 3>*>& Collection1Handles = Collection1->PhysObject->GetSolverParticleHandles();
 		const auto& SovlerParticleHandles = UnitTest.Solver->GetParticles().GetParticleHandles();
 
@@ -2036,7 +2021,7 @@ namespace GeometryCollectionTest
 		EXPECT_TRUE(ClusterMapContains(ClusterMap, Collection1Handles[2], { Collection1Handles[1],Collection1Handles[0] }));
 		EXPECT_TRUE(ClusterMapContains(ClusterMap, Collection1Handles[3], { Collection1Handles[2] }));
 
-		TArray<FTransform> Collection1_PreReleaseTM; GeometryCollectionAlgo::GlobalMatrices(DynamicCollection1->Transform, DynamicCollection1->Parent, Collection1_PreReleaseTM);
+		TArray<FTransform> Collection1_PreReleaseTM; GeometryCollectionAlgo::GlobalMatrices(TManagedArray<FTransform3f>(DynamicCollection1->GetTransforms()), DynamicCollection1->Parent, Collection1_PreReleaseTM);
 		for (int Idx = 0; Idx < Collection1_PreReleaseTM.Num() - 1; Idx++)
 		{
 			EXPECT_EQ(Collection1_PreReleaseTM[Idx].GetTranslation().Z, Collection1_InitialTM[Idx].GetTranslation().Z);
@@ -2053,7 +2038,7 @@ namespace GeometryCollectionTest
 		EXPECT_EQ(ClusterMap.Num(), 1);
 
 		// validate that DynamicCollection1 BODY 2 became dynamic and fell from the cluster. 
-		TArray<FTransform> Collection1_PostReleaseTM; GeometryCollectionAlgo::GlobalMatrices(DynamicCollection1->Transform, DynamicCollection1->Parent, Collection1_PostReleaseTM);
+		TArray<FTransform> Collection1_PostReleaseTM; GeometryCollectionAlgo::GlobalMatrices(TManagedArray<FTransform3f>(DynamicCollection1->GetTransforms()), DynamicCollection1->Parent, Collection1_PostReleaseTM);
 		EXPECT_NEAR(Collection1_PostReleaseTM[1].GetTranslation().Z, Collection1_PreReleaseTM[1].GetTranslation().Z, KINDA_SMALL_NUMBER); // the original kinematic should be frozen
 		EXPECT_LT(Collection1_PostReleaseTM[0].GetTranslation().Z, Collection1_PreReleaseTM[0].GetTranslation().Z);
 	}
diff --git a/Engine/Source/Programs/HeadlessChaos/Private/GeometryCollection/GeometryCollectionTestCollisionResolution.cpp b/Engine/Source/Programs/HeadlessChaos/Private/GeometryCollection/GeometryCollectionTestCollisionResolution.cpp
--- a/Engine/Source/Programs/HeadlessChaos/Private/GeometryCollection/GeometryCollectionTestCollisionResolution.cpp
+++ b/Engine/Source/Programs/HeadlessChaos/Private/GeometryCollection/GeometryCollectionTestCollisionResolution.cpp
@@ -49,7 +49,7 @@ namespace GeometryCollectionTest
 		}
 		{
 			// validate that Simplicials are null when CollisionType==Chaos_Volumetric
-			EXPECT_EQ(Collection->DynamicCollection->Transform.Num(), 4);
+			EXPECT_EQ(Collection->DynamicCollection->GetTransforms().Num(), 4);
 			EXPECT_EQ(Collection->DynamicCollection->Simplicials[0], nullptr);
 			EXPECT_EQ(Collection->DynamicCollection->Simplicials[1], nullptr);
 			EXPECT_EQ(Collection->DynamicCollection->Simplicials[2], nullptr);
@@ -58,7 +58,7 @@ namespace GeometryCollectionTest
 
 			const FReal MaxRestingSeparation = -UnitTest.Solver->GetEvolution()->GetGravityForces().GetAcceleration(0).Z * UnitTest.Dt * UnitTest.Dt;	// PBD resting separation will be up to this
 			EXPECT_LT(FMath::Abs(Collection->RestCollection->Transform[0].GetTranslation().Z)-10.f, KINDA_SMALL_NUMBER);
-			EXPECT_LT(FMath::Abs(Collection->DynamicCollection->Transform[0].GetTranslation().Z - 1.0), 0.1 + MaxRestingSeparation);
+			EXPECT_LT(FMath::Abs(Collection->DynamicCollection->GetTransform(0).GetTranslation().Z - 1.0), 0.1 + MaxRestingSeparation);
 		}
 	}
 
@@ -96,14 +96,14 @@ namespace GeometryCollectionTest
 		}
 		{
 			// validate that Simplicials are null when CollisionType==Chaos_Volumetric
-			EXPECT_EQ(Collection->DynamicCollection->Transform.Num(), 4);
+			EXPECT_EQ(Collection->DynamicCollection->GetTransforms().Num(), 4);
 			EXPECT_TRUE(Collection->DynamicCollection->Simplicials[0]!=nullptr);
 			EXPECT_TRUE(UnitTest.Solver->GetParticles().GetGeometryCollectionParticles().CollisionParticles(0)!=nullptr);
 			EXPECT_TRUE(Collection->DynamicCollection->Simplicials[0]->Size() == UnitTest.Solver->GetParticles().GetGeometryCollectionParticles().CollisionParticles(0)->Size());
 			EXPECT_TRUE(Collection->DynamicCollection->Simplicials[0]->Size() != 0);
 
 			EXPECT_LT(FMath::Abs(Collection->RestCollection->Transform[0].GetTranslation().Z - (Radius + 10.f)), KINDA_SMALL_NUMBER);
-			EXPECT_NEAR(Collection->DynamicCollection->Transform[0].GetTranslation().Z, Radius, KINDA_SMALL_NUMBER); 
+			EXPECT_NEAR(Collection->DynamicCollection->GetTransform(0).GetTranslation().Z, Radius, KINDA_SMALL_NUMBER); 
 		}
 	}
 
@@ -164,16 +164,16 @@ namespace GeometryCollectionTest
 
 		UnitTest.Initialize();
 		EXPECT_EQ(
-			SimplicialSphereCollection->DynamicCollection->Transform[0].GetTranslation().Z,
-			ImplicitSphereCollection->DynamicCollection->Transform[0].GetTranslation().Z + 3);
+			SimplicialSphereCollection->DynamicCollection->GetTransform(0).GetTranslation().Z,
+			ImplicitSphereCollection->DynamicCollection->GetTransform(0).GetTranslation().Z + 3);
 
-		const FVector FirstX = FVector(SimplicialSphereCollection->DynamicCollection->Transform[0].GetTranslation());
+		const FVector FirstX = FVector(SimplicialSphereCollection->DynamicCollection->GetTransform(0).GetTranslation());
 		FVector PrevX = FirstX;
 		for (int i = 0; i < 10; i++)
 		{
 			UnitTest.Advance();
 
-			const FVector CurrX = FVector(SimplicialSphereCollection->DynamicCollection->Transform[0].GetTranslation());
+			const FVector CurrX = FVector(SimplicialSphereCollection->DynamicCollection->GetTransform(0).GetTranslation());
 			EXPECT_NE(CurrX.Z, FirstX.Z); // moved since init
 			EXPECT_GE(PrevX.Z - CurrX.Z, -KINDA_SMALL_NUMBER); // falling in -Z, or stopped
 			EXPECT_LE(FMath::Abs(CurrX.X), KINDA_SMALL_NUMBER); // straight down
@@ -184,7 +184,7 @@ namespace GeometryCollectionTest
 		{
 			// We expect the simplical sphere to drop by 0.1 in Z and come to rest
 			// on top of the implicit sphere.
-			const FVector CurrX = FVector(SimplicialSphereCollection->DynamicCollection->Transform[0].GetTranslation());
+			const FVector CurrX = FVector(SimplicialSphereCollection->DynamicCollection->GetTransform(0).GetTranslation());
 			EXPECT_LE(CurrX.Z - 2.0, 0.2); // Relative large fudge factor accounts for aliasing?
 		}
 	}
@@ -242,16 +242,16 @@ namespace GeometryCollectionTest
 */
 		UnitTest.Initialize();
 		EXPECT_EQ(
-			BoxCollection0->DynamicCollection->Transform[0].GetTranslation().Z,
-			BoxCollection1->DynamicCollection->Transform[0].GetTranslation().Z + 3);
+			BoxCollection0->DynamicCollection->GetTransform(0).GetTranslation().Z,
+			BoxCollection1->DynamicCollection->GetTransform(0).GetTranslation().Z + 3);
 
-		const FVector FirstX = FVector(BoxCollection0->DynamicCollection->Transform[0].GetTranslation());
+		const FVector FirstX = FVector(BoxCollection0->DynamicCollection->GetTransform(0).GetTranslation());
 		FVector PrevX = FirstX;
 		for (int i = 0; i < 10; i++)
 		{
 			UnitTest.Advance();
 
-			const FVector& CurrX = FVector(BoxCollection0->DynamicCollection->Transform[0].GetTranslation());
+			const FVector& CurrX = FVector(BoxCollection0->DynamicCollection->GetTransform(0).GetTranslation());
 			EXPECT_NE(CurrX.Z, FirstX.Z); // moved since init
 			EXPECT_LE(CurrX.Z, PrevX.Z); // falling in -Z, or stopped
 			EXPECT_LE(FMath::Abs(CurrX.X), KINDA_SMALL_NUMBER); // No deflection
@@ -262,7 +262,7 @@ namespace GeometryCollectionTest
 		{
 			// We expect the simplical sphere to drop by 0.1 in Z and come to rest
 			// on top of the implicit sphere.
-			const FVector& CurrX = FVector(BoxCollection0->DynamicCollection->Transform[0].GetTranslation());
+			const FVector& CurrX = FVector(BoxCollection0->DynamicCollection->GetTransform(0).GetTranslation());
 			EXPECT_LE(CurrX.Z - 2.0, 0.2); // Relative large fudge factor accounts for aliasing?
 		}
 
@@ -328,22 +328,22 @@ namespace GeometryCollectionTest
 		}
 		{
 			// validate simplicials and implicits are configured correctly
-			EXPECT_EQ(Collection->DynamicCollection->Transform.Num(), 4);
+			EXPECT_EQ(Collection->DynamicCollection->GetTransforms().Num(), 4);
 			EXPECT_TRUE(Collection->DynamicCollection->Simplicials[0] != nullptr);
 			EXPECT_TRUE(UnitTest.Solver->GetParticles().GetGeometryCollectionParticles().CollisionParticles(0) != nullptr);
 			EXPECT_TRUE(Collection->DynamicCollection->Simplicials[0]->Size() == UnitTest.Solver->GetParticles().GetGeometryCollectionParticles().CollisionParticles(0)->Size());
 			EXPECT_TRUE(Collection->DynamicCollection->Simplicials[0]->Size() != 0);
 			EXPECT_TRUE(Collection->DynamicCollection->Implicits[0]->GetType() == (int32)Chaos::ImplicitObjectType::LevelSet);
 
-			EXPECT_EQ(CollectionStaticSphere->DynamicCollection->Transform.Num(), 4);
+			EXPECT_EQ(CollectionStaticSphere->DynamicCollection->GetTransforms().Num(), 4);
 			EXPECT_TRUE(CollectionStaticSphere->DynamicCollection->Simplicials[0] == nullptr);
 			EXPECT_TRUE(CollectionStaticSphere->DynamicCollection->Implicits[0]->GetType() == (int32)Chaos::ImplicitObjectType::Sphere);
 
 			// validate the ball collides and moved away from the static ball
 			EXPECT_LT(FMath::Abs(Collection->RestCollection->Transform[0].GetTranslation().Z) - 10.f, KINDA_SMALL_NUMBER);
-			EXPECT_TRUE(FMath::Abs(Collection->DynamicCollection->Transform[0].GetTranslation().X) < 0.001); // No deflection
-			EXPECT_TRUE(FMath::Abs(Collection->DynamicCollection->Transform[0].GetTranslation().Y) < 0.001); // No deflection
-			EXPECT_LT(Collection->DynamicCollection->Transform[0].GetTranslation().Z, 2.1f); // ball fell
+			EXPECT_TRUE(FMath::Abs(Collection->DynamicCollection->GetTransform(0).GetTranslation().X) < 0.001); // No deflection
+			EXPECT_TRUE(FMath::Abs(Collection->DynamicCollection->GetTransform(0).GetTranslation().Y) < 0.001); // No deflection
+			EXPECT_LT(Collection->DynamicCollection->GetTransform(0).GetTranslation().Z, 2.1f); // ball fell
 		}
 	}
 
@@ -391,15 +391,15 @@ namespace GeometryCollectionTest
 		}
 
 		UnitTest.Initialize();
-		EXPECT_EQ(SimplicialSphereCollection->DynamicCollection->Transform[0].GetTranslation().Z, ImplicitSphereCollection->DynamicCollection->Transform[0].GetTranslation().Z + 2.0f * Radius + 1.0f);
+		EXPECT_EQ(SimplicialSphereCollection->DynamicCollection->GetTransform(0).GetTranslation().Z, ImplicitSphereCollection->DynamicCollection->GetTransform(0).GetTranslation().Z + 2.0f * Radius + 1.0f);
 
-		const FVector FirstX = FVector(SimplicialSphereCollection->DynamicCollection->Transform[0].GetTranslation());
+		const FVector FirstX = FVector(SimplicialSphereCollection->DynamicCollection->GetTransform(0).GetTranslation());
 		FVector PrevX = FirstX;
 		for (int i = 0; i < 10; i++)
 		{
 			UnitTest.Advance();
 
-			const FVector& CurrX = FVector(SimplicialSphereCollection->DynamicCollection->Transform[0].GetTranslation());
+			const FVector& CurrX = FVector(SimplicialSphereCollection->DynamicCollection->GetTransform(0).GetTranslation());
 			EXPECT_NE(CurrX.Z, FirstX.Z); // moved since init
 			EXPECT_LE(FMath::Abs(CurrX.X), 0.1f); // straight down
 			EXPECT_LE(FMath::Abs(CurrX.Y), 0.1f); // straight down
@@ -409,7 +409,7 @@ namespace GeometryCollectionTest
 		{
 			// We expect the simplical sphere to drop by 0.1 in Z and come to rest
 			// on top of the implicit sphere.
-			const FVector& CurrX = FVector(SimplicialSphereCollection->DynamicCollection->Transform[0].GetTranslation());
+			const FVector& CurrX = FVector(SimplicialSphereCollection->DynamicCollection->GetTransform(0).GetTranslation());
 			EXPECT_LE(FMath::Abs(CurrX.Z - 2.0f * Radius), 0.1 * Radius);
 		}
 	}
@@ -467,16 +467,16 @@ namespace GeometryCollectionTest
 */
 		UnitTest.Initialize();
 		EXPECT_EQ(
-			BoxCollection0->DynamicCollection->Transform[0].GetTranslation().Z,
-			BoxCollection1->DynamicCollection->Transform[0].GetTranslation().Z + Length + 2.0f);
+			BoxCollection0->DynamicCollection->GetTransform(0).GetTranslation().Z,
+			BoxCollection1->DynamicCollection->GetTransform(0).GetTranslation().Z + Length + 2.0f);
 
-		const FVector FirstX = FVector(BoxCollection0->DynamicCollection->Transform[0].GetTranslation());
+		const FVector FirstX = FVector(BoxCollection0->DynamicCollection->GetTransform(0).GetTranslation());
 		FVector PrevX = FirstX;
 		for (int i = 0; i < 10; i++)
 		{
 			UnitTest.Advance();
 
-			const FVector& CurrX = FVector(BoxCollection0->DynamicCollection->Transform[0].GetTranslation());
+			const FVector& CurrX = FVector(BoxCollection0->DynamicCollection->GetTransform(0).GetTranslation());
 			EXPECT_NE(CurrX.Z, FirstX.Z); // moved since init
 			EXPECT_LE(CurrX.Z, PrevX.Z); // falling in -Z, or stopped
 			EXPECT_LE(FMath::Abs(CurrX.X), KINDA_SMALL_NUMBER); // straight down
@@ -487,7 +487,7 @@ namespace GeometryCollectionTest
 		{
 			// We expect the simplical cube to drop in Z direction and come to rest
 			// on top of the implicit cube.
-			const FVector& CurrX = FVector(BoxCollection0->DynamicCollection->Transform[0].GetTranslation());
+			const FVector& CurrX = FVector(BoxCollection0->DynamicCollection->GetTransform(0).GetTranslation());
 			EXPECT_LE(FMath::Abs(CurrX.Z - Length), 0.2 * Length); // Relative large fudge factor accounts for spatial aliasing and contact location averaging.
 		}
 	}
@@ -544,7 +544,7 @@ namespace GeometryCollectionTest
 
 			// validate the tetahedron collides and moved away from the static floor
 			FVec3 RestTranslation = Collection->RestCollection->Transform[0].GetTranslation();
-			FVec3 DynamicTranslation = Collection->DynamicCollection->Transform[0].GetTranslation();
+			FVec3 DynamicTranslation = Collection->DynamicCollection->GetTransform(0).GetTranslation();
 			EXPECT_EQ(RestTranslation.Z, 0.f);
 			EXPECT_NEAR(FMath::Abs(DynamicTranslation.X), 0.f, RestingDistanceTolerance);
 			EXPECT_NEAR(FMath::Abs(DynamicTranslation.Y), 0.f, RestingDistanceTolerance);
diff --git a/Engine/Source/Programs/HeadlessChaos/Private/GeometryCollection/GeometryCollectionTestInitilization.cpp b/Engine/Source/Programs/HeadlessChaos/Private/GeometryCollection/GeometryCollectionTestInitilization.cpp
--- a/Engine/Source/Programs/HeadlessChaos/Private/GeometryCollection/GeometryCollectionTestInitilization.cpp
+++ b/Engine/Source/Programs/HeadlessChaos/Private/GeometryCollection/GeometryCollectionTestInitilization.cpp
@@ -33,7 +33,7 @@ GTEST_TEST(AllTraits,GeometryCollection_Initilization_TransformedGeometryCollect
 		EXPECT_LT(X.Z,GlobalTranslation[2]);
 
 		TArray<FTransform> Transform;
-		GeometryCollectionAlgo::GlobalMatrices(Collection->DynamicCollection->Transform,Collection->DynamicCollection->Parent,Transform);
+		GeometryCollectionAlgo::GlobalMatrices(TManagedArray<FTransform3f>(Collection->DynamicCollection->GetTransforms()),Collection->DynamicCollection->Parent,Transform);
 		EXPECT_EQ(Collection->DynamicCollection->Parent[0],FGeometryCollection::Invalid); // is not a child
 		EXPECT_NEAR(Transform[0].GetTranslation().X - GlobalTranslation[0],0.0f,KINDA_SMALL_NUMBER);
 		EXPECT_NEAR(Transform[0].GetTranslation().Y - GlobalTranslation[1],0.0f,KINDA_SMALL_NUMBER);
@@ -64,7 +64,7 @@ GTEST_TEST(AllTraits,GeometryCollection_Initilization_TransformedGeometryCollect
 		EXPECT_LT(X.Z,GlobalTranslation[2]);
 
 		TArray<FTransform> Transform;
-		GeometryCollectionAlgo::GlobalMatrices(Collection->DynamicCollection->Transform,Collection->DynamicCollection->Parent,Transform);
+		GeometryCollectionAlgo::GlobalMatrices(TManagedArray<FTransform3f>(Collection->DynamicCollection->GetTransforms()),Collection->DynamicCollection->Parent,Transform);
 		EXPECT_EQ(Collection->DynamicCollection->Parent[0],FGeometryCollection::Invalid); // is not a child
 		EXPECT_NEAR(Transform[0].GetTranslation().X - GlobalTranslation[0],0.0f,KINDA_SMALL_NUMBER);
 		EXPECT_NEAR(Transform[0].GetTranslation().Y - GlobalTranslation[1],0.0f,KINDA_SMALL_NUMBER);
@@ -95,7 +95,7 @@ GTEST_TEST(AllTraits,GeometryCollection_Initilization_TransformedGeometryCollect
 		EXPECT_LT(X.Z,GlobalTranslation[2]);
 
 		TArray<FTransform> Transform;
-		GeometryCollectionAlgo::GlobalMatrices(Collection->DynamicCollection->Transform,Collection->DynamicCollection->Parent,Transform);
+		GeometryCollectionAlgo::GlobalMatrices(TManagedArray<FTransform3f>(Collection->DynamicCollection->GetTransforms()),Collection->DynamicCollection->Parent,Transform);
 		EXPECT_EQ(Collection->DynamicCollection->Parent[0],FGeometryCollection::Invalid); // is not a child
 		EXPECT_NEAR(Transform[0].GetTranslation().X - GlobalTranslation[0],0.0f,KINDA_SMALL_NUMBER);
 		EXPECT_NEAR(Transform[0].GetTranslation().Y - GlobalTranslation[1],0.0f,KINDA_SMALL_NUMBER);
@@ -146,7 +146,7 @@ GTEST_TEST(AllTraits,GeometryCollection_Initilization_TransformedGeometryCollect
 		FQuat ParticleR = UnitTest.Solver->GetParticles().GetGeometryCollectionParticles().R(0);
 
 		TArray<FTransform> Transform;
-		GeometryCollectionAlgo::GlobalMatrices(Collection->DynamicCollection->Transform, Collection->DynamicCollection->Parent, Transform);
+		GeometryCollectionAlgo::GlobalMatrices(TManagedArray<FTransform3f>(Collection->DynamicCollection->GetTransforms()), Collection->DynamicCollection->Parent, Transform);
 		EXPECT_EQ(Collection->DynamicCollection->Parent[0], FGeometryCollection::Invalid); // is not a child
 		EXPECT_NEAR(ParticlePos.X, GlobalTranslation[0], KINDA_SMALL_NUMBER); // Check particle position from solver 
 		EXPECT_NEAR(ParticlePos.Y, GlobalTranslation[1], KINDA_SMALL_NUMBER);
diff --git a/Engine/Source/Programs/HeadlessChaos/Private/GeometryCollection/GeometryCollectionTestSimulation.cpp b/Engine/Source/Programs/HeadlessChaos/Private/GeometryCollection/GeometryCollectionTestSimulation.cpp
--- a/Engine/Source/Programs/HeadlessChaos/Private/GeometryCollection/GeometryCollectionTestSimulation.cpp
+++ b/Engine/Source/Programs/HeadlessChaos/Private/GeometryCollection/GeometryCollectionTestSimulation.cpp
@@ -40,9 +40,9 @@ namespace GeometryCollectionTest
 
 		{ // test results
 			EXPECT_LT(FMath::Abs(Collection->RestCollection->Transform[0].GetTranslation().Z), SMALL_THRESHOLD); // rest never touched
-			EXPECT_EQ(Collection->DynamicCollection->Transform.Num(), 1); // simulated is falling
-			EXPECT_LT(Collection->DynamicCollection->Transform[0].GetTranslation().Z, 0.f);
-			EXPECT_NEAR(Collection->DynamicCollection->Transform[0].GetTranslation().Z, -980.f * UnitTest.Dt * UnitTest.Dt, 1e-2);// we seem to be twice gravity
+			EXPECT_EQ(Collection->DynamicCollection->GetTransforms().Num(), 1); // simulated is falling
+			EXPECT_LT(Collection->DynamicCollection->GetTransform(0).GetTranslation().Z, 0.f);
+			EXPECT_NEAR(Collection->DynamicCollection->GetTransform(0).GetTranslation().Z, -980.f * UnitTest.Dt * UnitTest.Dt, 1e-2);// we seem to be twice gravity
 		}
 	}
 
@@ -70,8 +70,8 @@ namespace GeometryCollectionTest
 
 		{
 			EXPECT_LT(FMath::Abs(Collection->RestCollection->Transform[0].GetTranslation().Z), SMALL_THRESHOLD);
-			EXPECT_EQ(Collection->DynamicCollection->Transform.Num(), 1);
-			EXPECT_LT(FMath::Abs(Collection->DynamicCollection->Transform[0].GetTranslation().Z - 0.1f * Scale), MEDIUM_THRESHOLD * Scale);
+			EXPECT_EQ(Collection->DynamicCollection->GetTransforms().Num(), 1);
+			EXPECT_LT(FMath::Abs(Collection->DynamicCollection->GetTransform(0).GetTranslation().Z - 0.1f * Scale), MEDIUM_THRESHOLD * Scale);
 		}
 	}
 
@@ -98,8 +98,8 @@ namespace GeometryCollectionTest
 
 		{ // test results
 			EXPECT_LT(FMath::Abs(Collection->RestCollection->Transform[0].GetTranslation().Z), SMALL_THRESHOLD);
-			EXPECT_EQ(Collection->DynamicCollection->Transform.Num(), 1);
-			EXPECT_LT(FMath::Abs(Collection->DynamicCollection->Transform[0].GetTranslation().Z) - Scale[0], SMALL_THRESHOLD);
+			EXPECT_EQ(Collection->DynamicCollection->GetTransforms().Num(), 1);
+			EXPECT_LT(FMath::Abs(Collection->DynamicCollection->GetTransform(0).GetTranslation().Z) - Scale[0], SMALL_THRESHOLD);
 		}
 	}
 
@@ -126,8 +126,8 @@ namespace GeometryCollectionTest
 
 		{
 			EXPECT_LT(FMath::Abs(Collection->RestCollection->Transform[0].GetTranslation().Z), SMALL_THRESHOLD);
-			EXPECT_EQ(Collection->DynamicCollection->Transform.Num(), 1);
-			EXPECT_LT(FMath::Abs(Collection->DynamicCollection->Transform[0].GetTranslation().Z - 0.1f * Scale[0]), MEDIUM_THRESHOLD * Scale[0]);
+			EXPECT_EQ(Collection->DynamicCollection->GetTransforms().Num(), 1);
+			EXPECT_LT(FMath::Abs(Collection->DynamicCollection->GetTransform(0).GetTranslation().Z - 0.1f * Scale[0]), MEDIUM_THRESHOLD * Scale[0]);
 		}
 	}
 
@@ -145,10 +145,9 @@ namespace GeometryCollectionTest
 			UnitTest.Advance();
 
 		{
-			TManagedArray<FTransform3f>& Transform = Collection->DynamicCollection->Transform;
-			EXPECT_EQ(Transform.Num(), 1);
+			EXPECT_EQ(Collection->DynamicCollection->GetTransforms().Num(), 1);
 			//UE_LOG(LogTest, Verbose, TEXT("Position : (%3.5f,%3.5f,%3.5f)"), Transform[0].GetTranslation().X, Transform[0].GetTranslation().Y, Transform[0].GetTranslation().Z);
-			EXPECT_EQ(Transform[0].GetTranslation().Z, 0.f);
+			EXPECT_EQ(Collection->DynamicCollection->GetTransform(0).GetTranslation().Z, 0.f);
 			EXPECT_EQ(Collection->DynamicCollection->DynamicState[0], (int32)EObjectStateTypeEnum::Chaos_Object_Kinematic);
 		}
 	}
@@ -168,7 +167,7 @@ namespace GeometryCollectionTest
 		UnitTest.AddSimulationObject(SleepingCollection);
 		UnitTest.Initialize();
 
-		const auto& Transform0 = SleepingCollection->DynamicCollection->Transform[0];
+		const auto& Transform0 = SleepingCollection->DynamicCollection->GetTransform(0);
 		for (int i = 0; i < 3; i++)
 		{
 			UnitTest.Advance();
@@ -178,7 +177,7 @@ namespace GeometryCollectionTest
 		{
 			// particle doesn't fall due to sleeping state
 			EXPECT_EQ(SleepingCollection->DynamicCollection->DynamicState[0], (int32)EObjectStateTypeEnum::Chaos_Object_Sleeping);
-			EXPECT_LT(FMath::Abs(SleepingCollection->DynamicCollection->Transform[0].GetTranslation().Z - InitialStartHeight), SMALL_THRESHOLD);
+			EXPECT_LT(FMath::Abs(SleepingCollection->DynamicCollection->GetTransform(0).GetTranslation().Z - InitialStartHeight), SMALL_THRESHOLD);
 		}
 
 	}
@@ -205,8 +204,8 @@ namespace GeometryCollectionTest
 		UnitTest.AddSimulationObject(MovingCollection);
 		UnitTest.Initialize();
 
-		const auto& Transform0 = MovingCollection->DynamicCollection->Transform[0];
-		const auto& Transform1 = SleepingCollection->DynamicCollection->Transform[0];
+		const auto& Transform0 = MovingCollection->DynamicCollection->GetTransform(0);
+		const auto& Transform1 = SleepingCollection->DynamicCollection->GetTransform(0);
 		for (int i = 0; i < 15; i++)
 		{
 			UnitTest.Advance();
@@ -237,7 +236,7 @@ namespace GeometryCollectionTest
 		UnitTest.AddSimulationObject(MovingCollection);
 		UnitTest.Initialize();
 
-		const auto& Transform0 = MovingCollection->DynamicCollection->Transform[0];
+		const auto& Transform0 = MovingCollection->DynamicCollection->GetTransform(0);
 		for (int i = 0; i < 5; i++)
 		{
 			UnitTest.Advance();
@@ -288,7 +287,6 @@ namespace GeometryCollectionTest
 
 		// testing...
 		auto GCParticles = Collection->PhysObject->GetSolverParticleHandles();
-		TManagedArray<FTransform3f>& Transform = Collection->DynamicCollection->Transform;
 		for (int Frame = 1; Frame < 200; Frame++)
 		{
 			if (Frame == 1)
@@ -304,38 +302,38 @@ namespace GeometryCollectionTest
 				GCParticles[3]->SetCollisionGroup(3);
 				GCParticles[4]->SetCollisionGroup(-1);
 
-				EXPECT_TRUE(Transform[0].GetRotation() == FQuat4f::Identity); // Can use defaulted zero rotation to indicate that the
-				EXPECT_TRUE(Transform[1].GetRotation() == FQuat4f::Identity); // rigid has not been affected. Should we though??
-				EXPECT_TRUE(Transform[2].GetRotation() == FQuat4f::Identity);
-				EXPECT_TRUE(Transform[3].GetRotation() == FQuat4f::Identity);
-				EXPECT_TRUE(Transform[4].GetRotation() == FQuat4f::Identity);
+				EXPECT_TRUE(Collection->DynamicCollection->GetTransform(0).GetRotation() == FQuat4f::Identity); // Can use defaulted zero rotation to indicate that the
+				EXPECT_TRUE(Collection->DynamicCollection->GetTransform(1).GetRotation() == FQuat4f::Identity); // rigid has not been affected. Should we though??
+				EXPECT_TRUE(Collection->DynamicCollection->GetTransform(2).GetRotation() == FQuat4f::Identity);
+				EXPECT_TRUE(Collection->DynamicCollection->GetTransform(3).GetRotation() == FQuat4f::Identity);
+				EXPECT_TRUE(Collection->DynamicCollection->GetTransform(4).GetRotation() == FQuat4f::Identity);
 			}
 
 			if (Frame == 100)
 			{
-				EXPECT_NEAR(Transform[0].GetTranslation().Z, 50.0f, 1.0f);
-				EXPECT_NEAR(Transform[1].GetTranslation().Z, 150.0f, 1.0f);
-				EXPECT_NEAR(Transform[2].GetTranslation().Z, 250.0f, 1.0f);
-				EXPECT_FALSE(Transform[0].GetRotation() == FQuat4f::Identity);
-				EXPECT_FALSE(Transform[1].GetRotation() == FQuat4f::Identity);
-				EXPECT_FALSE(Transform[2].GetRotation() == FQuat4f::Identity);
-				EXPECT_FALSE(Transform[3].GetRotation() == FQuat4f::Identity);
-				EXPECT_TRUE(Transform[4].GetRotation() == FQuat4f::Identity);
+				EXPECT_NEAR(Collection->DynamicCollection->GetTransform(0).GetTranslation().Z, 50.0f, 1.0f);
+				EXPECT_NEAR(Collection->DynamicCollection->GetTransform(1).GetTranslation().Z, 150.0f, 1.0f);
+				EXPECT_NEAR(Collection->DynamicCollection->GetTransform(2).GetTranslation().Z, 250.0f, 1.0f);
+				EXPECT_FALSE(Collection->DynamicCollection->GetTransform(0).GetRotation() == FQuat4f::Identity);
+				EXPECT_FALSE(Collection->DynamicCollection->GetTransform(1).GetRotation() == FQuat4f::Identity);
+				EXPECT_FALSE(Collection->DynamicCollection->GetTransform(2).GetRotation() == FQuat4f::Identity);
+				EXPECT_FALSE(Collection->DynamicCollection->GetTransform(3).GetRotation() == FQuat4f::Identity);
+				EXPECT_TRUE(Collection->DynamicCollection->GetTransform(4).GetRotation() == FQuat4f::Identity);
 			}
 			UnitTest.Advance();
 		}
 
-		EXPECT_NEAR(Transform[0].GetTranslation().Z, 50.0f, 1.0f);
-		EXPECT_NEAR(Transform[1].GetTranslation().Z, 150.0f, 1.0f);
-		EXPECT_NEAR(Transform[2].GetTranslation().Z, 250.0f, 1.0f);
-		EXPECT_NEAR(Transform[3].GetTranslation().Z, 150.0f, 1.0f);
-		EXPECT_FALSE(Transform[3].GetRotation() == FQuat4f::Identity);
-		EXPECT_TRUE(Transform[4].GetRotation() == FQuat4f::Identity); // Phased through everything, good.
-		EXPECT_LT(Transform[4].GetTranslation().Z, -100.0f);
+		EXPECT_NEAR(Collection->DynamicCollection->GetTransform(0).GetTranslation().Z, 50.0f, 1.0f);
+		EXPECT_NEAR(Collection->DynamicCollection->GetTransform(1).GetTranslation().Z, 150.0f, 1.0f);
+		EXPECT_NEAR(Collection->DynamicCollection->GetTransform(2).GetTranslation().Z, 250.0f, 1.0f);
+		EXPECT_NEAR(Collection->DynamicCollection->GetTransform(3).GetTranslation().Z, 150.0f, 1.0f);
+		EXPECT_FALSE(Collection->DynamicCollection->GetTransform(3).GetRotation() == FQuat4f::Identity);
+		EXPECT_TRUE(Collection->DynamicCollection->GetTransform(4).GetRotation() == FQuat4f::Identity); // Phased through everything, good.
+		EXPECT_LT(Collection->DynamicCollection->GetTransform(4).GetTranslation().Z, -100.0f);
 
 		GCParticles[0]->SetCollisionGroup(-1);
 		for (int i = 0; i < 50; i++) { UnitTest.Advance(); }
-		EXPECT_LT(Transform[0].GetTranslation().Z, -100.0f);
+		EXPECT_LT(Collection->DynamicCollection->GetTransform(0).GetTranslation().Z, -100.0f);
 
 	}
 
diff --git a/Engine/Source/Programs/HeadlessChaos/Private/GeometryCollection/GeometryCollectionTestSimulationField.cpp b/Engine/Source/Programs/HeadlessChaos/Private/GeometryCollection/GeometryCollectionTestSimulationField.cpp
--- a/Engine/Source/Programs/HeadlessChaos/Private/GeometryCollection/GeometryCollectionTestSimulationField.cpp
+++ b/Engine/Source/Programs/HeadlessChaos/Private/GeometryCollection/GeometryCollectionTestSimulationField.cpp
@@ -90,11 +90,10 @@ namespace GeometryCollectionTest
 		EXPECT_EQ(DynamicState[0], (int32)EObjectStateTypeEnum::Chaos_Object_Kinematic);
 
 		// simulated
-		const TManagedArray<FTransform3f>& Transform = Collection->DynamicCollection->Transform;
-		EXPECT_EQ(Transform.Num(), 1);
-		const FVector Translation1 = FVector(Transform[0].GetTranslation());
+		EXPECT_EQ(Collection->DynamicCollection->GetTransforms().Num(), 1);
+		const FVector Translation1 = FVector(Collection->DynamicCollection->GetTransform(0).GetTranslation());
 		EXPECT_NEAR((Translation0 - Translation1).Size(), 0.f, KINDA_SMALL_NUMBER);
-		EXPECT_NEAR(Transform[0].GetTranslation().Z, 1.f, KINDA_SMALL_NUMBER);
+		EXPECT_NEAR(Collection->DynamicCollection->GetTransform(0).GetTranslation().Z, 1.f, KINDA_SMALL_NUMBER);
 
 		FRadialIntMask* RadialMask = new FRadialIntMask();
 		RadialMask->Position = FVector(0.0, 0.0, 0.0);
@@ -110,9 +109,9 @@ namespace GeometryCollectionTest
 		}
 		EXPECT_EQ(DynamicState[0], (int32)EObjectStateTypeEnum::Chaos_Object_Dynamic);
 
-		const FVector Translation2 = FVector(Transform[0].GetTranslation());
+		const FVector Translation2 = FVector(Collection->DynamicCollection->GetTransform(0).GetTranslation());
 		EXPECT_NE(Translation1, Translation2);
-		EXPECT_LE(Transform[0].GetTranslation().Z, 0.f);
+		EXPECT_LE(Collection->DynamicCollection->GetTransform(0).GetTranslation().Z, 0.f);
 	}
 
 	GTEST_TEST(AllTraits, GeometryCollection_RigidBodies_Field_KinematicActivation)
@@ -135,11 +134,10 @@ namespace GeometryCollectionTest
 		}
 
 		// simulated
-		const TManagedArray<FTransform3f>& Transform = Collection->DynamicCollection->Transform;
-		EXPECT_EQ(Transform.Num(), 1);
-		const FVector Translation1 = FVector(Transform[0].GetTranslation());
+		EXPECT_EQ(Collection->DynamicCollection->GetTransforms().Num(), 1);
+		const FVector Translation1 = FVector(Collection->DynamicCollection->GetTransform(0).GetTranslation());
 		EXPECT_NEAR((Translation0 - Translation1).Size(), 0.f, KINDA_SMALL_NUMBER);
-		EXPECT_NEAR(Transform[0].GetTranslation().Z, 1.f, KINDA_SMALL_NUMBER);
+		EXPECT_NEAR(Collection->DynamicCollection->GetTransform(0).GetTranslation().Z, 1.f, KINDA_SMALL_NUMBER);
 
 		FRadialIntMask* RadialMask = new FRadialIntMask();
 		RadialMask->Position = FVector(0.0, 0.0, 0.0);
@@ -155,9 +153,9 @@ namespace GeometryCollectionTest
 			UnitTest.Advance();
 		}
 
-		const FVector Translation2 = FVector(Transform[0].GetTranslation());
+		const FVector Translation2 = FVector(Collection->DynamicCollection->GetTransform(0).GetTranslation());
 		EXPECT_NE(Translation1, Translation2);
-		EXPECT_LE(Transform[0].GetTranslation().Z, 0.f);
+		EXPECT_LE(Collection->DynamicCollection->GetTransform(0).GetTranslation().Z, 0.f);
 	}
 
 	GTEST_TEST(AllTraits, GeometryCollection_RigidBodies_Field_InitialLinearVelocity)
@@ -184,12 +182,11 @@ namespace GeometryCollectionTest
 
 		UnitTest.Initialize();
 
-		TManagedArray<FTransform3f>& Transform = Collection->DynamicCollection->Transform;
 		TManagedArray<int32>& DynamicState = Collection->DynamicCollection->DynamicState;
 
 		FReal PreviousY = 0.f;
-		EXPECT_EQ(Transform[0].GetTranslation().X, 0);
-		EXPECT_EQ(Transform[0].GetTranslation().Y, 0);
+		EXPECT_EQ(Collection->DynamicCollection->GetTransform(0).GetTranslation().X, 0);
+		EXPECT_EQ(Collection->DynamicCollection->GetTransform(0).GetTranslation().Y, 0);
 
 		for (int Frame = 0; Frame < 10; Frame++)
 		{
@@ -204,18 +201,18 @@ namespace GeometryCollectionTest
 			if (Frame >= 2)
 			{
 				EXPECT_EQ(DynamicState[0], (int32)EObjectStateTypeEnum::Chaos_Object_Dynamic);
-				EXPECT_EQ(Transform[0].GetTranslation().X, 0);
-				EXPECT_GT(Transform[0].GetTranslation().Y, PreviousY);
-				EXPECT_LT(Transform[0].GetTranslation().Z, 0);
+				EXPECT_EQ(Collection->DynamicCollection->GetTransform(0).GetTranslation().X, 0);
+				EXPECT_GT(Collection->DynamicCollection->GetTransform(0).GetTranslation().Y, PreviousY);
+				EXPECT_LT(Collection->DynamicCollection->GetTransform(0).GetTranslation().Z, 0);
 			}
 			else
 			{
 				EXPECT_EQ(DynamicState[0], (int32)EObjectStateTypeEnum::Chaos_Object_Kinematic);
-				EXPECT_EQ(Transform[0].GetTranslation().X, 0);
-				EXPECT_EQ(Transform[0].GetTranslation().Y, 0);
-				EXPECT_EQ(Transform[0].GetTranslation().Z, 0);
+				EXPECT_EQ(Collection->DynamicCollection->GetTransform(0).GetTranslation().X, 0);
+				EXPECT_EQ(Collection->DynamicCollection->GetTransform(0).GetTranslation().Y, 0);
+				EXPECT_EQ(Collection->DynamicCollection->GetTransform(0).GetTranslation().Z, 0);
 			}
-			PreviousY = Transform[0].GetTranslation().Y;
+			PreviousY = Collection->DynamicCollection->GetTransform(0).GetTranslation().Y;
 		}
 	}
 
@@ -241,7 +238,6 @@ namespace GeometryCollectionTest
 
 		UnitTest.Initialize();
 
-		TManagedArray<FTransform3f>& Transform = Collection->DynamicCollection->Transform;
 		for (int Frame = 0; Frame < 10; Frame++)
 		{
 			// Set everything inside the r=5.0 sphere to dynamic
@@ -256,16 +252,16 @@ namespace GeometryCollectionTest
 			if (Frame < 5)
 			{
 				// Before frame 5 nothing should have moved
-				EXPECT_LT(FMath::Abs(Transform[0].GetTranslation().Z - 5.f), SMALL_THRESHOLD);
+				EXPECT_LT(FMath::Abs(Collection->DynamicCollection->GetTransform(0).GetTranslation().Z - 5.f), SMALL_THRESHOLD);
 			}
 			else
 			{
 				// Frame 5 and after should be falling
-				EXPECT_LT(Transform[0].GetTranslation().Z, PreviousHeight);
+				EXPECT_LT(Collection->DynamicCollection->GetTransform(0).GetTranslation().Z, PreviousHeight);
 			}
 
 			// Track current height of the object
-			PreviousHeight = Transform[0].GetTranslation().Z;
+			PreviousHeight = Collection->DynamicCollection->GetTransform(0).GetTranslation().Z;
 		}
 
 	}
@@ -288,7 +284,6 @@ namespace GeometryCollectionTest
 
 		UnitTest.Initialize();
 
-		TManagedArray<FTransform3f>& Transform = Collection->DynamicCollection->Transform;
 		FReal PreviousY = 0.0;
 		for (int Frame = 0; Frame < 10; Frame++)
 		{
@@ -302,14 +297,14 @@ namespace GeometryCollectionTest
 
 			if (Frame < 5)
 			{
-				EXPECT_LT(FMath::Abs(Transform[0].GetTranslation().Y), SMALL_THRESHOLD);
+				EXPECT_LT(FMath::Abs(Collection->DynamicCollection->GetTransform(0).GetTranslation().Y), SMALL_THRESHOLD);
 			}
 			else
 			{
-				EXPECT_GT(Transform[0].GetTranslation().Y, PreviousY);
+				EXPECT_GT(Collection->DynamicCollection->GetTransform(0).GetTranslation().Y, PreviousY);
 			}
 
-			PreviousY = Transform[0].GetTranslation().Y;
+			PreviousY = Collection->DynamicCollection->GetTransform(0).GetTranslation().Y;
 
 		}
 
@@ -336,7 +331,6 @@ namespace GeometryCollectionTest
 
 		UnitTest.Initialize();
 
-		TManagedArray<FTransform3f>& Transform = Collection->DynamicCollection->Transform;
 		FReal PreviousY = 0.0;
 		for (int Frame = 0; Frame < 10; Frame++)
 		{
@@ -351,11 +345,11 @@ namespace GeometryCollectionTest
 			auto& Particles = UnitTest.Solver->GetParticles().GetGeometryCollectionParticles();
 			if (Frame < 5)
 			{
-				EXPECT_LT(FMath::Abs(Transform[0].GetRotation().Euler().Y), SMALL_THRESHOLD);
+				EXPECT_LT(FMath::Abs(Collection->DynamicCollection->GetTransform(0).GetRotation().Euler().Y), SMALL_THRESHOLD);
 			}
 			else
 			{
-				EXPECT_NE(FMath::Abs(Transform[0].GetRotation().Euler().Y), SMALL_THRESHOLD);
+				EXPECT_NE(FMath::Abs(Collection->DynamicCollection->GetTransform(0).GetRotation().Euler().Y), SMALL_THRESHOLD);
 				EXPECT_GT(Particles.W(0).Y, PreviousY);
 			}
 
@@ -385,7 +379,7 @@ namespace GeometryCollectionTest
 
 		UnitTest.Initialize();
 
-		TManagedArray<FTransform3f>& Transform = Collection->DynamicCollection->Transform;
+		
 		TManagedArray<bool>& Active = Collection->DynamicCollection->Active;
 		auto& Particles = UnitTest.Solver->GetParticles().GetGeometryCollectionParticles();
 		for (int Frame = 0; Frame < 20; Frame++)
@@ -404,8 +398,8 @@ namespace GeometryCollectionTest
 		EXPECT_EQ(Particles.Disabled(0), true);
 
 		// hasn't fallen any further than this due to being disabled
-		EXPECT_LT(Transform[0].GetTranslation().Z, 5.f);
-		EXPECT_GT(Transform[0].GetTranslation().Z, -5.0f);
+		EXPECT_LT(Collection->DynamicCollection->GetTransform(0).GetTranslation().Z, 5.f);
+		EXPECT_GT(Collection->DynamicCollection->GetTransform(0).GetTranslation().Z, -5.0f);
 	}
 
 	GTEST_TEST(AllTraits, GeometryCollection_RigidBodies_Field_LinearVelocity)
@@ -431,15 +425,15 @@ namespace GeometryCollectionTest
 		UnitTest.Advance();
 
 		FReal PreviousX = 0.0;
-		TManagedArray<FTransform3f>& Transform = Collection->DynamicCollection->Transform;
+		
 		for (int Frame = 1; Frame < 10; Frame++)
 		{
 			UnitTest.Solver->GetPerSolverField().AddTransientCommand({ GetFieldPhysicsName(EFieldPhysicsType::Field_LinearVelocity), VectorField->NewCopy() });
 
 			UnitTest.Advance();
 
-			EXPECT_GT(Transform[0].GetTranslation().X, PreviousX);
-			PreviousX = Transform[0].GetTranslation().X;
+			EXPECT_GT(Collection->DynamicCollection->GetTransform(0).GetTranslation().X, PreviousX);
+			PreviousX = Collection->DynamicCollection->GetTransform(0).GetTranslation().X;
 		}
 	}
 
@@ -489,9 +483,9 @@ namespace GeometryCollectionTest
 			if (Frame == 30)
 			{
 				// The boxes should have landed on each other and settled by now
-				EXPECT_NEAR(Collection[1]->DynamicCollection->Transform[0].GetTranslation().Z, (FReal)100, (FReal)20);
-				EXPECT_NEAR(Collection[2]->DynamicCollection->Transform[0].GetTranslation().Z, (FReal)300, (FReal)20);
-				EXPECT_NEAR(Collection[3]->DynamicCollection->Transform[0].GetTranslation().Z, (FReal)500, (FReal)20);
+				EXPECT_NEAR(Collection[1]->DynamicCollection->GetTransform(0).GetTranslation().Z, (FReal)100, (FReal)20);
+				EXPECT_NEAR(Collection[2]->DynamicCollection->GetTransform(0).GetTranslation().Z, (FReal)300, (FReal)20);
+				EXPECT_NEAR(Collection[3]->DynamicCollection->GetTransform(0).GetTranslation().Z, (FReal)500, (FReal)20);
 			}
 			if (Frame == 31)
 			{
@@ -501,9 +495,9 @@ namespace GeometryCollectionTest
 		}
 		// The bottom boxes should have fallen below the ground level, box 2 now on the ground with box 3 on top
 		auto& Particles = UnitTest.Solver->GetParticles().GetGeometryCollectionParticles();
-		EXPECT_LT(Collection[1]->DynamicCollection->Transform[0].GetTranslation().Z, 0);
-		EXPECT_TRUE(FMath::IsNearlyEqual((FReal)Collection[2]->DynamicCollection->Transform[0].GetTranslation().Z, (FReal)100, (FReal)20));
-		EXPECT_TRUE(FMath::IsNearlyEqual((FReal)Collection[3]->DynamicCollection->Transform[0].GetTranslation().Z, (FReal)300, (FReal)20));
+		EXPECT_LT(Collection[1]->DynamicCollection->GetTransform(0).GetTranslation().Z, 0);
+		EXPECT_TRUE(FMath::IsNearlyEqual((FReal)Collection[2]->DynamicCollection->GetTransform(0).GetTranslation().Z, (FReal)100, (FReal)20));
+		EXPECT_TRUE(FMath::IsNearlyEqual((FReal)Collection[3]->DynamicCollection->GetTransform(0).GetTranslation().Z, (FReal)300, (FReal)20));
 
 	}
 
@@ -973,7 +967,7 @@ namespace GeometryCollectionTest
 		UnitTest.Solver->GetPerSolverField().AddTransientCommand({ TargetNameForce, VectorFieldBackward->NewCopy() });
 		UnitTest.Advance();
 
-		TManagedArray<FTransform3f>& Transform = Collection->DynamicCollection->Transform;
+		
 		for (int Frame = 1; Frame < 10; Frame++)
 		{
 			UnitTest.Solver->GetPerSolverField().AddTransientCommand({ GetFieldPhysicsName(EFieldPhysicsType::Field_LinearVelocity), VectorFieldRight->NewCopy() });
@@ -983,11 +977,11 @@ namespace GeometryCollectionTest
 
 			UnitTest.Advance();
 
-			EXPECT_NEAR(Transform[0].GetTranslation().X, ExpectedLocation.X, KINDA_SMALL_NUMBER);
-			EXPECT_NEAR(Transform[0].GetTranslation().Y, ExpectedLocation.Y, KINDA_SMALL_NUMBER);
-			EXPECT_LT(Transform[0].GetTranslation().Z, ExpectedLocation.Z);
-			EXPECT_LT(Transform[0].GetTranslation().Z, LastZ);
-			LastZ = Transform[0].GetTranslation().Z;
+			EXPECT_NEAR(Collection->DynamicCollection->GetTransform(0).GetTranslation().X, ExpectedLocation.X, KINDA_SMALL_NUMBER);
+			EXPECT_NEAR(Collection->DynamicCollection->GetTransform(0).GetTranslation().Y, ExpectedLocation.Y, KINDA_SMALL_NUMBER);
+			EXPECT_LT(Collection->DynamicCollection->GetTransform(0).GetTranslation().Z, ExpectedLocation.Z);
+			EXPECT_LT(Collection->DynamicCollection->GetTransform(0).GetTranslation().Z, LastZ);
+			LastZ = Collection->DynamicCollection->GetTransform(0).GetTranslation().Z;
 		}
 	}
 
@@ -1020,7 +1014,7 @@ namespace GeometryCollectionTest
 
 		TArray<Chaos::TPBDRigidClusteredParticleHandle<FReal, 3>*>& ParticleHandles = Collection->PhysObject->GetSolverParticleHandles();
 		Chaos::TVector<float, 3> CurrV = ParticleHandles[0]->V();
-		TManagedArray<FTransform3f>& Transform = Collection->DynamicCollection->Transform;
+		
 
 		for (int Frame = 1; Frame < 10; Frame++)
 		{
@@ -1031,12 +1025,12 @@ namespace GeometryCollectionTest
 
 			CurrV = ParticleHandles[0]->V();
 			EXPECT_NEAR(CurrV.X, Params.InitialLinearVelocity.X, KINDA_SMALL_NUMBER); // Velocity in +x
-			EXPECT_GT(Transform[0].GetTranslation().X, LastLocation.X); // Pos in +x
+			EXPECT_GT(Collection->DynamicCollection->GetTransform(0).GetTranslation().X, LastLocation.X); // Pos in +x
 
 			// Still falling?
-			EXPECT_LT(Transform[0].GetTranslation().Z, LastLocation.Z);
+			EXPECT_LT(Collection->DynamicCollection->GetTransform(0).GetTranslation().Z, LastLocation.Z);
 
-			LastLocation = FVector(Transform[0].GetTranslation());
+			LastLocation = FVector(Collection->DynamicCollection->GetTransform(0).GetTranslation());
 
 		}
 	}
@@ -1064,9 +1058,9 @@ namespace GeometryCollectionTest
 
 		UnitTest.Initialize();
 
-		TManagedArray<FTransform3f>& Transform = Collection->DynamicCollection->Transform;
-		FReal PreviousHeight = Transform[0].GetTranslation().Z;
-		FReal PreviousX = Transform[0].GetRotation().Euler().X;
+		
+		FReal PreviousHeight = Collection->DynamicCollection->GetTransform(0).GetTranslation().Z;
+		FReal PreviousX = Collection->DynamicCollection->GetTransform(0).GetRotation().Euler().X;
 		for (int Frame = 0; Frame < 10; Frame++)
 		{
 			
@@ -1077,7 +1071,7 @@ namespace GeometryCollectionTest
 			UnitTest.Advance();
 
 			Chaos::TPBDGeometryCollectionParticles<Chaos::FReal, 3>& Particles = UnitTest.Solver->GetParticles().GetGeometryCollectionParticles();
-			EXPECT_NE(FMath::Abs(Transform[0].GetRotation().Euler().Y), SMALL_THRESHOLD); // not rotating in Y?
+			EXPECT_NE(FMath::Abs(Collection->DynamicCollection->GetTransform(0).GetRotation().Euler().Y), SMALL_THRESHOLD); // not rotating in Y?
 			EXPECT_GT(Particles.W(0).X, PreviousX); // rotating in X?
 			EXPECT_LT(Particles.X(0).Z, PreviousHeight); // still falling?
 
diff --git a/Engine/Source/Programs/HeadlessChaos/Private/GeometryCollection/GeometryCollectionTestSimulationSolver.cpp b/Engine/Source/Programs/HeadlessChaos/Private/GeometryCollection/GeometryCollectionTestSimulationSolver.cpp
--- a/Engine/Source/Programs/HeadlessChaos/Private/GeometryCollection/GeometryCollectionTestSimulationSolver.cpp
+++ b/Engine/Source/Programs/HeadlessChaos/Private/GeometryCollection/GeometryCollectionTestSimulationSolver.cpp
@@ -62,9 +62,8 @@ using namespace ChaosTest;
 			EXPECT_LT(FMath::Abs(RestTransform[0].GetTranslation().Z), SMALL_THRESHOLD);
 
 			// simulated
-			TManagedArray<FTransform3f>& Transform = Collection->DynamicCollection->Transform;
-			EXPECT_EQ(Transform.Num(), 1);
-			EXPECT_LT(FMath::Abs(Transform[0].GetTranslation().Z), SMALL_THRESHOLD);
+			EXPECT_EQ(Collection->DynamicCollection->GetTransforms().Num(), 1);
+			EXPECT_LT(FMath::Abs(Collection->DynamicCollection->GetTransform(0).GetTranslation().Z), SMALL_THRESHOLD);
 		}
 	}
 
@@ -99,7 +98,6 @@ using namespace ChaosTest;
 		UnitTest.Initialize();
 
 		FVector3f StartingClusterPosition;
-		TManagedArray<FTransform3f>* Transform;
 		FReal StartingRigidDistance;
 
 		UnitTest.Solver->RegisterSimOneShotCallback([&]()
@@ -117,10 +115,10 @@ using namespace ChaosTest;
 			// Set the one cluster to disabled
 			UnitTest.Solver->GetEvolution()->DisableParticle(Collection->PhysObject->GetSolverClusterHandles()[0]);
 
-			Transform = &Collection->DynamicCollection->Transform;
-			StartingRigidDistance = ((*Transform)[1].GetTranslation() - (*Transform)[0].GetTranslation()).Size();
+
+			StartingRigidDistance = (Collection->DynamicCollection->GetTransform(1).GetTranslation() - Collection->DynamicCollection->GetTransform(0).GetTranslation()).Size();
 			EXPECT_LT(StartingRigidDistance - 20.0f,SMALL_THRESHOLD);
-			StartingClusterPosition = (*Transform)[2].GetTranslation();
+			StartingClusterPosition = Collection->DynamicCollection->GetTransform(2).GetTranslation();
 		});
 
 		FReal CurrentRigidDistance = 0.0;
@@ -130,11 +128,11 @@ using namespace ChaosTest;
 			UnitTest.Advance();
 
 			// Distance between gc cubes remains the same
-			CurrentRigidDistance = ((*Transform)[1].GetTranslation() - (*Transform)[0].GetTranslation()).Size();
+			CurrentRigidDistance = (Collection->DynamicCollection->GetTransform(1).GetTranslation() - Collection->DynamicCollection->GetTransform(0).GetTranslation()).Size();
 			EXPECT_LT(StartingRigidDistance-CurrentRigidDistance, SMALL_THRESHOLD);
 
 			// Clustered particle doesn't move
-			EXPECT_LT((StartingClusterPosition - (*Transform)[2].GetTranslation()).Size(), SMALL_THRESHOLD);
+			EXPECT_LT((StartingClusterPosition - Collection->DynamicCollection->GetTransform(2).GetTranslation()).Size(), SMALL_THRESHOLD);
 		}
 		
 	}
diff --git a/Engine/Source/Runtime/Experimental/Chaos/Private/GeometryCollectionProxyData.cpp b/Engine/Source/Runtime/Experimental/Chaos/Private/GeometryCollectionProxyData.cpp
--- a/Engine/Source/Runtime/Experimental/Chaos/Private/GeometryCollectionProxyData.cpp
+++ b/Engine/Source/Runtime/Experimental/Chaos/Private/GeometryCollectionProxyData.cpp
@@ -31,6 +31,22 @@ void FTransformDynamicCollection::Construct()
 	AddExternalAttribute<int32>(FGeometryCollection::StatusFlagsAttribute, FTransformCollection::TransformGroup, StatusFlags);
 }
 
+const FTransform3f& FTransformDynamicCollection::GetTransform(int32 Index) const
+{
+	return Transform[Index];
+}
+
+void FTransformDynamicCollection::SetTransform(int32 Index, const FTransform3f& InTransform)
+{
+	Transform[Index] = InTransform;
+}
+
+const TArray<FTransform3f>& FTransformDynamicCollection::GetTransforms() const
+{
+	return Transform.GetConstArray();
+}
+
+
 /*
 * FGeometryDynamicCollection (FTransformDynamicCollection)
 */
diff --git a/Engine/Source/Runtime/Experimental/Chaos/Public/GeometryCollection/ManagedArray.h b/Engine/Source/Runtime/Experimental/Chaos/Public/GeometryCollection/ManagedArray.h
--- a/Engine/Source/Runtime/Experimental/Chaos/Public/GeometryCollection/ManagedArray.h
+++ b/Engine/Source/Runtime/Experimental/Chaos/Public/GeometryCollection/ManagedArray.h
@@ -1183,6 +1183,25 @@ public:
 template<>
 class TManagedArray<FTransform3f> : public TManagedArrayBase<FTransform3f>
 {
+public:
+	FORCEINLINE TManagedArray()
+	{}
+
+	FORCEINLINE TManagedArray(const TArray<FTransform3f>& Other)
+		: TManagedArrayBase<FTransform3f>(Other)
+	{}
+
+	FORCEINLINE TManagedArray(const TManagedArray<FTransform3f>& Other) = delete;
+	FORCEINLINE TManagedArray(TManagedArray<FTransform3f>&& Other) = default;
+	FORCEINLINE TManagedArray(TArray<FTransform3f>&& Other)
+		: TManagedArrayBase<FTransform3f>(MoveTemp(Other))
+	{}
+	FORCEINLINE TManagedArray& operator=(TManagedArray<FTransform3f>&& Other) = default;
+
+	virtual ~TManagedArray()
+	{}
+
+protected:
 	/**
 	* Init from a predefined Array of matching type
 	*/
diff --git a/Engine/Source/Runtime/Experimental/Chaos/Public/GeometryCollectionProxyData.h b/Engine/Source/Runtime/Experimental/Chaos/Public/GeometryCollectionProxyData.h
--- a/Engine/Source/Runtime/Experimental/Chaos/Public/GeometryCollectionProxyData.h
+++ b/Engine/Source/Runtime/Experimental/Chaos/Public/GeometryCollectionProxyData.h
@@ -28,17 +28,23 @@ public:
 	FTransformDynamicCollection(FTransformDynamicCollection&&) = delete;
 	FTransformDynamicCollection& operator=(FTransformDynamicCollection&&) = delete;
 
+	const FTransform3f& GetTransform(int32 Index) const;
+	void SetTransform(int32 Index, const FTransform3f& Transform);
+	const TArray<FTransform3f>& GetTransforms() const;
+
 	// Transform Group
-	TManagedArray<FTransform3f> Transform;
 	TManagedArray<int32>        Parent;
 	TManagedArray<TSet<int32>>  Children;
 	TManagedArray<int32>        SimulationType;
 	TManagedArray<int32>        StatusFlags;
 
-protected:
-
+private:
+	TManagedArray<FTransform3f> Transform;
 	/** Construct */
 	CHAOS_API void Construct();
+
+	friend class FGeometryCollectionPhysicsProxy;
+	friend class UGeometryCollectionComponent;
 };
 
 
