Fix compilation error in Headless Chaos

diff --git a/Engine/Source/Programs/HeadlessChaos/Private/HeadlessChaosTestRewind.cpp b/Engine/Source/Programs/HeadlessChaos/Private/HeadlessChaosTestRewind.cpp
--- a/Engine/Source/Programs/HeadlessChaos/Private/HeadlessChaosTestRewind.cpp
+++ b/Engine/Source/Programs/HeadlessChaos/Private/HeadlessChaosTestRewind.cpp
@@ -3833,7 +3833,7 @@ namespace ChaosTest {
 			for (int Step = 0; Step <= LastStep; ++Step)
 			{
 				UnitTest.Advance();
-				Xs.Add(Collection->DynamicCollection->Transform[0].GetTranslation()[2]);
+				Xs.Add(Collection->DynamicCollection->GetTransform(0).GetTranslation()[2]);
 			}
 
 			const int32 RewindStep = 3;
