[Chaos] Updated the test in Headless Chaos for AABB tree generation to take into account the new Time Slice mode that uses a milliseconds budget instead of amount of work.

This is needed because now that the new mode is enabled by default, the existing test is failing.

Note the preflight apears as failed, but it is because other test are also failing at the moment. With this CL the test we are trying to fix, Passes.

diff --git a/Engine/Source/Programs/HeadlessChaos/Private/HeadlessChaosTestBroadphase.cpp b/Engine/Source/Programs/HeadlessChaos/Private/HeadlessChaosTestBroadphase.cpp
--- a/Engine/Source/Programs/HeadlessChaos/Private/HeadlessChaosTestBroadphase.cpp
+++ b/Engine/Source/Programs/HeadlessChaos/Private/HeadlessChaosTestBroadphase.cpp
@@ -863,7 +863,9 @@ namespace ChaosTest
 		using TreeType = TAABBTree<int32, TAABBTreeLeafArray<int32>>;
 
 		FImplicitObjectPtr Box;
-		auto Boxes = BuildBoxes(Box);
+
+		// If we are time slicing by a milliseconds budget, create a large tree so it takes time to process 
+		auto Boxes = FAABBTimeSliceCVars::bUseTimeSliceMillisecondBudget ? BuildBoxes(Box, 50, FVec3(50.0f,50.0f,50.0f)) : BuildBoxes(Box) ;	
 
 		// build AABB in one go
 		TreeType SpatialBuildImmediate(
@@ -874,7 +876,9 @@ namespace ChaosTest
 			, 0); // build entire tree in one go, no timeslicing
 
 		EXPECT_TRUE(SpatialBuildImmediate.IsAsyncTimeSlicingComplete());
-		
+
+		const double SlicedTreeGenerationStartTime = FPlatformTime::Seconds();
+	
 		// build AABB in time-sliced sections
 		TreeType SpatialTimesliced(
 			MakeParticleView(Boxes.Get())
@@ -883,12 +887,39 @@ namespace ChaosTest
 			, TreeType::DefaultMaxPayloadBounds
 			, 20); // build in small iteration steps, 20 iterations per call to ProgressAsyncTimeSlicing
 
-		EXPECT_FALSE(SpatialTimesliced.IsAsyncTimeSlicingComplete());
+		EXPECT_FALSE(!FAABBTimeSliceCVars::bUseTimeSliceMillisecondBudget && SpatialTimesliced.IsAsyncTimeSlicingComplete());	
+
+		// This is far from accurate, but give us some wiggle room to test it with default settings without needed to implement code to simulate a precise pause inside the Tree implementation.
+		const float MaxSliceDurationWithErrorMargin = FAABBTimeSliceCVars::MaxProcessingTimePerSliceSeconds + 0.01f;
+		double LargestSliceDuration = 0;
+
+		int32 IterationNumber = 1;
+		bool bSliceDoneWithinBudget = true;
 
 		while (!SpatialTimesliced.IsAsyncTimeSlicingComplete())
 		{
+			const double SliceStartTime = FPlatformTime::Seconds();
+	
 			SpatialTimesliced.ProgressAsyncTimeSlicing(false);
-		}	
+			
+			if (FAABBTimeSliceCVars::bUseTimeSliceMillisecondBudget)
+			{
+				const double ElapsedTime = FPlatformTime::Seconds() - SliceStartTime;
+
+				bSliceDoneWithinBudget &= ElapsedTime < MaxSliceDurationWithErrorMargin;
+
+				LargestSliceDuration = FMath::Max(LargestSliceDuration, ElapsedTime);
+			}
+
+			IterationNumber++;
+		}
+
+		EXPECT_TRUE(bSliceDoneWithinBudget);
+
+		const FStringView TimeSliceMode = FAABBTimeSliceCVars::bUseTimeSliceMillisecondBudget ? TEXT("MillisecondsBudget") : TEXT("AmountOfWorkToDo");
+		const double TotalGenerationTime = FPlatformTime::Seconds() - SlicedTreeGenerationStartTime;
+
+		UE_LOG(LogHeadlessChaos, Verbose, TEXT("Time Sliced Tree Generation took [%f] seconds | Using Mode [%s] | In [%d] Iterations | LargestSliceDuration [%f] | EvaluatedMaxTimeSlicedDurarion [%f]"), TotalGenerationTime, TimeSliceMode.GetData(), IterationNumber, LargestSliceDuration, MaxSliceDurationWithErrorMargin);
 
 		// now check both AABBs have the same hierarchy
 		// (indices will be different but walking tree should give same results)
