Reduce memory usage in Dynamic Collection
Remove parent array from Dynamic Collection, it is now replaced by a bit array

diff --git a/Engine/Plugins/Experimental/GeometryCollectionPlugin/Source/GeometryCollectionEditor/Private/GeometryCollection/GeometryCollectionSelectRigidBodyEdMode.cpp b/Engine/Plugins/Experimental/GeometryCollectionPlugin/Source/GeometryCollectionEditor/Private/GeometryCollection/GeometryCollectionSelectRigidBodyEdMode.cpp
--- a/Engine/Plugins/Experimental/GeometryCollectionPlugin/Source/GeometryCollectionEditor/Private/GeometryCollection/GeometryCollectionSelectRigidBodyEdMode.cpp
+++ b/Engine/Plugins/Experimental/GeometryCollectionPlugin/Source/GeometryCollectionEditor/Private/GeometryCollection/GeometryCollectionSelectRigidBodyEdMode.cpp
@@ -138,52 +138,8 @@ bool FGeometryCollectionSelectRigidBodyEdMode::MouseMove(FEditorViewportClient*
 
 bool FGeometryCollectionSelectRigidBodyEdMode::HandleClick(FEditorViewportClient* /*InViewportClient*/, HHitProxy* HitProxy, const FViewportClick& /*Click*/)
 {
-#if GEOMETRYCOLLECTION_EDITOR_SELECTION
-	if (HitProxy && HitProxy->IsA(HGeometryCollection::StaticGetType()))
-	{
-		const HGeometryCollection* const HitGeometryCollection = static_cast<const HGeometryCollection*>(HitProxy);
-		const UGeometryCollectionComponent* const GeometryCollectionComponent = HitGeometryCollection->Component;
-		if (GeometryCollectionComponent)
-		{
-			// Retrieve the transform index
-			int32 TransformIndex = HitGeometryCollection->BoneIndex;
-
-			// Get to the parent transform if the hit transform is still attached to its parent
-			const TManagedArray<int32>& ParentArray = GeometryCollectionComponent->GetParentArray();
-			while (ParentArray[TransformIndex] != FGeometryCollectionBoneNode::InvalidBone)
-			{
-				TransformIndex = ParentArray[TransformIndex];
-			}
-			/*
-			// Retrieve the rigid body id
-			//const TManagedArray<int32>& RigidBodyIdArray = GeometryCollectionComponent->GetRigidBodyIdArray();
-			const TManagedArray<FGuid>& RigidBodyIdArray = GeometryCollectionComponent->GetRigidBodyGuidArray();
-			//const int32 RigidBodyId = (TransformIndex != INDEX_NONE && ensure(TransformIndex < RigidBodyIdArray.Num())) ? RigidBodyIdArray[TransformIndex]: INDEX_NONE;
-			const FGuid& RigidBodyId = (TransformIndex != INDEX_NONE && ensure(TransformIndex < RigidBodyIdArray.Num())) ? RigidBodyIdArray[TransformIndex] : FGuid();
-
-			// Update the rigid body id property
-			//if (RigidBodyId != INDEX_NONE)
-			if (RigidBodyId.IsValid())
-			{
-				UE_LOG(LogGeometryCollectionSelectRigidBodyEdMode, Verbose, TEXT("Hit GeometryCollectionActor %s at rigid body %s."), *GeometryCollectionComponent->GetOwner()->GetName(), *RigidBodyId.ToString());
-
-				if (const TSharedPtr<IPropertyHandle> PropertyHandleIdPin = PropertyHandleId.Pin())
-				{
-					PropertyHandleIdPin->SetValue(RigidBodyId.ToString());
-				}
-				if (const TSharedPtr<IPropertyHandle> PropertyHandleSolverPin = PropertyHandleSolver.Pin())
-				{
-					PropertyHandleSolverPin->SetValue(GeometryCollectionComponent->ChaosSolverActor);
-				}
-			}
-			*/
-		}
-	}
-#endif // #if GEOMETRYCOLLECTION_EDITOR_SELECTION
-
 	// Deactivate editor mode
 	DeactivateMode();
-
 	return true;
 }
 
diff --git a/Engine/Source/Programs/HeadlessChaos/Private/GeometryCollection/GeometryCollectionTestClustering.cpp b/Engine/Source/Programs/HeadlessChaos/Private/GeometryCollection/GeometryCollectionTestClustering.cpp
--- a/Engine/Source/Programs/HeadlessChaos/Private/GeometryCollection/GeometryCollectionTestClustering.cpp
+++ b/Engine/Source/Programs/HeadlessChaos/Private/GeometryCollection/GeometryCollectionTestClustering.cpp
@@ -631,7 +631,8 @@ namespace GeometryCollectionTest
 		EXPECT_TRUE(ClusterMapContains(ClusterMap, Collection1Handles[2], { Collection1Handles[1],Collection1Handles[0] }));
 		EXPECT_TRUE(ClusterMapContains(ClusterMap, Collection1Handles[3], { Collection1Handles[2] }));
 		const TManagedArray<FTransform3f> Transforms1(DynamicCollection1->GetTransforms());
-		TArray<FTransform> Collection1_PreReleaseTM; GeometryCollectionAlgo::GlobalMatrices(Transforms1, DynamicCollection1->Parent, Collection1_PreReleaseTM);
+		TArray<FTransform> Collection1_PreReleaseTM; 
+		GeometryCollectionAlgo::GlobalMatrices(Transforms1, *DynamicCollection1.Get(), Collection1_PreReleaseTM);
 		for (int Idx = 0; Idx < Collection1_PreReleaseTM.Num(); Idx++) {
 			EXPECT_TRUE( (Collection1_PreReleaseTM[Idx].GetTranslation()-Collection1_InitialTM[Idx].GetTranslation()).Size()<KINDA_SMALL_NUMBER);
 		}
@@ -643,7 +644,8 @@ namespace GeometryCollectionTest
 		EXPECT_EQ(ClusterMap.Num(), 1);
 		EXPECT_TRUE(ClusterMapContains(ClusterMap, Collection1Handles[2], { Collection1Handles[1],Collection1Handles[0] }));
 		const TManagedArray<FTransform3f> Transforms2(DynamicCollection1->GetTransforms());
-		TArray<FTransform> Collection1_PostReleaseTM; GeometryCollectionAlgo::GlobalMatrices(Transforms2, DynamicCollection1->Parent, Collection1_PostReleaseTM);
+		TArray<FTransform> Collection1_PostReleaseTM; 
+		GeometryCollectionAlgo::GlobalMatrices(Transforms2, *DynamicCollection1.Get(), Collection1_PostReleaseTM);
 		for (int Idx = 0; Idx < Collection1_PostReleaseTM.Num(); Idx++) {
 			EXPECT_TRUE((Collection1_PostReleaseTM[Idx].GetTranslation() - Collection1_InitialTM[Idx].GetTranslation()).Size() < KINDA_SMALL_NUMBER);
 		}
@@ -654,7 +656,7 @@ namespace GeometryCollectionTest
 		EXPECT_EQ(SovlerParticleHandles.Size(), 4);
 		EXPECT_EQ(ClusterMap.Num(), 0);
 		const TManagedArray<FTransform3f> Transforms3(DynamicCollection1->GetTransforms());
-		TArray<FTransform> Collection1_PostRelease2TM; GeometryCollectionAlgo::GlobalMatrices(Transforms3, DynamicCollection1->Parent, Collection1_PostRelease2TM);
+		TArray<FTransform> Collection1_PostRelease2TM; GeometryCollectionAlgo::GlobalMatrices(Transforms3, *DynamicCollection1, Collection1_PostRelease2TM);
 		for (int Idx = 0; Idx < Collection1_PostRelease2TM.Num(); Idx++) {
 			EXPECT_TRUE((Collection1_PostRelease2TM[Idx].GetTranslation() - Collection1_InitialTM[Idx].GetTranslation()).Size() < KINDA_SMALL_NUMBER);
 		}
@@ -1460,10 +1462,10 @@ namespace GeometryCollectionTest
 			if (Frame == 0)
 			{
 				TArray<FTransform> GlobalTransform;
-				GeometryCollectionAlgo::GlobalMatrices(TManagedArray<FTransform3f>(DynamicCollection->GetTransforms()), DynamicCollection->Parent, GlobalTransform);
+				GeometryCollectionAlgo::GlobalMatrices(TManagedArray<FTransform3f>(DynamicCollection->GetTransforms()), *DynamicCollection, GlobalTransform);
 
 				TArray<FTransform> GlobalTransform2;
-				GeometryCollectionAlgo::GlobalMatrices(TManagedArray<FTransform3f>(DynamicCollection2->GetTransforms()), DynamicCollection2->Parent, GlobalTransform2);
+				GeometryCollectionAlgo::GlobalMatrices(TManagedArray<FTransform3f>(DynamicCollection2->GetTransforms()), *DynamicCollection2, GlobalTransform2);
 
 				// build relative transforms distances
 				for (int32 i = 0; i < (int32)GlobalTransform.Num()-1; i++)
@@ -1482,10 +1484,10 @@ namespace GeometryCollectionTest
 
 		
 		TArray<FTransform> GlobalTransform;
-		GeometryCollectionAlgo::GlobalMatrices(TManagedArray<FTransform3f>(DynamicCollection->GetTransforms()), DynamicCollection->Parent, GlobalTransform);
+		GeometryCollectionAlgo::GlobalMatrices(TManagedArray<FTransform3f>(DynamicCollection->GetTransforms()), *DynamicCollection, GlobalTransform);
 
 		TArray<FTransform> GlobalTransform2;
-		GeometryCollectionAlgo::GlobalMatrices(TManagedArray<FTransform3f>(DynamicCollection2->GetTransforms()), DynamicCollection2->Parent, GlobalTransform2);
+		GeometryCollectionAlgo::GlobalMatrices(TManagedArray<FTransform3f>(DynamicCollection2->GetTransforms()), *DynamicCollection2, GlobalTransform2);
 
 		// build relative transforms distances
 		TArray<FReal> Distances2;
@@ -1558,10 +1560,10 @@ namespace GeometryCollectionTest
 		});
 
 		TArray<FTransform> PrevGlobalTransform;
-		GeometryCollectionAlgo::GlobalMatrices(TManagedArray<FTransform3f>(DynamicCollection->GetTransforms()), DynamicCollection->Parent, PrevGlobalTransform);
+		GeometryCollectionAlgo::GlobalMatrices(TManagedArray<FTransform3f>(DynamicCollection->GetTransforms()), *DynamicCollection, PrevGlobalTransform);
 
 		TArray<FTransform> PrevGlobalTransform2;
-		GeometryCollectionAlgo::GlobalMatrices(TManagedArray<FTransform3f>(DynamicCollection2->GetTransforms()), DynamicCollection2->Parent, PrevGlobalTransform2);
+		GeometryCollectionAlgo::GlobalMatrices(TManagedArray<FTransform3f>(DynamicCollection2->GetTransforms()), *DynamicCollection2, PrevGlobalTransform2);
 
 
 		for (int Frame = 0; Frame < 100; Frame++)
@@ -1569,21 +1571,21 @@ namespace GeometryCollectionTest
 			UnitTest.Advance();
 
 			TArray<FTransform> GlobalTransform;
-			GeometryCollectionAlgo::GlobalMatrices(TManagedArray<FTransform3f>(DynamicCollection->GetTransforms()), DynamicCollection->Parent, GlobalTransform);
+			GeometryCollectionAlgo::GlobalMatrices(TManagedArray<FTransform3f>(DynamicCollection->GetTransforms()), *DynamicCollection, GlobalTransform);
 
 			TArray<FTransform> GlobalTransform2;
-			GeometryCollectionAlgo::GlobalMatrices(TManagedArray<FTransform3f>(DynamicCollection2->GetTransforms()), DynamicCollection2->Parent, GlobalTransform2);
+			GeometryCollectionAlgo::GlobalMatrices(TManagedArray<FTransform3f>(DynamicCollection2->GetTransforms()), *DynamicCollection2, GlobalTransform2);
 
 			EXPECT_EQ(ClusterMap.Num(), 1);
 			EXPECT_TRUE(ClusterMapContains(ClusterMap, ParticleHandles.Handle(6)->CastToRigidParticle(), { ParticleHandles.Handle(1)->CastToRigidParticle(),ParticleHandles.Handle(0)->CastToRigidParticle(),ParticleHandles.Handle(3)->CastToRigidParticle(),ParticleHandles.Handle(4)->CastToRigidParticle() }));
 
-			EXPECT_TRUE(DynamicCollection->Parent[0] == INDEX_NONE);
-			EXPECT_TRUE(DynamicCollection->Parent[1] == INDEX_NONE);
-			EXPECT_TRUE(DynamicCollection->Parent[2] == INDEX_NONE);
+			EXPECT_TRUE(DynamicCollection->GetParent(0) == INDEX_NONE);
+			EXPECT_TRUE(DynamicCollection->GetParent(1) == INDEX_NONE);
+			EXPECT_TRUE(DynamicCollection->GetParent(2) == INDEX_NONE);
 
-			EXPECT_TRUE(DynamicCollection2->Parent[0] == INDEX_NONE);
-			EXPECT_TRUE(DynamicCollection2->Parent[1] == INDEX_NONE);
-			EXPECT_TRUE(DynamicCollection2->Parent[2] == INDEX_NONE);
+			EXPECT_TRUE(DynamicCollection2->GetParent(0) == INDEX_NONE);
+			EXPECT_TRUE(DynamicCollection2->GetParent(1) == INDEX_NONE);
+			EXPECT_TRUE(DynamicCollection2->GetParent(2) == INDEX_NONE);
 
 			EXPECT_TRUE(GlobalTransform[0].GetTranslation().X == PrevGlobalTransform[0].GetTranslation().X);
 			EXPECT_TRUE(GlobalTransform[1].GetTranslation().X == PrevGlobalTransform[1].GetTranslation().X);
@@ -1699,10 +1701,10 @@ namespace GeometryCollectionTest
 				FTransform RootTransform(Root->R(), Root->X());
 
 				TArray<FTransform> GlobalTransform1;
-				GeometryCollectionAlgo::GlobalMatrices(TManagedArray<FTransform3f>(DynamicCollection->GetTransforms()), DynamicCollection->Parent, GlobalTransform1);
+				GeometryCollectionAlgo::GlobalMatrices(TManagedArray<FTransform3f>(DynamicCollection->GetTransforms()), *DynamicCollection, GlobalTransform1);
 
 				TArray<FTransform> GlobalTransform2;
-				GeometryCollectionAlgo::GlobalMatrices(TManagedArray<FTransform3f>(DynamicCollection2->GetTransforms()), DynamicCollection2->Parent, GlobalTransform2);
+				GeometryCollectionAlgo::GlobalMatrices(TManagedArray<FTransform3f>(DynamicCollection2->GetTransforms()), *DynamicCollection2, GlobalTransform2);
 
 				EXPECT_TRUE(!InitialRootPosition.Equals(Root->X())); // root moves
 
@@ -1717,13 +1719,13 @@ namespace GeometryCollectionTest
 			// Validate that the children have been removed from the 
 			// parenting hierarchy.
 			//
-			EXPECT_TRUE(DynamicCollection->Parent[0] == INDEX_NONE);
-			EXPECT_TRUE(DynamicCollection->Parent[1] == INDEX_NONE);
-			EXPECT_TRUE(DynamicCollection->Parent[2] == INDEX_NONE);
+			EXPECT_TRUE(DynamicCollection->GetParent(0) == INDEX_NONE);
+			EXPECT_TRUE(DynamicCollection->GetParent(1) == INDEX_NONE);
+			EXPECT_TRUE(DynamicCollection->GetParent(2) == INDEX_NONE);
 
-			EXPECT_TRUE(DynamicCollection2->Parent[0] == INDEX_NONE);
-			EXPECT_TRUE(DynamicCollection2->Parent[1] == INDEX_NONE);
-			EXPECT_TRUE(DynamicCollection2->Parent[2] == INDEX_NONE);
+			EXPECT_TRUE(DynamicCollection2->GetParent(0) == INDEX_NONE);
+			EXPECT_TRUE(DynamicCollection2->GetParent(1) == INDEX_NONE);
+			EXPECT_TRUE(DynamicCollection2->GetParent(2) == INDEX_NONE);
 		}
 	}
 	
@@ -1775,8 +1777,8 @@ namespace GeometryCollectionTest
 		TSharedPtr<FGeometryDynamicCollection> DynamicCollection1 = Collection1->DynamicCollection;
 		TSharedPtr<FGeometryDynamicCollection> DynamicCollection2 = Collection2->DynamicCollection;
 
-		TArray<FTransform> Collection1_InitialTM; GeometryCollectionAlgo::GlobalMatrices(TManagedArray<FTransform3f>(DynamicCollection1->GetTransforms()), DynamicCollection1->Parent, Collection1_InitialTM);
-		TArray<FTransform> Collection2_InitialTM; GeometryCollectionAlgo::GlobalMatrices(TManagedArray<FTransform3f>(DynamicCollection2->GetTransforms()), DynamicCollection2->Parent, Collection2_InitialTM);
+		TArray<FTransform> Collection1_InitialTM; GeometryCollectionAlgo::GlobalMatrices(TManagedArray<FTransform3f>(DynamicCollection1->GetTransforms()), *DynamicCollection1, Collection1_InitialTM);
+		TArray<FTransform> Collection2_InitialTM; GeometryCollectionAlgo::GlobalMatrices(TManagedArray<FTransform3f>(DynamicCollection2->GetTransforms()), *DynamicCollection2, Collection2_InitialTM);
 
 		TArray<Chaos::TPBDRigidClusteredParticleHandle<FReal, 3>*>& Collection1Handles = Collection1->PhysObject->GetSolverParticleHandles();
 		TArray<Chaos::TPBDRigidClusteredParticleHandle<FReal, 3>*>& Collection2Handles = Collection2->PhysObject->GetSolverParticleHandles();
@@ -1810,8 +1812,8 @@ namespace GeometryCollectionTest
 		EXPECT_TRUE(ClusterMapContains(ClusterMap, ClusterUnion.InternalCluster, { Collection1Handles[0], Collection1Handles[1], Collection2Handles[0], Collection2Handles[1] }));
 		EXPECT_EQ(ClusterMap[ClusterUnion.InternalCluster].Num(), 4);
 
-		TArray<FTransform> Collection1_PreReleaseTM; GeometryCollectionAlgo::GlobalMatrices(TManagedArray<FTransform3f>(DynamicCollection1->GetTransforms()), DynamicCollection1->Parent, Collection1_PreReleaseTM);
-		TArray<FTransform> Collection2_PreReleaseTM; GeometryCollectionAlgo::GlobalMatrices(TManagedArray<FTransform3f>(DynamicCollection2->GetTransforms()), DynamicCollection2->Parent, Collection2_PreReleaseTM);
+		TArray<FTransform> Collection1_PreReleaseTM; GeometryCollectionAlgo::GlobalMatrices(TManagedArray<FTransform3f>(DynamicCollection1->GetTransforms()), *DynamicCollection1, Collection1_PreReleaseTM);
+		TArray<FTransform> Collection2_PreReleaseTM; GeometryCollectionAlgo::GlobalMatrices(TManagedArray<FTransform3f>(DynamicCollection2->GetTransforms()), *DynamicCollection2, Collection2_PreReleaseTM);
 		for (int Idx = 0; Idx < Collection1_PreReleaseTM.Num() - 1; Idx++) {
 			EXPECT_LT(Collection1_PreReleaseTM[Idx].GetTranslation().Z, Collection1_InitialTM[Idx].GetTranslation().Z);
 			EXPECT_LT(Collection2_PreReleaseTM[Idx].GetTranslation().Z, Collection2_InitialTM[Idx].GetTranslation().Z);
@@ -1835,8 +1837,8 @@ namespace GeometryCollectionTest
 
 		EXPECT_EQ(SovlerParticleHandles.Size(), 7);
 
-		TArray<FTransform> Collection1_PostReleaseTM; GeometryCollectionAlgo::GlobalMatrices(TManagedArray<FTransform3f>(DynamicCollection1->GetTransforms()), DynamicCollection1->Parent, Collection1_PostReleaseTM);
-		TArray<FTransform> Collection2_PostReleaseTM; GeometryCollectionAlgo::GlobalMatrices(TManagedArray<FTransform3f>(DynamicCollection2->GetTransforms()), DynamicCollection2->Parent, Collection2_PostReleaseTM);
+		TArray<FTransform> Collection1_PostReleaseTM; GeometryCollectionAlgo::GlobalMatrices(TManagedArray<FTransform3f>(DynamicCollection1->GetTransforms()), *DynamicCollection1, Collection1_PostReleaseTM);
+		TArray<FTransform> Collection2_PostReleaseTM; GeometryCollectionAlgo::GlobalMatrices(TManagedArray<FTransform3f>(DynamicCollection2->GetTransforms()), *DynamicCollection2, Collection2_PostReleaseTM);
 		for (int Idx = 0; Idx < Collection1_PostReleaseTM.Num() - 1; Idx++) {
 			EXPECT_LT(Collection1_PostReleaseTM[Idx].GetTranslation().Z, Collection1_PreReleaseTM[Idx].GetTranslation().Z);
 			EXPECT_LT(Collection2_PostReleaseTM[Idx].GetTranslation().Z, Collection2_PreReleaseTM[Idx].GetTranslation().Z);
@@ -1891,8 +1893,8 @@ namespace GeometryCollectionTest
 		Chaos::FClusterUnionManager& ClusterUnionManager = Clustering.GetClusterUnionManager();
 		const auto& ClusterMap = Clustering.GetChildrenMap();
 
-		TArray<FTransform> Collection1_InitialTM; GeometryCollectionAlgo::GlobalMatrices(TManagedArray<FTransform3f>(DynamicCollection1->GetTransforms()), DynamicCollection1->Parent, Collection1_InitialTM);
-		TArray<FTransform> Collection2_InitialTM; GeometryCollectionAlgo::GlobalMatrices(TManagedArray<FTransform3f>(DynamicCollection2->GetTransforms()), DynamicCollection2->Parent, Collection2_InitialTM);
+		TArray<FTransform> Collection1_InitialTM; GeometryCollectionAlgo::GlobalMatrices(TManagedArray<FTransform3f>(DynamicCollection1->GetTransforms()), *DynamicCollection1, Collection1_InitialTM);
+		TArray<FTransform> Collection2_InitialTM; GeometryCollectionAlgo::GlobalMatrices(TManagedArray<FTransform3f>(DynamicCollection2->GetTransforms()), *DynamicCollection2, Collection2_InitialTM);
 
 		TArray<Chaos::TPBDRigidClusteredParticleHandle<FReal, 3>*>& Collection1Handles = Collection1->PhysObject->GetSolverParticleHandles();
 		TArray<Chaos::TPBDRigidClusteredParticleHandle<FReal, 3>*>& Collection2Handles = Collection2->PhysObject->GetSolverParticleHandles();
@@ -1925,8 +1927,8 @@ namespace GeometryCollectionTest
 		EXPECT_TRUE(ClusterMapContains(ClusterMap, ClusterUnion.InternalCluster, { Collection1Handles[0], Collection1Handles[1], Collection2Handles[0], Collection2Handles[1] }));
 		EXPECT_EQ(ClusterMap[ClusterUnion.InternalCluster].Num(), 4);
 
-		TArray<FTransform> Collection1_PreReleaseTM; GeometryCollectionAlgo::GlobalMatrices(TManagedArray<FTransform3f>(DynamicCollection1->GetTransforms()), DynamicCollection1->Parent, Collection1_PreReleaseTM);
-		TArray<FTransform> Collection2_PreReleaseTM; GeometryCollectionAlgo::GlobalMatrices(TManagedArray<FTransform3f>(DynamicCollection2->GetTransforms()), DynamicCollection2->Parent, Collection2_PreReleaseTM);
+		TArray<FTransform> Collection1_PreReleaseTM; GeometryCollectionAlgo::GlobalMatrices(TManagedArray<FTransform3f>(DynamicCollection1->GetTransforms()), *DynamicCollection1, Collection1_PreReleaseTM);
+		TArray<FTransform> Collection2_PreReleaseTM; GeometryCollectionAlgo::GlobalMatrices(TManagedArray<FTransform3f>(DynamicCollection2->GetTransforms()), *DynamicCollection2, Collection2_PreReleaseTM);
 		for (int Idx = 0; Idx < Collection1_PreReleaseTM.Num() - 1; Idx++) 
 		{
 			EXPECT_EQ(Collection1_PreReleaseTM[Idx].GetTranslation().Z, Collection1_InitialTM[Idx].GetTranslation().Z);
@@ -1953,8 +1955,8 @@ namespace GeometryCollectionTest
 
 		// validate that DynamicCollection2 became dynamic and fell from the cluster. 
 
-		TArray<FTransform> Collection1_PostReleaseTM; GeometryCollectionAlgo::GlobalMatrices(TManagedArray<FTransform3f>(DynamicCollection1->GetTransforms()), DynamicCollection1->Parent, Collection1_PostReleaseTM);
-		TArray<FTransform> Collection2_PostReleaseTM; GeometryCollectionAlgo::GlobalMatrices(TManagedArray<FTransform3f>(DynamicCollection2->GetTransforms()), DynamicCollection2->Parent, Collection2_PostReleaseTM);
+		TArray<FTransform> Collection1_PostReleaseTM; GeometryCollectionAlgo::GlobalMatrices(TManagedArray<FTransform3f>(DynamicCollection1->GetTransforms()), *DynamicCollection1, Collection1_PostReleaseTM);
+		TArray<FTransform> Collection2_PostReleaseTM; GeometryCollectionAlgo::GlobalMatrices(TManagedArray<FTransform3f>(DynamicCollection2->GetTransforms()), *DynamicCollection2, Collection2_PostReleaseTM);
 		for (int Idx = 0; Idx < Collection1_PostReleaseTM.Num() - 1; Idx++) 
 		{
 			if(Idx == 1)
@@ -2002,7 +2004,7 @@ namespace GeometryCollectionTest
 
 		auto& Clustering = UnitTest.Solver->GetEvolution()->GetRigidClustering();
 		const auto& ClusterMap = Clustering.GetChildrenMap();
-		TArray<FTransform> Collection1_InitialTM; GeometryCollectionAlgo::GlobalMatrices(TManagedArray<FTransform3f>(DynamicCollection1->GetTransforms()), DynamicCollection1->Parent, Collection1_InitialTM);
+		TArray<FTransform> Collection1_InitialTM; GeometryCollectionAlgo::GlobalMatrices(TManagedArray<FTransform3f>(DynamicCollection1->GetTransforms()), *DynamicCollection1, Collection1_InitialTM);
 		TArray<Chaos::TPBDRigidClusteredParticleHandle<FReal, 3>*>& Collection1Handles = Collection1->PhysObject->GetSolverParticleHandles();
 		const auto& SovlerParticleHandles = UnitTest.Solver->GetParticles().GetParticleHandles();
 
@@ -2021,7 +2023,7 @@ namespace GeometryCollectionTest
 		EXPECT_TRUE(ClusterMapContains(ClusterMap, Collection1Handles[2], { Collection1Handles[1],Collection1Handles[0] }));
 		EXPECT_TRUE(ClusterMapContains(ClusterMap, Collection1Handles[3], { Collection1Handles[2] }));
 
-		TArray<FTransform> Collection1_PreReleaseTM; GeometryCollectionAlgo::GlobalMatrices(TManagedArray<FTransform3f>(DynamicCollection1->GetTransforms()), DynamicCollection1->Parent, Collection1_PreReleaseTM);
+		TArray<FTransform> Collection1_PreReleaseTM; GeometryCollectionAlgo::GlobalMatrices(TManagedArray<FTransform3f>(DynamicCollection1->GetTransforms()), *DynamicCollection1, Collection1_PreReleaseTM);
 		for (int Idx = 0; Idx < Collection1_PreReleaseTM.Num() - 1; Idx++)
 		{
 			EXPECT_EQ(Collection1_PreReleaseTM[Idx].GetTranslation().Z, Collection1_InitialTM[Idx].GetTranslation().Z);
@@ -2038,7 +2040,7 @@ namespace GeometryCollectionTest
 		EXPECT_EQ(ClusterMap.Num(), 1);
 
 		// validate that DynamicCollection1 BODY 2 became dynamic and fell from the cluster. 
-		TArray<FTransform> Collection1_PostReleaseTM; GeometryCollectionAlgo::GlobalMatrices(TManagedArray<FTransform3f>(DynamicCollection1->GetTransforms()), DynamicCollection1->Parent, Collection1_PostReleaseTM);
+		TArray<FTransform> Collection1_PostReleaseTM; GeometryCollectionAlgo::GlobalMatrices(TManagedArray<FTransform3f>(DynamicCollection1->GetTransforms()), *DynamicCollection1, Collection1_PostReleaseTM);
 		EXPECT_NEAR(Collection1_PostReleaseTM[1].GetTranslation().Z, Collection1_PreReleaseTM[1].GetTranslation().Z, KINDA_SMALL_NUMBER); // the original kinematic should be frozen
 		EXPECT_LT(Collection1_PostReleaseTM[0].GetTranslation().Z, Collection1_PreReleaseTM[0].GetTranslation().Z);
 	}
diff --git a/Engine/Source/Programs/HeadlessChaos/Private/GeometryCollection/GeometryCollectionTestCollisionResolution.cpp b/Engine/Source/Programs/HeadlessChaos/Private/GeometryCollection/GeometryCollectionTestCollisionResolution.cpp
--- a/Engine/Source/Programs/HeadlessChaos/Private/GeometryCollection/GeometryCollectionTestCollisionResolution.cpp
+++ b/Engine/Source/Programs/HeadlessChaos/Private/GeometryCollection/GeometryCollectionTestCollisionResolution.cpp
@@ -31,7 +31,7 @@ namespace GeometryCollectionTest
 
 			const TManagedArray<FTransform>& MassToLocal = Collection->RestCollection->GetAttribute<FTransform>("MassToLocal", FGeometryCollection::TransformGroup);
 
-			EXPECT_EQ(Collection->DynamicCollection->Parent[0], 1); // is a child of index one
+			EXPECT_EQ(Collection->DynamicCollection->GetParent(0), 1); // is a child of index one
 			EXPECT_TRUE(MassToLocal[0].Equals(FTransform::Identity)); // we are not testing MassToLocal in this test
 			
 			
@@ -80,7 +80,7 @@ namespace GeometryCollectionTest
 
 			const TManagedArray<FTransform>& MassToLocal = Collection->RestCollection->GetAttribute<FTransform>("MassToLocal", FGeometryCollection::TransformGroup);
 
-			EXPECT_EQ(Collection->DynamicCollection->Parent[0], 1); // is a child of index one
+			EXPECT_EQ(Collection->DynamicCollection->GetParent(0), 1); // is a child of index one
 			EXPECT_TRUE(MassToLocal[0].Equals(FTransform::Identity)); // we are not testing MassToLocal in this test
 
 			UnitTest.AddSimulationObject(Collection);
@@ -293,7 +293,7 @@ namespace GeometryCollectionTest
 			
 			const TManagedArray<FTransform>& MassToLocal = Collection->RestCollection->GetAttribute<FTransform>("MassToLocal", FGeometryCollection::TransformGroup);
 
-			EXPECT_EQ(Collection->DynamicCollection->Parent[0], 1); // is a child of index one
+			EXPECT_EQ(Collection->DynamicCollection->GetParent(0), 1); // is a child of index one
 			EXPECT_TRUE(MassToLocal[0].Equals(FTransform::Identity)); // we are not testing MassToLocal in this test
 
 			UnitTest.AddSimulationObject(Collection);
@@ -314,7 +314,7 @@ namespace GeometryCollectionTest
 			
 			const TManagedArray<FTransform>& MassToLocal = Collection->RestCollection->GetAttribute<FTransform>("MassToLocal", FGeometryCollection::TransformGroup);
 			
-			EXPECT_EQ(CollectionStaticSphere->DynamicCollection->Parent[0], 1); // is a child of index one
+			EXPECT_EQ(CollectionStaticSphere->DynamicCollection->GetParent(0), 1); // is a child of index one
 			EXPECT_TRUE(MassToLocal[0].Equals(FTransform::Identity)); // we are not testing MassToLocal in this test
 
 			UnitTest.AddSimulationObject(CollectionStaticSphere);
@@ -518,7 +518,7 @@ namespace GeometryCollectionTest
 
 			const TManagedArray<FTransform>& MassToLocal = Collection->RestCollection->GetAttribute<FTransform>("MassToLocal", FGeometryCollection::TransformGroup);
 
-			EXPECT_EQ(Collection->DynamicCollection->Parent[0], -1); // is a child of index one
+			EXPECT_EQ(Collection->DynamicCollection->GetParent(0), -1); // is a child of index one
 			EXPECT_NEAR((MassToLocal[0].GetTranslation()-FVector(0,0,Scale + 10)).Size(),0,KINDA_SMALL_NUMBER);
 
 			UnitTest.AddSimulationObject(Collection);
diff --git a/Engine/Source/Programs/HeadlessChaos/Private/GeometryCollection/GeometryCollectionTestInitilization.cpp b/Engine/Source/Programs/HeadlessChaos/Private/GeometryCollection/GeometryCollectionTestInitilization.cpp
--- a/Engine/Source/Programs/HeadlessChaos/Private/GeometryCollection/GeometryCollectionTestInitilization.cpp
+++ b/Engine/Source/Programs/HeadlessChaos/Private/GeometryCollection/GeometryCollectionTestInitilization.cpp
@@ -16,7 +16,7 @@ GTEST_TEST(AllTraits,GeometryCollection_Initilization_TransformedGeometryCollect
 	CreationParameters Params; Params.DynamicState = EObjectStateTypeEnum::Chaos_Object_Dynamic; Params.EnableClustering = false;
 	Params.RootTransform = FTransform(GlobalRotation,GlobalTranslation); Params.NestedTransforms ={FTransform(FVector(10)),FTransform::Identity,FTransform::Identity};
 	FGeometryCollectionWrapper* Collection = TNewSimulationObject<GeometryType::GeometryCollectionWithSingleRigid>::Init(Params)->template As<FGeometryCollectionWrapper>();
-	EXPECT_EQ(Collection->DynamicCollection->Parent[0],1); // is a child of index one
+	EXPECT_EQ(Collection->DynamicCollection->GetParent(0),1); // is a child of index one
 
 	FFramework UnitTest;
 	UnitTest.AddSimulationObject(Collection);
@@ -33,8 +33,8 @@ GTEST_TEST(AllTraits,GeometryCollection_Initilization_TransformedGeometryCollect
 		EXPECT_LT(X.Z,GlobalTranslation[2]);
 
 		TArray<FTransform> Transform;
-		GeometryCollectionAlgo::GlobalMatrices(TManagedArray<FTransform3f>(Collection->DynamicCollection->GetTransforms()),Collection->DynamicCollection->Parent,Transform);
-		EXPECT_EQ(Collection->DynamicCollection->Parent[0],FGeometryCollection::Invalid); // is not a child
+		GeometryCollectionAlgo::GlobalMatrices(TManagedArray<FTransform3f>(Collection->DynamicCollection->GetTransforms()),*Collection->DynamicCollection,Transform);
+		EXPECT_EQ(Collection->DynamicCollection->GetParent(0),FGeometryCollection::Invalid); // is not a child
 		EXPECT_NEAR(Transform[0].GetTranslation().X - GlobalTranslation[0],0.0f,KINDA_SMALL_NUMBER);
 		EXPECT_NEAR(Transform[0].GetTranslation().Y - GlobalTranslation[1],0.0f,KINDA_SMALL_NUMBER);
 		EXPECT_LT(Transform[0].GetTranslation().Z,GlobalTranslation[2]);
@@ -64,8 +64,8 @@ GTEST_TEST(AllTraits,GeometryCollection_Initilization_TransformedGeometryCollect
 		EXPECT_LT(X.Z,GlobalTranslation[2]);
 
 		TArray<FTransform> Transform;
-		GeometryCollectionAlgo::GlobalMatrices(TManagedArray<FTransform3f>(Collection->DynamicCollection->GetTransforms()),Collection->DynamicCollection->Parent,Transform);
-		EXPECT_EQ(Collection->DynamicCollection->Parent[0],FGeometryCollection::Invalid); // is not a child
+		GeometryCollectionAlgo::GlobalMatrices(TManagedArray<FTransform3f>(Collection->DynamicCollection->GetTransforms()),*Collection->DynamicCollection,Transform);
+		EXPECT_EQ(Collection->DynamicCollection->GetParent(0),FGeometryCollection::Invalid); // is not a child
 		EXPECT_NEAR(Transform[0].GetTranslation().X - GlobalTranslation[0],0.0f,KINDA_SMALL_NUMBER);
 		EXPECT_NEAR(Transform[0].GetTranslation().Y - GlobalTranslation[1],0.0f,KINDA_SMALL_NUMBER);
 		EXPECT_LT(Transform[0].GetTranslation().Z,GlobalTranslation[2]);
@@ -95,8 +95,8 @@ GTEST_TEST(AllTraits,GeometryCollection_Initilization_TransformedGeometryCollect
 		EXPECT_LT(X.Z,GlobalTranslation[2]);
 
 		TArray<FTransform> Transform;
-		GeometryCollectionAlgo::GlobalMatrices(TManagedArray<FTransform3f>(Collection->DynamicCollection->GetTransforms()),Collection->DynamicCollection->Parent,Transform);
-		EXPECT_EQ(Collection->DynamicCollection->Parent[0],FGeometryCollection::Invalid); // is not a child
+		GeometryCollectionAlgo::GlobalMatrices(TManagedArray<FTransform3f>(Collection->DynamicCollection->GetTransforms()),*Collection->DynamicCollection,Transform);
+		EXPECT_EQ(Collection->DynamicCollection->GetParent(0),FGeometryCollection::Invalid); // is not a child
 		EXPECT_NEAR(Transform[0].GetTranslation().X - GlobalTranslation[0],0.0f,KINDA_SMALL_NUMBER);
 		EXPECT_NEAR(Transform[0].GetTranslation().Y - GlobalTranslation[1],0.0f,KINDA_SMALL_NUMBER);
 		EXPECT_LT(Transform[0].GetTranslation().Z,GlobalTranslation[2]);
@@ -146,8 +146,8 @@ GTEST_TEST(AllTraits,GeometryCollection_Initilization_TransformedGeometryCollect
 		FQuat ParticleR = UnitTest.Solver->GetParticles().GetGeometryCollectionParticles().R(0);
 
 		TArray<FTransform> Transform;
-		GeometryCollectionAlgo::GlobalMatrices(TManagedArray<FTransform3f>(Collection->DynamicCollection->GetTransforms()), Collection->DynamicCollection->Parent, Transform);
-		EXPECT_EQ(Collection->DynamicCollection->Parent[0], FGeometryCollection::Invalid); // is not a child
+		GeometryCollectionAlgo::GlobalMatrices(TManagedArray<FTransform3f>(Collection->DynamicCollection->GetTransforms()), *Collection->DynamicCollection, Transform);
+		EXPECT_EQ(Collection->DynamicCollection->GetParent(0), FGeometryCollection::Invalid); // is not a child
 		EXPECT_NEAR(ParticlePos.X, GlobalTranslation[0], KINDA_SMALL_NUMBER); // Check particle position from solver 
 		EXPECT_NEAR(ParticlePos.Y, GlobalTranslation[1], KINDA_SMALL_NUMBER);
 		EXPECT_NEAR(Transform[0].GetTranslation().X, 0.0f, KINDA_SMALL_NUMBER);
diff --git a/Engine/Source/Programs/HeadlessChaos/Private/GeometryCollection/GeometryCollectionTestSimulationSolver.cpp b/Engine/Source/Programs/HeadlessChaos/Private/GeometryCollection/GeometryCollectionTestSimulationSolver.cpp
--- a/Engine/Source/Programs/HeadlessChaos/Private/GeometryCollection/GeometryCollectionTestSimulationSolver.cpp
+++ b/Engine/Source/Programs/HeadlessChaos/Private/GeometryCollection/GeometryCollectionTestSimulationSolver.cpp
@@ -108,10 +108,9 @@ using namespace ChaosTest;
 			EXPECT_EQ(SimulationType[1],FGeometryCollection::ESimulationTypes::FST_Rigid);
 			EXPECT_EQ(SimulationType[2],FGeometryCollection::ESimulationTypes::FST_Clustered);
 
-			TManagedArray<int32>& Parent = Collection->DynamicCollection->Parent;
-			EXPECT_EQ(Parent[0],2);
-			EXPECT_EQ(Parent[1],2);
-			EXPECT_EQ(Parent[2],-1);
+			EXPECT_EQ(Collection->DynamicCollection->GetParent(0),2);
+			EXPECT_EQ(Collection->DynamicCollection->GetParent(1),2);
+			EXPECT_EQ(Collection->DynamicCollection->GetParent(2),-1);
 
 			// Set the one cluster to disabled
 			UnitTest.Solver->GetEvolution()->DisableParticle(Collection->PhysObject->GetSolverClusterHandles()[0]);
diff --git a/Engine/Source/Programs/HeadlessChaos/Private/GeometryCollection/GeometryCollectionTestUtility.cpp b/Engine/Source/Programs/HeadlessChaos/Private/GeometryCollection/GeometryCollectionTestUtility.cpp
--- a/Engine/Source/Programs/HeadlessChaos/Private/GeometryCollection/GeometryCollectionTestUtility.cpp
+++ b/Engine/Source/Programs/HeadlessChaos/Private/GeometryCollection/GeometryCollectionTestUtility.cpp
@@ -21,7 +21,7 @@ namespace GeometryCollectionTest {
 
 	TSharedPtr<FGeometryDynamicCollection> GeometryCollectionToGeometryDynamicCollection(const FGeometryCollection* InputCollection, EObjectStateTypeEnum DynamicStateDefault)
 	{
-		TSharedPtr<FGeometryDynamicCollection> NewCollection(new FGeometryDynamicCollection());
+		TSharedPtr<FGeometryDynamicCollection> NewCollection(new FGeometryDynamicCollection(InputCollection));
 		NewCollection->CopyAttribute(*InputCollection, FTransformCollection::TransformAttribute, FGeometryCollection::TransformGroup);
 		NewCollection->CopyAttribute(*InputCollection, FTransformCollection::ParentAttribute, FGeometryCollection::TransformGroup);
 		NewCollection->CopyAttribute(*InputCollection, FTransformCollection::ChildrenAttribute, FGeometryCollection::TransformGroup);
diff --git a/Engine/Source/Runtime/Experimental/Chaos/Private/Chaos/PhysicsObjectInternal.h b/Engine/Source/Runtime/Experimental/Chaos/Private/Chaos/PhysicsObjectInternal.h
--- a/Engine/Source/Runtime/Experimental/Chaos/Private/Chaos/PhysicsObjectInternal.h
+++ b/Engine/Source/Runtime/Experimental/Chaos/Private/Chaos/PhysicsObjectInternal.h
@@ -70,7 +70,7 @@ namespace Chaos
 				FGeometryCollectionPhysicsProxy* GeometryCollectionProxy = static_cast<FGeometryCollectionPhysicsProxy*>(Proxy);
 				FGeometryDynamicCollection& Collection = GetGeometryCollectionDynamicCollection<Id>(*GeometryCollectionProxy);
 
-				if (int32 Index = Collection.Parent[BodyIndex]; Index != INDEX_NONE)
+				if (int32 Index = Collection.GetParent(BodyIndex); Index != INDEX_NONE)
 				{
 					return GeometryCollectionProxy->GetPhysicsObjectByIndex(Index);
 				}
diff --git a/Engine/Source/Runtime/Experimental/Chaos/Private/GeometryCollection/GeometryCollectionAlgo.cpp b/Engine/Source/Runtime/Experimental/Chaos/Private/GeometryCollection/GeometryCollectionAlgo.cpp
--- a/Engine/Source/Runtime/Experimental/Chaos/Private/GeometryCollection/GeometryCollectionAlgo.cpp
+++ b/Engine/Source/Runtime/Experimental/Chaos/Private/GeometryCollection/GeometryCollectionAlgo.cpp
@@ -7,6 +7,7 @@
 #include "GeometryCollection/GeometryCollection.h"
 #include "GeometryCollection/RecordedTransformTrack.h"
 #include "GeometryCollection/Facades/CollectionHierarchyFacade.h"
+#include "GeometryCollectionProxyData.h"
 #include "Async/ParallelFor.h"
 
 DEFINE_LOG_CATEGORY_STATIC(GeometryCollectionAlgoLog, Log, All);
@@ -291,6 +292,34 @@ namespace GeometryCollectionAlgo
 	// Helper type for computing global matrices
 	using FIndicesNeedMatricesArray = TArray<int32, TInlineAllocator<16>>;
 
+	// return false if we failed to get the indices, which can happen if the parent array does not describe a valid tree (e.g., if it contains a loop)
+	bool GlobalMatricesGetIndicesToProcessHelper(const int32 Index, const FGeometryDynamicCollection& DynamicCollection, int32 ParentNum, TArray<bool>& IsTransformComputed, FIndicesNeedMatricesArray& OutToProcess)
+	{
+		checkSlow(Index != FGeometryCollection::Invalid);
+		checkSlow(OutToProcess.IsEmpty());
+
+		OutToProcess.Add(Index);
+		bool bFoundRootOrComputed = false;
+		const int32 MaxDepth = ParentNum;
+		while (OutToProcess.Num() <= MaxDepth)
+		{
+			const int32 Parent = DynamicCollection.GetParent(OutToProcess.Last());
+			if (Parent == INDEX_NONE || IsTransformComputed[Parent])
+			{
+				bFoundRootOrComputed = true;
+				break;
+			}
+			OutToProcess.Add(Parent);
+		}
+		if (!ensureMsgf(bFoundRootOrComputed, TEXT("Geometry Collection has invalid parent hierarchy, could not find root to create global transforms")))
+		{
+			return false;
+		}
+
+		return true;
+	}
+
+	
 	// return false if we failed to get the indices, which can happen if the parent array does not describe a valid tree (e.g., if it contains a loop)
 	bool GlobalMatricesGetIndicesToProcessHelper(const int32 Index, const TManagedArray<int32>& Parents, TArray<bool>& IsTransformComputed, FIndicesNeedMatricesArray& OutToProcess)
 	{
@@ -303,7 +332,7 @@ namespace GeometryCollectionAlgo
 		while (OutToProcess.Num() <= MaxDepth)
 		{
 			int32 Parent = Parents[OutToProcess.Last()];
-			if (Parent < 0 || IsTransformComputed[Parent])
+			if (Parent == INDEX_NONE || IsTransformComputed[Parent])
 			{
 				bFoundRootOrComputed = true;
 				break;
@@ -318,6 +347,44 @@ namespace GeometryCollectionAlgo
 		return true;
 	}
 
+	template<typename TransformType>
+	void GlobalMatricesHelper(const int32 Index, const FGeometryDynamicCollection& DynamicCollection, const TManagedArray<TransformType>& Transform, TArray<bool>& IsTransformComputed, const TManagedArray<FTransform>* UniformScale, TArray<FTransform>& OutGlobalTransforms)
+	{
+		if (IsTransformComputed[Index])
+		{
+			return;
+		}
+
+		FIndicesNeedMatricesArray ToProcess;
+		if (!GlobalMatricesGetIndicesToProcessHelper(Index, DynamicCollection, Transform.Num(), IsTransformComputed, ToProcess))
+		{
+			return;
+		}
+
+		while (!ToProcess.IsEmpty())
+		{
+			const int32 ProcessIndex = ToProcess.Pop(false);
+			const int32 ParentIndex = DynamicCollection.GetParent(ProcessIndex);
+			FTransform Result = FTransform(Transform[ProcessIndex]);
+			if (ParentIndex != FGeometryCollection::Invalid)
+			{
+				Result *= OutGlobalTransforms[ParentIndex];
+			}
+
+			if (UniformScale)
+			{
+				OutGlobalTransforms[ProcessIndex] = (*UniformScale)[ProcessIndex] * Result;
+			}
+			else
+			{
+				OutGlobalTransforms[ProcessIndex] = Result;
+			}
+
+			IsTransformComputed[ProcessIndex] = true;
+		}
+	}
+
+
 	template<typename TransformType>
 	void GlobalMatricesHelper(const int32 Index, const TManagedArray<int32>& Parents, const TManagedArray<TransformType>& Transform, TArray<bool>& IsTransformComputed, const TManagedArray<FTransform>* UniformScale, TArray<FMatrix>& OutGlobalTransforms)
 	{
@@ -393,6 +460,44 @@ namespace GeometryCollectionAlgo
 		}
 	}
 
+	FTransform GlobalMatricesHelperForIndices(const int32 Index, const FGeometryDynamicCollection& DynamicCollection, const TManagedArray<FTransform3f>& Transform, TArray<bool>& IsTransformComputed, const TManagedArray<FTransform>* UniformScale, TArray<FTransform>& TransformCache)
+	{
+		if (IsTransformComputed[Index])
+		{
+			return TransformCache[Index];
+		}
+
+		FIndicesNeedMatricesArray ToProcess;
+		if (!GlobalMatricesGetIndicesToProcessHelper(Index, DynamicCollection, Transform.Num(), IsTransformComputed, ToProcess))
+		{
+			return TransformCache[Index];
+		}
+
+		while (!ToProcess.IsEmpty())
+		{
+			const int32 ProcessIndex = ToProcess.Pop(false);
+			const int32 ParentIndex = DynamicCollection.GetParent(ProcessIndex);
+			FTransform Result = FTransform(Transform[ProcessIndex]);
+			if (ParentIndex != FGeometryCollection::Invalid)
+			{
+				Result *= TransformCache[ParentIndex];
+			}
+
+			if (UniformScale)
+			{
+				TransformCache[ProcessIndex] = (*UniformScale)[ProcessIndex] * Result;
+			}
+			else
+			{
+				TransformCache[ProcessIndex] = Result;
+			}
+
+			IsTransformComputed[ProcessIndex] = true;
+		}
+
+		return TransformCache[Index];
+	}
+
 	template<typename TransformType>
 	FTransform GlobalMatricesHelperForIndices(const int32 Index, const TManagedArray<int32>& Parents, const TManagedArray<TransformType>& Transform, TArray<bool>& IsTransformComputed, const TManagedArray<FTransform>* UniformScale, TArray<FTransform>& TransformCache)
 	{
@@ -433,12 +538,27 @@ namespace GeometryCollectionAlgo
 	}
 
 
+	FTransform GlobalMatrix(const TManagedArray<FTransform3f>& RelativeTransforms, const FGeometryDynamicCollection& DynamicCollection, int32 Index)
+	{
+		FTransform Transform = FTransform::Identity;
+
+		if (RelativeTransforms.IsValidIndex(Index))
+		{
+			while (Index != FGeometryCollection::Invalid)
+			{
+				Transform = Transform * FTransform(RelativeTransforms[Index]);
+				Index = DynamicCollection.GetParent(Index);
+			}
+		}
+		return Transform;
+	}
+
 	template<class TransformType>
 	FTransform GlobalMatrixTemplate(const TManagedArray<TransformType>& RelativeTransforms, const TManagedArray<int32>& Parents, int32 Index)
 	{
 		FTransform Transform = FTransform::Identity;
 
-		if (0 <= Index && Index < RelativeTransforms.Num())
+		if (RelativeTransforms.IsValidIndex(Index))
 		{
 			while (Index != FGeometryCollection::Invalid)
 			{
@@ -459,6 +579,21 @@ namespace GeometryCollectionAlgo
 		return GlobalMatrixTemplate(RelativeTransforms, Parents, Index);
 	}
 
+	void GlobalMatrices(const TManagedArray<FTransform3f>& RelativeTransforms, const FGeometryDynamicCollection& DynamicCollection, const TArray<int32>& Indices, TArray<FTransform>& OutGlobalTransforms)
+	{
+		TArray<bool> IsTransformComputed;
+		IsTransformComputed.AddDefaulted(RelativeTransforms.Num());
+
+		TArray<FTransform> TransformCache;
+		TransformCache.SetNumUninitialized(RelativeTransforms.Num(), false);
+
+		OutGlobalTransforms.SetNumUninitialized(Indices.Num(), false);
+		for (int Idx = 0; Idx < Indices.Num(); Idx++)
+		{
+			OutGlobalTransforms[Idx] = GlobalMatricesHelperForIndices(Indices[Idx], DynamicCollection, RelativeTransforms, IsTransformComputed, nullptr, TransformCache);
+		}
+	}
+
 	template<class TransformType>
 	void GlobalMatricesTemplate(const TManagedArray<TransformType>& RelativeTransforms, const TManagedArray<int32>& Parents, const TArray<int32>& Indices, TArray<FTransform>& OutGlobalTransforms)
 	{
@@ -515,6 +650,21 @@ namespace GeometryCollectionAlgo
 		}
 	}
 
+	void GlobalMatrices(const TManagedArray<FTransform3f>& RelativeTransforms, const FGeometryDynamicCollection& DynamicCollection, TArray<FTransform>& OutGlobalTransforms)
+	{
+		int32 NumTransforms = RelativeTransforms.Num();
+
+		TArray<bool> IsTransformComputed;
+		IsTransformComputed.AddDefaulted(NumTransforms);
+
+		OutGlobalTransforms.SetNumUninitialized(NumTransforms, false);
+
+		for (int BoneIdx = 0; BoneIdx < NumTransforms; ++BoneIdx)
+		{
+			GlobalMatricesHelper<FTransform3f>(BoneIdx, DynamicCollection, RelativeTransforms, IsTransformComputed, nullptr, OutGlobalTransforms);
+		}
+	}
+
 	template<typename MatrixType, typename TransformType>
 	void GlobalMatrices(const TManagedArray<TransformType>& RelativeTransforms, const TManagedArray<int32>& Parents, TArray<MatrixType>& OutGlobalTransforms)
 	{
diff --git a/Engine/Source/Runtime/Experimental/Chaos/Private/GeometryCollection/ManagedArrayCollection.cpp b/Engine/Source/Runtime/Experimental/Chaos/Private/GeometryCollection/ManagedArrayCollection.cpp
--- a/Engine/Source/Runtime/Experimental/Chaos/Private/GeometryCollection/ManagedArrayCollection.cpp
+++ b/Engine/Source/Runtime/Experimental/Chaos/Private/GeometryCollection/ManagedArrayCollection.cpp
@@ -554,9 +554,14 @@ void FManagedArrayCollection::CopyAttribute(const FManagedArrayCollection& InCol
 
 	const FValueType& OriginalValue = InCollection.Map[SrcKey];
 	const FValueType& DestValue = Map[DestKey];
-	// Now a FTransform3f can be copy in a FTransform, and the way around
-	// check(OriginalValue.ArrayType == DestValue.ArrayType);
-	DestValue.Value->Init(*OriginalValue.Value);
+	if (OriginalValue.ArrayType == DestValue.ArrayType)
+	{
+		DestValue.Value->Init(*OriginalValue.Value);
+	}
+	else
+	{
+		DestValue.Value->Convert(*OriginalValue.Value);
+	}
 }
 
 bool FManagedArrayCollection::IsConnected(FName StartingNode, FName TargetNode)
diff --git a/Engine/Source/Runtime/Experimental/Chaos/Private/GeometryCollectionProxyData.cpp b/Engine/Source/Runtime/Experimental/Chaos/Private/GeometryCollectionProxyData.cpp
--- a/Engine/Source/Runtime/Experimental/Chaos/Private/GeometryCollectionProxyData.cpp
+++ b/Engine/Source/Runtime/Experimental/Chaos/Private/GeometryCollectionProxyData.cpp
@@ -12,20 +12,21 @@ GeometryCollectionProxyData.cpp:
 * FTransformDynamicCollection (FManagedArrayCollection)
 */
 
-FTransformDynamicCollection::FTransformDynamicCollection()
+FTransformDynamicCollection::FTransformDynamicCollection(const FGeometryCollection* InRestCollection)
 	: FManagedArrayCollection()
+	, RestCollection(InRestCollection)
 {
+	check(RestCollection != nullptr);
 	Construct();
 }
 
-
 void FTransformDynamicCollection::Construct()
 {
 	FManagedArrayCollection::FConstructionParameters TransformDependency(FTransformCollection::TransformGroup);
 
 	// Transform Group
 	AddExternalAttribute<FTransform3f>(FTransformCollection::TransformAttribute, FTransformCollection::TransformGroup, Transform);
-	AddExternalAttribute<int32>(FTransformCollection::ParentAttribute, FTransformCollection::TransformGroup, Parent);
+	AddExternalAttribute<bool>(FTransformCollection::ParentAttribute, FTransformCollection::TransformGroup, HasParent);
 	AddExternalAttribute<TSet<int32>>(FTransformCollection::ChildrenAttribute, FTransformCollection::TransformGroup, Children);
 }
 
@@ -44,6 +45,27 @@ const TArray<FTransform3f>& FTransformDynamicCollection::GetTransforms() const
 	return Transform.GetConstArray();
 }
 
+const TManagedArray<bool>& FTransformDynamicCollection::GetHasParent() const
+{
+	return HasParent;
+}
+
+bool FTransformDynamicCollection::GetHasParent(int32 Index) const
+{
+	return HasParent[Index];
+}
+
+void FTransformDynamicCollection::SetHasParent(int32 Index, bool Value)
+{
+	HasParent[Index] = Value;
+}
+
+int32 FTransformDynamicCollection::GetParent(int32 Index) const 
+{
+	check(RestCollection != nullptr);
+	return HasParent[Index] ? RestCollection->Parent[Index] : INDEX_NONE;
+}
+
 
 /*
 * FGeometryDynamicCollection (FTransformDynamicCollection)
@@ -62,8 +84,8 @@ const FName FGeometryDynamicCollection::SharedImplicitsAttribute("SharedImplicit
 const FName FGeometryDynamicCollection::CollisionMaskAttribute("CollisionMask");
 const FName FGeometryDynamicCollection::CollisionGroupAttribute("CollisionGroup");
 
-FGeometryDynamicCollection::FGeometryDynamicCollection()
-	: FTransformDynamicCollection()
+FGeometryDynamicCollection::FGeometryDynamicCollection(const FGeometryCollection* InRestCollection)
+	: FTransformDynamicCollection(InRestCollection)
 {
 	// Transform Group
 	AddExternalAttribute<bool>(FGeometryDynamicCollection::ActiveAttribute, FTransformCollection::TransformGroup, Active);
@@ -124,8 +146,8 @@ void FGeometryDynamicCollection::FInitialVelocityFacade::CopyFrom(const FGeometr
 FGeometryCollectionDynamicStateFacade::FGeometryCollectionDynamicStateFacade(FManagedArrayCollection& InCollection)
 	: ActiveAttribute(InCollection, FGeometryDynamicCollection::ActiveAttribute,  FTransformCollection::TransformGroup)
 	, DynamicStateAttribute(InCollection, FGeometryDynamicCollection::DynamicStateAttribute,  FTransformCollection::TransformGroup)
-	, ChildrenAttribute(InCollection, "Children",  FTransformCollection::TransformGroup)
-	, ParentAttribute(InCollection, "Parent",  FTransformCollection::TransformGroup)
+	, ChildrenAttribute(InCollection, FTransformCollection::ChildrenAttribute,  FTransformCollection::TransformGroup)
+	, HasParentAttribute(InCollection, FTransformCollection::ParentAttribute, FTransformCollection::TransformGroup)
 	, InternalClusterParentTypeAttribute(InCollection, "InternalClusterParentTypeArray", FGeometryCollection::TransformGroup)
 {
 }
@@ -135,9 +157,8 @@ bool FGeometryCollectionDynamicStateFacade::IsValid() const
 	return ActiveAttribute.IsValid()
 		&& DynamicStateAttribute.IsValid()
 		&& ChildrenAttribute.IsValid()
-		&& ParentAttribute.IsValid()
-		&& InternalClusterParentTypeAttribute.IsValid()
-		;
+		&& HasParentAttribute.IsValid()
+		&& InternalClusterParentTypeAttribute.IsValid();
 }
 
 bool FGeometryCollectionDynamicStateFacade::IsActive(int32 TransformIndex) const
@@ -165,7 +186,7 @@ bool FGeometryCollectionDynamicStateFacade::HasChildren(int32 TransformIndex) co
 bool FGeometryCollectionDynamicStateFacade::HasBrokenOff(int32 TransformIndex) const
 {
 	const bool bIsActive = IsActive(TransformIndex);
-	const bool bHasParent = (ParentAttribute.Get()[TransformIndex] != INDEX_NONE);
+	const bool bHasParent = HasParentAttribute.Get()[TransformIndex];
 	return bIsActive && (!bHasParent) && IsDynamicOrSleeping(TransformIndex);
 }
 
diff --git a/Engine/Source/Runtime/Experimental/Chaos/Private/PhysicsProxy/GeometryCollectionPhysicsProxy.cpp b/Engine/Source/Runtime/Experimental/Chaos/Private/PhysicsProxy/GeometryCollectionPhysicsProxy.cpp
--- a/Engine/Source/Runtime/Experimental/Chaos/Private/PhysicsProxy/GeometryCollectionPhysicsProxy.cpp
+++ b/Engine/Source/Runtime/Experimental/Chaos/Private/PhysicsProxy/GeometryCollectionPhysicsProxy.cpp
@@ -467,7 +467,7 @@ FGeometryCollectionPhysicsProxy::FGeometryCollectionPhysicsProxy(
 	, LastSyncCountGT(MAX_uint32)
 #endif
 	, CollisionParticlesPerObjectFraction(CollisionParticlesPerObjectFractionDefault)
-
+	, PhysicsThreadCollection(Parameters.RestCollection)
 	, GameThreadCollection(GameThreadCollectionIn)
 	, GameThreadPerFrameData(SimulationParameters)
 	, MaterialOverrideMassScaleMultiplierChange(0)
@@ -647,7 +647,7 @@ void FGeometryCollectionPhysicsProxy::Initialize(Chaos::FPBDRigidsEvolutionBase
 				// step 1 : find parents
 				for(const int32 ChildIndex: ChildrenToCheckForParentFix)
 				{
-					const int32 ParentIndex = GameThreadCollection.Parent[ChildIndex];
+					const int32 ParentIndex = GameThreadCollection.GetParent(ChildIndex);
 					if (ParentIndex != INDEX_NONE)
 					{
 						ParentToPotentiallyFix.Add(ParentIndex);
@@ -782,7 +782,7 @@ void FGeometryCollectionPhysicsProxy::Initialize(Chaos::FPBDRigidsEvolutionBase
 		// make sure we set Activate the right way when clustering is enabled ( only root should be enabled at start ) 
 		for (int32 TransformIndex = 0; TransformIndex < NumTransforms; ++TransformIndex)
 		{
-			bool bIsRoot = (GameThreadCollection.Parent[TransformIndex] == INDEX_NONE);
+			const bool bIsRoot = !GameThreadCollection.GetHasParent(TransformIndex);
 			GameThreadCollection.Active[TransformIndex] = bIsRoot;
 			PhysicsThreadCollection.Active[TransformIndex] = bIsRoot;
 
@@ -1065,7 +1065,6 @@ void FGeometryCollectionPhysicsProxy::InitializeBodiesPT(Chaos::FPBDRigidsSolver
 		const TManagedArray<Chaos::FImplicitObjectPtr>& Implicits = DynamicCollection.Implicits;
 		const TManagedArray<TUniquePtr<FCollisionStructureManager::FSimplicial>>& Simplicials = DynamicCollection.Simplicials;
 		const TManagedArray<TSet<int32>>& Children = DynamicCollection.Children;
-		const TManagedArray<int32>& Parent = DynamicCollection.Parent;
 
 		// In PushToPhysicsState, we're going to compute a relative transform from Parameters.PrevWorldTransform
 		// to a particle's current world transform to get its relative transform. Then, that relative transform is
@@ -1076,7 +1075,7 @@ void FGeometryCollectionPhysicsProxy::InitializeBodiesPT(Chaos::FPBDRigidsSolver
 		Parameters.PrevWorldTransform = Parameters.WorldTransform;
 
 		TArray<FTransform> Transform;
-		GeometryCollectionAlgo::GlobalMatrices(DynamicCollection.Transform, Parent, Transform);
+		GeometryCollectionAlgo::GlobalMatrices(DynamicCollection.Transform, DynamicCollection, Transform);
 
 		// Here Clean up Additional particles
 		TBitArray<> EffectiveParticles;
@@ -1633,7 +1632,6 @@ Chaos::TPBDGeometryCollectionParticleHandle<Chaos::FReal, 3>* FGeometryCollectio
 {
 	FGeometryDynamicCollection& DynamicCollection = PhysicsThreadCollection;
 	TManagedArray<int32>& DynamicState = DynamicCollection.DynamicState;
-	TManagedArray<int32>& ParentIndex = DynamicCollection.Parent;
 	TManagedArray<FTransform3f>& Transform = DynamicCollection.Transform;
 	const TManagedArray<FTransform>& MassToLocal = Parameters.RestCollection->GetAttribute<FTransform>(MassToLocalAttributeName, FTransformCollection::TransformGroup);
 	const TManagedArray<Chaos::FImplicitObjectPtr>& Implicits = DynamicCollection.Implicits;
@@ -1641,7 +1639,7 @@ Chaos::TPBDGeometryCollectionParticleHandle<Chaos::FReal, 3>* FGeometryCollectio
 	Chaos::Facades::FCollectionAnchoringFacade AnchoringFacade(DynamicCollection);
 
 	//If we are a root particle use the world transform, otherwise set the relative transform
-	const FTransform CollectionSpaceTransform = GeometryCollectionAlgo::GlobalMatrix(Transform, ParentIndex, CollectionClusterIndex);
+	const FTransform CollectionSpaceTransform = GeometryCollectionAlgo::GlobalMatrix(Transform, DynamicCollection, CollectionClusterIndex);
 	const Chaos::TRigidTransform<Chaos::FReal, 3> ParticleTM = MassToLocal[CollectionClusterIndex] * CollectionSpaceTransform * Parameters.WorldTransform;
 
 	// Gather unique indices from GT to pass into PT handle creation
@@ -1719,7 +1717,7 @@ Chaos::TPBDGeometryCollectionParticleHandle<Chaos::FReal, 3>* FGeometryCollectio
 		Shape->SetMaterial(Parameters.PhysicalMaterialHandle);
 	}
 
-	const FTransform ParentTransform = GeometryCollectionAlgo::GlobalMatrix(DynamicCollection.Transform, DynamicCollection.Parent, CollectionClusterIndex);
+	const FTransform ParentTransform = GeometryCollectionAlgo::GlobalMatrix(DynamicCollection.Transform, DynamicCollection, CollectionClusterIndex);
 
 	// Populate bounds as we didn't pass a shared implicit to PopulateSimulatedParticle this will have been skipped, now that we have the full cluster we can build it
 	if (Handle->GetGeometry() && Handle->GetGeometry()->HasBoundingBox())
@@ -1752,14 +1750,13 @@ FGeometryCollectionPhysicsProxy::BuildClusters_Internal(
 
 	FGeometryDynamicCollection& DynamicCollection = PhysicsThreadCollection;
 	TManagedArray<int32>& DynamicState = DynamicCollection.DynamicState;
-	TManagedArray<int32>& ParentIndex = DynamicCollection.Parent;
 	TManagedArray<FTransform3f>& Transform = DynamicCollection.Transform;
 	const TManagedArray<FTransform>& MassToLocal = Parameters.RestCollection->GetAttribute<FTransform>(MassToLocalAttributeName, FTransformCollection::TransformGroup);
 	//TManagedArray<TSharedPtr<FCollisionStructureManager::FSimplicial> >& Simplicials = DynamicCollection.Simplicials;
 	TManagedArray<Chaos::FImplicitObjectPtr>& Implicits = DynamicCollection.Implicits;
 
 	//If we are a root particle use the world transform, otherwise set the relative transform
-	const FTransform CollectionSpaceTransform = GeometryCollectionAlgo::GlobalMatrix(Transform, ParentIndex, CollectionClusterIndex);
+	const FTransform CollectionSpaceTransform = GeometryCollectionAlgo::GlobalMatrix(Transform, DynamicCollection, CollectionClusterIndex);
 	const Chaos::TRigidTransform<Chaos::FReal, 3> ParticleTM = MassToLocal[CollectionClusterIndex] * CollectionSpaceTransform * Parameters.WorldTransform;
 
 	//create new cluster particle
@@ -1788,7 +1785,7 @@ FGeometryCollectionPhysicsProxy::BuildClusters_Internal(
 	// Construct an active cluster particle, disable children, derive M and I from children:
 	Chaos::TPBDRigidClusteredParticleHandle<Chaos::FReal, 3>* Parent =
 		static_cast<Chaos::FPBDRigidsSolver*>(Solver)->GetEvolution()->GetRigidClustering().CreateClusterParticle(
-			(ParentIndex[CollectionClusterIndex] == INDEX_NONE) ? Parameters.ClusterGroupIndex : 0, 
+			(!DynamicCollection.GetHasParent(CollectionClusterIndex)) ? Parameters.ClusterGroupIndex : 0, 
 			MoveTemp(ChildHandlesCopy),
 			ClusterCreationParameters,
 			Implicits[CollectionClusterIndex], // union from children if null
@@ -1854,7 +1851,7 @@ FGeometryCollectionPhysicsProxy::BuildClusters_Internal(
 		Shape->SetMaterial(Parameters.PhysicalMaterialHandle);
 	}
 
-	const FTransform ParentTransform = GeometryCollectionAlgo::GlobalMatrix(DynamicCollection.Transform, DynamicCollection.Parent, CollectionClusterIndex);
+	const FTransform ParentTransform = GeometryCollectionAlgo::GlobalMatrix(DynamicCollection.Transform, DynamicCollection, CollectionClusterIndex);
 
 	int32 MinCollisionGroup = INT_MAX;
 	for(int32 Idx=0; Idx < ChildHandles.Num(); Idx++)
@@ -2622,9 +2619,8 @@ void FGeometryCollectionPhysicsProxy::SetNotifyGlobalCrumblings_External(bool bN
 int32 FGeometryCollectionPhysicsProxy::CalculateHierarchyLevel(const FGeometryDynamicCollection& DynamicCollection, int32 TransformIndex)
 {
 	int32 Level = 0;
-	while (DynamicCollection.Parent[TransformIndex] != INDEX_NONE)
+	while ((TransformIndex = DynamicCollection.GetParent(TransformIndex)) != INDEX_NONE)
 	{
-		TransformIndex = DynamicCollection.Parent[TransformIndex];
 		Level++;
 	}
 	return Level;
@@ -2658,8 +2654,7 @@ TBitArray<> FGeometryCollectionPhysicsProxy::CalculateClustersToCreateFromChildr
 			// step 1 : find parents
 			for (const int32 ChildIndex : ChildrenToCheckForParentFix)
 			{
-				const int32 ParentIndex = DynamicCollection.Parent[ChildIndex];
-				if (ParentIndex != INDEX_NONE)
+				if (int32 ParentIndex = DynamicCollection.GetParent(ChildIndex); ParentIndex != INDEX_NONE)
 				{
 					ParentToPotentiallyFix.Add(ParentIndex);
 				}
@@ -3341,10 +3336,10 @@ static void UpdateParticleHandleTransformIfNeeded(Chaos::FPBDRigidsSolver& Curre
 inline static bool RemoveFromParentInCollection_Internal(FGeometryDynamicCollection& PhysicsThreadCollection, int32 TransformGroupIndex)
 {
 	// CONTEXT : PHYSICS THREAD
-	if (PhysicsThreadCollection.Parent[TransformGroupIndex] != INDEX_NONE)
+	if (int32 ParentIndex = PhysicsThreadCollection.GetParent(TransformGroupIndex); ParentIndex != INDEX_NONE)
 	{
-		PhysicsThreadCollection.Children[PhysicsThreadCollection.Parent[TransformGroupIndex]].Remove(TransformGroupIndex);
-		PhysicsThreadCollection.Parent[TransformGroupIndex] = INDEX_NONE;
+		PhysicsThreadCollection.Children[ParentIndex].Remove(TransformGroupIndex);
+		PhysicsThreadCollection.SetHasParent(TransformGroupIndex, false);
 		return true;
 	}
 	return false;
@@ -3402,7 +3397,7 @@ void FGeometryCollectionPhysicsProxy::BufferPhysicsResults_Internal(Chaos::FPBDR
 
 			FGeometryCollectionResults::FStateData StateData;
 			StateData.TransformIndex = TransformGroupIndex;
-			StateData.ParentTransformIndex = PhysicsThreadCollection.Parent[TransformGroupIndex];
+			StateData.HasParent = PhysicsThreadCollection.GetHasParent(TransformGroupIndex);
 			StateData.InternalClusterUniqueIdx = INDEX_NONE;
 			StateData.State.DisabledState = Handle->Disabled();
 			StateData.State.HasInternalClusterParent = false;
@@ -3429,16 +3424,16 @@ void FGeometryCollectionPhysicsProxy::BufferPhysicsResults_Internal(Chaos::FPBDR
 			Chaos::FPBDRigidClusteredParticleHandle* ClusterParent = Handle->Parent();
 
 			// Has parent changed? ( can be a new one or nullptr)
-			const bool bHasNewParent = (SolverClusterID[TransformGroupIndex] != ClusterParent) || (!ClusterParent && PhysicsThreadCollection.Parent[TransformGroupIndex]);
+			const bool bHasNewParent = (SolverClusterID[TransformGroupIndex] != ClusterParent) || (!ClusterParent && PhysicsThreadCollection.GetHasParent(TransformGroupIndex));
 			if (bHasNewParent)
 			{
 				// Force all driven rigid bodies out of the transform hierarchy ( because the new parent can only be an internal cluster )
 				if (RemoveFromParentInCollection_Internal(PhysicsThreadCollection, TransformGroupIndex))
 				{
-					// If the parent of this NON DISABLED body is set to anything other than INDEX_NONE,
+					// If the parent flag of this NON DISABLED body is set to false,
 					// then it was just unparented, likely either by rigid clustering or by fields.  We
 					// need to force all such enabled rigid bodies out of the transform hierarchy.
-					StateData.ParentTransformIndex = INDEX_NONE;
+					StateData.HasParent = false;
 
 					// Indicate that this object needs to be updated and the proxy is active.
 					StateData.State.DisabledState = false;
@@ -3752,7 +3747,7 @@ bool FGeometryCollectionPhysicsProxy::PullNonInterpolatableDataFromSinglePhysics
 			bIsCollectionDirty = true;
 		}
 
-		if (UpdateValue(GameThreadCollection.Parent[TransformGroupIndex], StateData.ParentTransformIndex))
+		if (UpdateValue(GameThreadCollection.HasParent[TransformGroupIndex], StateData.HasParent))
 		{
 			bIsCollectionDirty = true;
 		}
@@ -3945,7 +3940,7 @@ bool FGeometryCollectionPhysicsProxy::PullFromPhysicsState(const Chaos::FDirtyGe
 						{
 							const FParticle& GTParticle = *GTParticles[TransformIndex];
 							const FTransform& ParticleMassToLocal = MassToLocal[TransformIndex];
-							const int32 ParentTransformIndex = GameThreadCollection.Parent[TransformIndex];
+							const int32 ParentTransformIndex = GameThreadCollection.GetParent(TransformIndex);
 
 							const FTransform& WorldTransform = ParticleMassToLocal.Inverse() * FTransform { GTParticle.R(), GTParticle.X() };
 
diff --git a/Engine/Source/Runtime/Experimental/Chaos/Public/GeometryCollection/GeometryCollectionAlgo.h b/Engine/Source/Runtime/Experimental/Chaos/Public/GeometryCollection/GeometryCollectionAlgo.h
--- a/Engine/Source/Runtime/Experimental/Chaos/Public/GeometryCollection/GeometryCollectionAlgo.h
+++ b/Engine/Source/Runtime/Experimental/Chaos/Public/GeometryCollection/GeometryCollectionAlgo.h
@@ -7,6 +7,7 @@
 #include "GeometryCollection/TransformCollection.h"
 
 class FGeometryCollection;
+class FGeometryDynamicCollection;
 
 namespace GeometryCollectionAlgo
 {
@@ -134,6 +135,7 @@ namespace GeometryCollectionAlgo
 	*/
 	FTransform CHAOS_API GlobalMatrix(const TManagedArray<FTransform>& RelativeTransforms, const TManagedArray<int32>& Parents, int32 Index);
 	FTransform CHAOS_API GlobalMatrix(const TManagedArray<FTransform3f>& RelativeTransforms, const TManagedArray<int32>& Parents, int32 Index);
+	FTransform CHAOS_API GlobalMatrix(const TManagedArray<FTransform3f>& RelativeTransforms, const FGeometryDynamicCollection& DynamicCollection, int32 Index);
 
 
 	/*
@@ -141,6 +143,7 @@ namespace GeometryCollectionAlgo
 	*/
 	void CHAOS_API GlobalMatrices(const TManagedArray<FTransform>& RelativeTransforms, const TManagedArray<int32>& Parents, const TArray<int32>& Indices, TArray<FTransform>& Transforms);
 	void CHAOS_API GlobalMatrices(const TManagedArray<FTransform3f>& RelativeTransforms, const TManagedArray<int32>& Parents, const TArray<int32>& Indices, TArray<FTransform>& Transforms);
+	void CHAOS_API GlobalMatrices(const TManagedArray<FTransform3f>& RelativeTransforms, const FGeometryDynamicCollection& DynamicCollection, const TArray<int32>& Indices, TArray<FTransform>& Transforms);
 
 	/*
 	 *  Recursively traverse from a root node down
@@ -156,6 +159,8 @@ namespace GeometryCollectionAlgo
 	template<typename MatrixType, typename TransformType>
 	void CHAOS_API GlobalMatrices(const TManagedArray<TransformType>& RelativeTransforms, const TManagedArray<int32>& Parents, TArray<MatrixType>& Transforms);
 
+	void CHAOS_API GlobalMatrices(const TManagedArray<FTransform3f>& RelativeTransforms, const FGeometryDynamicCollection& DynamicCollection, TArray<FTransform>& Transforms);
+
 
 	/*
 	*  Gets pairs of elements whose bounding boxes overlap.
diff --git a/Engine/Source/Runtime/Experimental/Chaos/Public/GeometryCollection/ManagedArray.h b/Engine/Source/Runtime/Experimental/Chaos/Public/GeometryCollection/ManagedArray.h
--- a/Engine/Source/Runtime/Experimental/Chaos/Public/GeometryCollection/ManagedArray.h
+++ b/Engine/Source/Runtime/Experimental/Chaos/Public/GeometryCollection/ManagedArray.h
@@ -179,6 +179,11 @@ protected:
 	*/
 	virtual void Init(const FManagedArrayBase& ) {};
 
+	/**
+	* Convert from a predefined Array, the managed array itself should have defined its conversion procedure
+	*/
+	virtual void Convert(const FManagedArrayBase&) { check(false); /* This type has no conversion process defined*/ };
+
 	/**
 	* Copy a range of values from the ConstArray into this
 	*/
@@ -673,8 +678,6 @@ void InitHelper(TArray<TDst>& Array, const TManagedArrayBase<TSrc>& NewTypedArra
 	}
 }
 
-
-
 template <typename T>
 void InitHelper(TArray<TUniquePtr<T>>& Array, const TManagedArrayBase<TUniquePtr<T>>& NewTypedArray, int32 Size)
 {
@@ -1001,7 +1004,7 @@ public:
 		}
 
 	}
-private:
+protected:
 	/**
 	* Protected Resize to prevent external resizing of the array
 	*
@@ -1205,27 +1208,39 @@ protected:
 	/**
 	* Init from a predefined Array of matching type
 	*/
-	virtual void Init(const FManagedArrayBase& NewArray) override
+	virtual void Convert(const FManagedArrayBase& NewArray) override
 	{
-		if (NewArray.GetTypeSize() == GetTypeSize())
-		{
-			TManagedArrayBase::Init(NewArray);
-		}
-		else
-		{
-			check(NewArray.GetTypeSize() == 2 * GetTypeSize());
-			check(NewArray.GetTypeSize() == sizeof(FTransform));
-			const TManagedArrayBase<FTransform>& NewTypedArray = static_cast<const TManagedArrayBase<FTransform>&>(NewArray);
-			const int32 Size = NewTypedArray.Num();
-			Resize(Size);
-			InitHelper(Array, NewTypedArray, Size);
-		}
+		check(NewArray.GetTypeSize() != GetTypeSize());
+		check(NewArray.GetTypeSize() == 2 * GetTypeSize());
+		check(NewArray.GetTypeSize() == sizeof(FTransform));
+		const TManagedArrayBase<FTransform>& NewTypedArray = static_cast<const TManagedArrayBase<FTransform>&>(NewArray);
+		const int32 Size = NewTypedArray.Num();
+		Resize(Size);
+		InitHelper(Array, NewTypedArray, Size);
 	}
 };
 
 template<>
 class TManagedArray<bool> : public FManagedBitArrayBase
-{};
+{
+protected:
+	/**
+	* Init from a predefined Array of matching type
+	*/
+	virtual void Convert(const FManagedArrayBase& NewArray) override
+	{
+		check(NewArray.GetTypeSize() != GetTypeSize());
+		check(NewArray.GetTypeSize() == sizeof(int32));
+		
+		const TManagedArrayBase<int32>& NewTypedArray = static_cast<const TManagedArrayBase<int32>&>(NewArray);
+		const int32 Size = NewTypedArray.Num();
+		Resize(Size);
+		for (int32 Index = 0; Index < Size; Index++)
+		{
+			Array[Index] = NewTypedArray[Index] != INDEX_NONE;
+		}
+	}
+};
 
 template<>
 class TManagedArray<TSet<int32>> : public TManagedArrayBase<TSet<int32>>
diff --git a/Engine/Source/Runtime/Experimental/Chaos/Public/GeometryCollectionProxyData.h b/Engine/Source/Runtime/Experimental/Chaos/Public/GeometryCollectionProxyData.h
--- a/Engine/Source/Runtime/Experimental/Chaos/Public/GeometryCollectionProxyData.h
+++ b/Engine/Source/Runtime/Experimental/Chaos/Public/GeometryCollectionProxyData.h
@@ -22,7 +22,7 @@ class FTransformDynamicCollection : public FManagedArrayCollection
 public:
 	typedef FManagedArrayCollection Super;
 
-	CHAOS_API FTransformDynamicCollection();
+	CHAOS_API FTransformDynamicCollection(const FGeometryCollection* InRestCollection);
 	FTransformDynamicCollection(FTransformDynamicCollection&) = delete;
 	FTransformDynamicCollection& operator=(const FTransformDynamicCollection&) = delete;
 	FTransformDynamicCollection(FTransformDynamicCollection&&) = delete;
@@ -32,17 +32,24 @@ public:
 	CHAOS_API void SetTransform(int32 Index, const FTransform3f& Transform);
 	CHAOS_API const TArray<FTransform3f>& GetTransforms() const;
 
+	CHAOS_API const TManagedArray<bool>& GetHasParent() const;
+	CHAOS_API bool GetHasParent(int32 Index) const;
+	CHAOS_API void SetHasParent(int32 Index, bool Value);
+	CHAOS_API int32 GetParent(int32 Index) const;
+
 	// Transform Group
-	TManagedArray<int32>        Parent;
 	TManagedArray<TSet<int32>>  Children;
 
 private:
+	TManagedArray<bool>         HasParent;
 	TManagedArray<FTransform3f> Transform;
 	/** Construct */
 	CHAOS_API void Construct();
 
 	friend class FGeometryCollectionPhysicsProxy;
 	friend class UGeometryCollectionComponent;
+
+	const FGeometryCollection* RestCollection;
 };
 
 
@@ -55,7 +62,7 @@ private:
 class FGeometryDynamicCollection : public FTransformDynamicCollection
 {
 public:
-	CHAOS_API FGeometryDynamicCollection();
+	CHAOS_API FGeometryDynamicCollection(const FGeometryCollection* InRestCollection);
 	FGeometryDynamicCollection(FGeometryDynamicCollection&) = delete;
 	FGeometryDynamicCollection& operator=(const FGeometryDynamicCollection&) = delete;
 	FGeometryDynamicCollection(FGeometryDynamicCollection&&) = delete;
@@ -164,7 +171,7 @@ private:
 	TManagedArrayAccessor<TSet<int32>> ChildrenAttribute;
 	
 	/** Current parent (potentially different from the initial parent) */
-	TManagedArrayAccessor<int32> ParentAttribute;
+	TManagedArrayAccessor<bool> HasParentAttribute;
 
 	/** type of internal state parent */
 	TManagedArrayAccessor<uint8> InternalClusterParentTypeAttribute;
@@ -400,8 +407,8 @@ public:
 	struct FStateData
 	{
 		int32  TransformIndex;
-		int32  ParentTransformIndex;
 		int32  InternalClusterUniqueIdx;
+		bool   HasParent;
 		FState State;
 	};
 
diff --git a/Engine/Source/Runtime/Experimental/GeometryCollectionEngine/Private/GeometryCollection/GeometryCollectionComponent.cpp b/Engine/Source/Runtime/Experimental/GeometryCollectionEngine/Private/GeometryCollection/GeometryCollectionComponent.cpp
--- a/Engine/Source/Runtime/Experimental/GeometryCollectionEngine/Private/GeometryCollection/GeometryCollectionComponent.cpp
+++ b/Engine/Source/Runtime/Experimental/GeometryCollectionEngine/Private/GeometryCollection/GeometryCollectionComponent.cpp
@@ -519,6 +519,7 @@ TManagedArray<FTransform3f>& UGeometryCollectionComponent::GetTransformArrayCopy
 	}
 	return *IndirectTransformArray;
 }
+
 void UGeometryCollectionComponent::ResetTransformArrayDynamic()
 {
 	IndirectTransformArray = nullptr;
@@ -528,6 +529,32 @@ const TManagedArray<FTransform>& UGeometryCollectionComponent::GetTransformArray
 	return RestCollection->GetGeometryCollection()->Transform;
 }
 
+TManagedArray<int32>& UGeometryCollectionComponent::GetParentArrayCopyOnWrite()
+{
+	if (!IndirectParentArray)
+	{
+		DynamicCollection->AddAttribute<int32>(FTransformCollection::ParentAttribute, FTransformCollection::TransformGroup);
+		DynamicCollection->CopyAttribute(*RestCollection->GetGeometryCollection(), FTransformCollection::ParentAttribute, FTransformCollection::TransformGroup);
+		IndirectParentArray = &DynamicCollection->ModifyAttribute<int32>(FTransformCollection::ParentAttribute, FTransformCollection::TransformGroup);
+		CopyOnWriteAttributeList.Add(reinterpret_cast<FManagedArrayBase**>(&IndirectParentArray));
+	}
+	return *IndirectParentArray;
+}
+
+int32 UGeometryCollectionComponent::GetParent(int32 Index) const
+{
+	if (DynamicCollection)
+	{
+		return DynamicCollection->GetParent(Index);
+	}
+	return RestCollection->GetGeometryCollection()->Parent[Index];
+}
+
+const TManagedArray<int32>& UGeometryCollectionComponent::GetParentArrayRest() const
+{
+	return RestCollection->GetGeometryCollection()->Parent;
+}
+
 UGeometryCollectionComponent::UGeometryCollectionComponent(const FObjectInitializer& ObjectInitializer)
 	: Super(ObjectInitializer)
 	, ChaosSolverActor(nullptr)
@@ -1374,11 +1401,11 @@ bool UGeometryCollectionComponent::DoCustomNavigableGeometryExport(FNavigableGeo
 	TArray<FTransform> GeomToComponent;
 	if (DynamicCollection)
 	{
-		GeometryCollectionAlgo::GlobalMatrices(GetTransformArray(), GetParentArray(), TransformIndexBuffer, GeomToComponent);
+		GeometryCollectionAlgo::GlobalMatrices(GetTransformArray(), *DynamicCollection, TransformIndexBuffer, GeomToComponent);
 	}
 	else
 	{
-		GeometryCollectionAlgo::GlobalMatrices(TManagedArray<FTransform>(GetCurrentRestTransforms()), GetParentArray(), TransformIndexBuffer, GeomToComponent);
+		GeometryCollectionAlgo::GlobalMatrices(TManagedArray<FTransform>(GetCurrentRestTransforms()), GetParentArrayRest(), TransformIndexBuffer, GeomToComponent);
 	}
 
 	OutVertexBuffer.AddUninitialized(VertexCount);
@@ -2917,18 +2944,17 @@ void UGeometryCollectionComponent::SetInitialClusterBreaks(const TArray<int32>&
 {
 	if (DynamicCollection)
 	{
-		TManagedArray<int32>& Parent = DynamicCollection->Parent;
 		TManagedArray <TSet<int32>>& Children = DynamicCollection->Children;
-		const int32 NumTransforms = Parent.Num();
+		const int32 NumTransforms = DynamicCollection->GetTransforms().Num();
 
 		for (int32 ReleaseIndex : ReleaseIndices)
 		{
 			if (ReleaseIndex < NumTransforms)
 			{
-				if (Parent[ReleaseIndex] > INDEX_NONE)
+				if (int32 ParentIndex = DynamicCollection->GetParent(ReleaseIndex); ParentIndex > INDEX_NONE)
 				{
-					Children[Parent[ReleaseIndex]].Remove(ReleaseIndex);
-					Parent[ReleaseIndex] = INDEX_NONE;
+					Children[ParentIndex].Remove(ReleaseIndex);
+					DynamicCollection->SetHasParent(ReleaseIndex, false);
 				}
 			}
 		}
@@ -3470,7 +3496,7 @@ void UGeometryCollectionComponent::ResetDynamicCollection()
 #endif
 	if (bCreateDynamicCollection && RestCollection && RestCollection->GetGeometryCollection())
 	{
-		DynamicCollection = MakeUnique<FGeometryDynamicCollection>();
+		DynamicCollection = MakeUnique<FGeometryDynamicCollection>(RestCollection->GetGeometryCollection().Get());
 		for (const auto DynamicArray : CopyOnWriteAttributeList)
 		{
 			*DynamicArray = nullptr;
@@ -5434,12 +5460,11 @@ const TArray<FTransform>& UGeometryCollectionComponent::FComponentSpaceTransform
 	if (bFastPath)
 	{
 		const TArray<int32>& BreadthFirstTransformIndices = Component->RestCollection->GetBreadthFirstTransformIndices();
-		const TArray<int32>& ParentArray = Component->GetParentArray().GetConstArray();
 
 		for (int32 Index = 0; Index < BreadthFirstTransformIndices.Num(); Index++)
 		{
 			const int32 TransformIndex = BreadthFirstTransformIndices[Index];
-			const int32 ParentTransformIndex = ParentArray[TransformIndex];
+			const int32 ParentTransformIndex = Component->GetParent(TransformIndex);
 
 
 			FTransform CurrentTransform;
@@ -5467,11 +5492,11 @@ const TArray<FTransform>& UGeometryCollectionComponent::FComponentSpaceTransform
 	{
 		if (Component->DynamicCollection)
 		{
-			GeometryCollectionAlgo::GlobalMatrices(Component->GetTransformArray(), Component->GetParentArray(), Transforms);
+			GeometryCollectionAlgo::GlobalMatrices(Component->GetTransformArray(), *Component->GetDynamicCollection(), Transforms);
 		}
 		else
 		{
-			GeometryCollectionAlgo::GlobalMatrices(TManagedArray<FTransform>(Component->GetCurrentRestTransforms()), Component->GetParentArray(), Transforms);
+			GeometryCollectionAlgo::GlobalMatrices(TManagedArray<FTransform>(Component->GetCurrentRestTransforms()), Component->GetParentArrayRest(), Transforms);
 		}
 		
 	}
diff --git a/Engine/Source/Runtime/Experimental/GeometryCollectionEngine/Public/GeometryCollection/GeometryCollectionComponent.h b/Engine/Source/Runtime/Experimental/GeometryCollectionEngine/Public/GeometryCollection/GeometryCollectionComponent.h
--- a/Engine/Source/Runtime/Experimental/GeometryCollectionEngine/Public/GeometryCollection/GeometryCollectionComponent.h
+++ b/Engine/Source/Runtime/Experimental/GeometryCollectionEngine/Public/GeometryCollection/GeometryCollectionComponent.h
@@ -885,7 +885,6 @@ public:
 	/* Transform group */																				\
 	COPY_ON_WRITE_ATTRIBUTE(FString, BoneName, FTransformCollection::TransformGroup)					\
 	COPY_ON_WRITE_ATTRIBUTE(FLinearColor, BoneColor, FTransformCollection::TransformGroup)				\
-	COPY_ON_WRITE_ATTRIBUTE(int32, Parent, FTransformCollection::TransformGroup)						\
 	COPY_ON_WRITE_ATTRIBUTE(TSet<int32>, Children, FTransformCollection::TransformGroup)				\
 	COPY_ON_WRITE_ATTRIBUTE(int32, TransformToGeometryIndex, FTransformCollection::TransformGroup)		\
 	COPY_ON_WRITE_ATTRIBUTE(int32, ExemplarIndex, FTransformCollection::TransformGroup)					\
@@ -900,6 +899,12 @@ public:
 	private:
 		TManagedArray<FTransform3f>* IndirectTransformArray;
 	public:
+	GEOMETRYCOLLECTIONENGINE_API TManagedArray<int32>& GetParentArrayCopyOnWrite();
+	GEOMETRYCOLLECTIONENGINE_API int32 GetParent(int32 Index) const;
+	GEOMETRYCOLLECTIONENGINE_API const TManagedArray<int32>& GetParentArrayRest() const;
+	private:
+		TManagedArray<int32>* IndirectParentArray;
+	public:
 
 
 	UPROPERTY(EditAnywhere, NoClear, BlueprintReadOnly, Category = "ChaosPhysics")
