Clean up : Remove friendship from Dynamic Collection with Geometry Collection Component and Geometry Collection Component Physics Proxy

diff --git a/Engine/Source/Programs/HeadlessChaos/Private/GeometryCollection/GeometryCollectionTestClustering.cpp b/Engine/Source/Programs/HeadlessChaos/Private/GeometryCollection/GeometryCollectionTestClustering.cpp
--- a/Engine/Source/Programs/HeadlessChaos/Private/GeometryCollection/GeometryCollectionTestClustering.cpp
+++ b/Engine/Source/Programs/HeadlessChaos/Private/GeometryCollection/GeometryCollectionTestClustering.cpp
@@ -685,9 +685,8 @@ namespace GeometryCollectionTest
 		EXPECT_EQ(ClusterMap.Num(), 2);
 		EXPECT_TRUE(ClusterMapContains(ClusterMap, Collection1Handles[2], { Collection1Handles[1],Collection1Handles[0] }));
 		EXPECT_TRUE(ClusterMapContains(ClusterMap, Collection1Handles[3], { Collection1Handles[2] }));
-		const TManagedArray<FTransform3f> Transforms1(DynamicCollection1->GetTransforms());
 		TArray<FTransform> Collection1_PreReleaseTM; 
-		GeometryCollectionAlgo::GlobalMatrices(Transforms1, *DynamicCollection1.Get(), Collection1_PreReleaseTM);
+		GeometryCollectionAlgo::Private::GlobalMatrices(*DynamicCollection1.Get(), Collection1_PreReleaseTM);
 		for (int Idx = 0; Idx < Collection1_PreReleaseTM.Num(); Idx++) {
 			EXPECT_TRUE( (Collection1_PreReleaseTM[Idx].GetTranslation()-Collection1_InitialTM[Idx].GetTranslation()).Size()<KINDA_SMALL_NUMBER);
 		}
@@ -698,9 +697,8 @@ namespace GeometryCollectionTest
 		EXPECT_EQ(SovlerParticleHandles.Size(), 4);
 		EXPECT_EQ(ClusterMap.Num(), 1);
 		EXPECT_TRUE(ClusterMapContains(ClusterMap, Collection1Handles[2], { Collection1Handles[1],Collection1Handles[0] }));
-		const TManagedArray<FTransform3f> Transforms2(DynamicCollection1->GetTransforms());
 		TArray<FTransform> Collection1_PostReleaseTM; 
-		GeometryCollectionAlgo::GlobalMatrices(Transforms2, *DynamicCollection1.Get(), Collection1_PostReleaseTM);
+		GeometryCollectionAlgo::Private::GlobalMatrices(*DynamicCollection1.Get(), Collection1_PostReleaseTM);
 		for (int Idx = 0; Idx < Collection1_PostReleaseTM.Num(); Idx++) {
 			EXPECT_TRUE((Collection1_PostReleaseTM[Idx].GetTranslation() - Collection1_InitialTM[Idx].GetTranslation()).Size() < KINDA_SMALL_NUMBER);
 		}
@@ -710,8 +708,7 @@ namespace GeometryCollectionTest
 
 		EXPECT_EQ(SovlerParticleHandles.Size(), 4);
 		EXPECT_EQ(ClusterMap.Num(), 0);
-		const TManagedArray<FTransform3f> Transforms3(DynamicCollection1->GetTransforms());
-		TArray<FTransform> Collection1_PostRelease2TM; GeometryCollectionAlgo::GlobalMatrices(Transforms3, *DynamicCollection1, Collection1_PostRelease2TM);
+		TArray<FTransform> Collection1_PostRelease2TM; GeometryCollectionAlgo::Private::GlobalMatrices( *DynamicCollection1, Collection1_PostRelease2TM);
 		for (int Idx = 0; Idx < Collection1_PostRelease2TM.Num(); Idx++) {
 			EXPECT_TRUE((Collection1_PostRelease2TM[Idx].GetTranslation() - Collection1_InitialTM[Idx].GetTranslation()).Size() < KINDA_SMALL_NUMBER);
 		}
@@ -1517,10 +1514,10 @@ namespace GeometryCollectionTest
 			if (Frame == 0)
 			{
 				TArray<FTransform> GlobalTransform;
-				GeometryCollectionAlgo::GlobalMatrices(TManagedArray<FTransform3f>(DynamicCollection->GetTransforms()), *DynamicCollection, GlobalTransform);
+				GeometryCollectionAlgo::Private::GlobalMatrices(*DynamicCollection, GlobalTransform);
 
 				TArray<FTransform> GlobalTransform2;
-				GeometryCollectionAlgo::GlobalMatrices(TManagedArray<FTransform3f>(DynamicCollection2->GetTransforms()), *DynamicCollection2, GlobalTransform2);
+				GeometryCollectionAlgo::Private::GlobalMatrices(*DynamicCollection2, GlobalTransform2);
 
 				// build relative transforms distances
 				for (int32 i = 0; i < (int32)GlobalTransform.Num()-1; i++)
@@ -1539,10 +1536,10 @@ namespace GeometryCollectionTest
 
 		
 		TArray<FTransform> GlobalTransform;
-		GeometryCollectionAlgo::GlobalMatrices(TManagedArray<FTransform3f>(DynamicCollection->GetTransforms()), *DynamicCollection, GlobalTransform);
+		GeometryCollectionAlgo::Private::GlobalMatrices(*DynamicCollection, GlobalTransform);
 
 		TArray<FTransform> GlobalTransform2;
-		GeometryCollectionAlgo::GlobalMatrices(TManagedArray<FTransform3f>(DynamicCollection2->GetTransforms()), *DynamicCollection2, GlobalTransform2);
+		GeometryCollectionAlgo::Private::GlobalMatrices(*DynamicCollection2, GlobalTransform2);
 
 		// build relative transforms distances
 		TArray<FReal> Distances2;
@@ -1615,10 +1612,10 @@ namespace GeometryCollectionTest
 		});
 
 		TArray<FTransform> PrevGlobalTransform;
-		GeometryCollectionAlgo::GlobalMatrices(TManagedArray<FTransform3f>(DynamicCollection->GetTransforms()), *DynamicCollection, PrevGlobalTransform);
+		GeometryCollectionAlgo::Private::GlobalMatrices(*DynamicCollection, PrevGlobalTransform);
 
 		TArray<FTransform> PrevGlobalTransform2;
-		GeometryCollectionAlgo::GlobalMatrices(TManagedArray<FTransform3f>(DynamicCollection2->GetTransforms()), *DynamicCollection2, PrevGlobalTransform2);
+		GeometryCollectionAlgo::Private::GlobalMatrices(*DynamicCollection2, PrevGlobalTransform2);
 
 
 		for (int Frame = 0; Frame < 100; Frame++)
@@ -1626,10 +1623,10 @@ namespace GeometryCollectionTest
 			UnitTest.Advance();
 
 			TArray<FTransform> GlobalTransform;
-			GeometryCollectionAlgo::GlobalMatrices(TManagedArray<FTransform3f>(DynamicCollection->GetTransforms()), *DynamicCollection, GlobalTransform);
+			GeometryCollectionAlgo::Private::GlobalMatrices(*DynamicCollection, GlobalTransform);
 
 			TArray<FTransform> GlobalTransform2;
-			GeometryCollectionAlgo::GlobalMatrices(TManagedArray<FTransform3f>(DynamicCollection2->GetTransforms()), *DynamicCollection2, GlobalTransform2);
+			GeometryCollectionAlgo::Private::GlobalMatrices(*DynamicCollection2, GlobalTransform2);
 
 			EXPECT_EQ(ClusterMap.Num(), 1);
 			EXPECT_TRUE(ClusterMapContains(ClusterMap, ParticleHandles.Handle(6)->CastToRigidParticle(), { ParticleHandles.Handle(1)->CastToRigidParticle(),ParticleHandles.Handle(0)->CastToRigidParticle(),ParticleHandles.Handle(3)->CastToRigidParticle(),ParticleHandles.Handle(4)->CastToRigidParticle() }));
@@ -1756,10 +1753,10 @@ namespace GeometryCollectionTest
 				FTransform RootTransform(Root->R(), Root->X());
 
 				TArray<FTransform> GlobalTransform1;
-				GeometryCollectionAlgo::GlobalMatrices(TManagedArray<FTransform3f>(DynamicCollection->GetTransforms()), *DynamicCollection, GlobalTransform1);
+				GeometryCollectionAlgo::Private::GlobalMatrices(*DynamicCollection, GlobalTransform1);
 
 				TArray<FTransform> GlobalTransform2;
-				GeometryCollectionAlgo::GlobalMatrices(TManagedArray<FTransform3f>(DynamicCollection2->GetTransforms()), *DynamicCollection2, GlobalTransform2);
+				GeometryCollectionAlgo::Private::GlobalMatrices(*DynamicCollection2, GlobalTransform2);
 
 				EXPECT_TRUE(!InitialRootPosition.Equals(Root->X())); // root moves
 
@@ -1832,8 +1829,8 @@ namespace GeometryCollectionTest
 		TSharedPtr<FGeometryDynamicCollection> DynamicCollection1 = Collection1->DynamicCollection;
 		TSharedPtr<FGeometryDynamicCollection> DynamicCollection2 = Collection2->DynamicCollection;
 
-		TArray<FTransform> Collection1_InitialTM; GeometryCollectionAlgo::GlobalMatrices(TManagedArray<FTransform3f>(DynamicCollection1->GetTransforms()), *DynamicCollection1, Collection1_InitialTM);
-		TArray<FTransform> Collection2_InitialTM; GeometryCollectionAlgo::GlobalMatrices(TManagedArray<FTransform3f>(DynamicCollection2->GetTransforms()), *DynamicCollection2, Collection2_InitialTM);
+		TArray<FTransform> Collection1_InitialTM; GeometryCollectionAlgo::Private::GlobalMatrices(*DynamicCollection1, Collection1_InitialTM);
+		TArray<FTransform> Collection2_InitialTM; GeometryCollectionAlgo::Private::GlobalMatrices(*DynamicCollection2, Collection2_InitialTM);
 
 		TArray<Chaos::TPBDRigidClusteredParticleHandle<FReal, 3>*>& Collection1Handles = Collection1->PhysObject->GetSolverParticleHandles();
 		TArray<Chaos::TPBDRigidClusteredParticleHandle<FReal, 3>*>& Collection2Handles = Collection2->PhysObject->GetSolverParticleHandles();
@@ -1867,8 +1864,8 @@ namespace GeometryCollectionTest
 		EXPECT_TRUE(ClusterMapContains(ClusterMap, ClusterUnion.InternalCluster, { Collection1Handles[0], Collection1Handles[1], Collection2Handles[0], Collection2Handles[1] }));
 		EXPECT_EQ(ClusterMap[ClusterUnion.InternalCluster].Num(), 4);
 
-		TArray<FTransform> Collection1_PreReleaseTM; GeometryCollectionAlgo::GlobalMatrices(TManagedArray<FTransform3f>(DynamicCollection1->GetTransforms()), *DynamicCollection1, Collection1_PreReleaseTM);
-		TArray<FTransform> Collection2_PreReleaseTM; GeometryCollectionAlgo::GlobalMatrices(TManagedArray<FTransform3f>(DynamicCollection2->GetTransforms()), *DynamicCollection2, Collection2_PreReleaseTM);
+		TArray<FTransform> Collection1_PreReleaseTM; GeometryCollectionAlgo::Private::GlobalMatrices(*DynamicCollection1, Collection1_PreReleaseTM);
+		TArray<FTransform> Collection2_PreReleaseTM; GeometryCollectionAlgo::Private::GlobalMatrices(*DynamicCollection2, Collection2_PreReleaseTM);
 		for (int Idx = 0; Idx < Collection1_PreReleaseTM.Num() - 1; Idx++) {
 			EXPECT_LT(Collection1_PreReleaseTM[Idx].GetTranslation().Z, Collection1_InitialTM[Idx].GetTranslation().Z);
 			EXPECT_LT(Collection2_PreReleaseTM[Idx].GetTranslation().Z, Collection2_InitialTM[Idx].GetTranslation().Z);
@@ -1892,8 +1889,8 @@ namespace GeometryCollectionTest
 
 		EXPECT_EQ(SovlerParticleHandles.Size(), 7);
 
-		TArray<FTransform> Collection1_PostReleaseTM; GeometryCollectionAlgo::GlobalMatrices(TManagedArray<FTransform3f>(DynamicCollection1->GetTransforms()), *DynamicCollection1, Collection1_PostReleaseTM);
-		TArray<FTransform> Collection2_PostReleaseTM; GeometryCollectionAlgo::GlobalMatrices(TManagedArray<FTransform3f>(DynamicCollection2->GetTransforms()), *DynamicCollection2, Collection2_PostReleaseTM);
+		TArray<FTransform> Collection1_PostReleaseTM; GeometryCollectionAlgo::Private::GlobalMatrices(*DynamicCollection1, Collection1_PostReleaseTM);
+		TArray<FTransform> Collection2_PostReleaseTM; GeometryCollectionAlgo::Private::GlobalMatrices(*DynamicCollection2, Collection2_PostReleaseTM);
 		for (int Idx = 0; Idx < Collection1_PostReleaseTM.Num() - 1; Idx++) {
 			EXPECT_LT(Collection1_PostReleaseTM[Idx].GetTranslation().Z, Collection1_PreReleaseTM[Idx].GetTranslation().Z);
 			EXPECT_LT(Collection2_PostReleaseTM[Idx].GetTranslation().Z, Collection2_PreReleaseTM[Idx].GetTranslation().Z);
@@ -1948,8 +1945,8 @@ namespace GeometryCollectionTest
 		Chaos::FClusterUnionManager& ClusterUnionManager = Clustering.GetClusterUnionManager();
 		const auto& ClusterMap = Clustering.GetChildrenMap();
 
-		TArray<FTransform> Collection1_InitialTM; GeometryCollectionAlgo::GlobalMatrices(TManagedArray<FTransform3f>(DynamicCollection1->GetTransforms()), *DynamicCollection1, Collection1_InitialTM);
-		TArray<FTransform> Collection2_InitialTM; GeometryCollectionAlgo::GlobalMatrices(TManagedArray<FTransform3f>(DynamicCollection2->GetTransforms()), *DynamicCollection2, Collection2_InitialTM);
+		TArray<FTransform> Collection1_InitialTM; GeometryCollectionAlgo::Private::GlobalMatrices(*DynamicCollection1, Collection1_InitialTM);
+		TArray<FTransform> Collection2_InitialTM; GeometryCollectionAlgo::Private::GlobalMatrices(*DynamicCollection2, Collection2_InitialTM);
 
 		TArray<Chaos::TPBDRigidClusteredParticleHandle<FReal, 3>*>& Collection1Handles = Collection1->PhysObject->GetSolverParticleHandles();
 		TArray<Chaos::TPBDRigidClusteredParticleHandle<FReal, 3>*>& Collection2Handles = Collection2->PhysObject->GetSolverParticleHandles();
@@ -1982,8 +1979,8 @@ namespace GeometryCollectionTest
 		EXPECT_TRUE(ClusterMapContains(ClusterMap, ClusterUnion.InternalCluster, { Collection1Handles[0], Collection1Handles[1], Collection2Handles[0], Collection2Handles[1] }));
 		EXPECT_EQ(ClusterMap[ClusterUnion.InternalCluster].Num(), 4);
 
-		TArray<FTransform> Collection1_PreReleaseTM; GeometryCollectionAlgo::GlobalMatrices(TManagedArray<FTransform3f>(DynamicCollection1->GetTransforms()), *DynamicCollection1, Collection1_PreReleaseTM);
-		TArray<FTransform> Collection2_PreReleaseTM; GeometryCollectionAlgo::GlobalMatrices(TManagedArray<FTransform3f>(DynamicCollection2->GetTransforms()), *DynamicCollection2, Collection2_PreReleaseTM);
+		TArray<FTransform> Collection1_PreReleaseTM; GeometryCollectionAlgo::Private::GlobalMatrices(*DynamicCollection1, Collection1_PreReleaseTM);
+		TArray<FTransform> Collection2_PreReleaseTM; GeometryCollectionAlgo::Private::GlobalMatrices(*DynamicCollection2, Collection2_PreReleaseTM);
 		for (int Idx = 0; Idx < Collection1_PreReleaseTM.Num() - 1; Idx++) 
 		{
 			EXPECT_EQ(Collection1_PreReleaseTM[Idx].GetTranslation().Z, Collection1_InitialTM[Idx].GetTranslation().Z);
@@ -2010,8 +2007,8 @@ namespace GeometryCollectionTest
 
 		// validate that DynamicCollection2 became dynamic and fell from the cluster. 
 
-		TArray<FTransform> Collection1_PostReleaseTM; GeometryCollectionAlgo::GlobalMatrices(TManagedArray<FTransform3f>(DynamicCollection1->GetTransforms()), *DynamicCollection1, Collection1_PostReleaseTM);
-		TArray<FTransform> Collection2_PostReleaseTM; GeometryCollectionAlgo::GlobalMatrices(TManagedArray<FTransform3f>(DynamicCollection2->GetTransforms()), *DynamicCollection2, Collection2_PostReleaseTM);
+		TArray<FTransform> Collection1_PostReleaseTM; GeometryCollectionAlgo::Private::GlobalMatrices(*DynamicCollection1, Collection1_PostReleaseTM);
+		TArray<FTransform> Collection2_PostReleaseTM; GeometryCollectionAlgo::Private::GlobalMatrices(*DynamicCollection2, Collection2_PostReleaseTM);
 		for (int Idx = 0; Idx < Collection1_PostReleaseTM.Num() - 1; Idx++) 
 		{
 			if(Idx == 1)
@@ -2059,7 +2056,7 @@ namespace GeometryCollectionTest
 
 		auto& Clustering = UnitTest.Solver->GetEvolution()->GetRigidClustering();
 		const auto& ClusterMap = Clustering.GetChildrenMap();
-		TArray<FTransform> Collection1_InitialTM; GeometryCollectionAlgo::GlobalMatrices(TManagedArray<FTransform3f>(DynamicCollection1->GetTransforms()), *DynamicCollection1, Collection1_InitialTM);
+		TArray<FTransform> Collection1_InitialTM; GeometryCollectionAlgo::Private::GlobalMatrices(*DynamicCollection1, Collection1_InitialTM);
 		TArray<Chaos::TPBDRigidClusteredParticleHandle<FReal, 3>*>& Collection1Handles = Collection1->PhysObject->GetSolverParticleHandles();
 		const auto& SovlerParticleHandles = UnitTest.Solver->GetParticles().GetParticleHandles();
 
@@ -2078,7 +2075,7 @@ namespace GeometryCollectionTest
 		EXPECT_TRUE(ClusterMapContains(ClusterMap, Collection1Handles[2], { Collection1Handles[1],Collection1Handles[0] }));
 		EXPECT_TRUE(ClusterMapContains(ClusterMap, Collection1Handles[3], { Collection1Handles[2] }));
 
-		TArray<FTransform> Collection1_PreReleaseTM; GeometryCollectionAlgo::GlobalMatrices(TManagedArray<FTransform3f>(DynamicCollection1->GetTransforms()), *DynamicCollection1, Collection1_PreReleaseTM);
+		TArray<FTransform> Collection1_PreReleaseTM; GeometryCollectionAlgo::Private::GlobalMatrices(*DynamicCollection1, Collection1_PreReleaseTM);
 		for (int Idx = 0; Idx < Collection1_PreReleaseTM.Num() - 1; Idx++)
 		{
 			EXPECT_EQ(Collection1_PreReleaseTM[Idx].GetTranslation().Z, Collection1_InitialTM[Idx].GetTranslation().Z);
@@ -2095,7 +2092,7 @@ namespace GeometryCollectionTest
 		EXPECT_EQ(ClusterMap.Num(), 1);
 
 		// validate that DynamicCollection1 BODY 2 became dynamic and fell from the cluster. 
-		TArray<FTransform> Collection1_PostReleaseTM; GeometryCollectionAlgo::GlobalMatrices(TManagedArray<FTransform3f>(DynamicCollection1->GetTransforms()), *DynamicCollection1, Collection1_PostReleaseTM);
+		TArray<FTransform> Collection1_PostReleaseTM; GeometryCollectionAlgo::Private::GlobalMatrices(*DynamicCollection1, Collection1_PostReleaseTM);
 		EXPECT_NEAR(Collection1_PostReleaseTM[1].GetTranslation().Z, Collection1_PreReleaseTM[1].GetTranslation().Z, KINDA_SMALL_NUMBER); // the original kinematic should be frozen
 		EXPECT_LT(Collection1_PostReleaseTM[0].GetTranslation().Z, Collection1_PreReleaseTM[0].GetTranslation().Z);
 	}
diff --git a/Engine/Source/Programs/HeadlessChaos/Private/GeometryCollection/GeometryCollectionTestInitilization.cpp b/Engine/Source/Programs/HeadlessChaos/Private/GeometryCollection/GeometryCollectionTestInitilization.cpp
--- a/Engine/Source/Programs/HeadlessChaos/Private/GeometryCollection/GeometryCollectionTestInitilization.cpp
+++ b/Engine/Source/Programs/HeadlessChaos/Private/GeometryCollection/GeometryCollectionTestInitilization.cpp
@@ -33,7 +33,7 @@ GTEST_TEST(AllTraits,GeometryCollection_Initilization_TransformedGeometryCollect
 		EXPECT_LT(X.Z,GlobalTranslation[2]);
 
 		TArray<FTransform> Transform;
-		GeometryCollectionAlgo::GlobalMatrices(TManagedArray<FTransform3f>(Collection->DynamicCollection->GetTransforms()),*Collection->DynamicCollection,Transform);
+		GeometryCollectionAlgo::Private::GlobalMatrices(*Collection->DynamicCollection,Transform);
 		EXPECT_EQ(Collection->DynamicCollection->GetParent(0),FGeometryCollection::Invalid); // is not a child
 		EXPECT_NEAR(Transform[0].GetTranslation().X - GlobalTranslation[0],0.0f,KINDA_SMALL_NUMBER);
 		EXPECT_NEAR(Transform[0].GetTranslation().Y - GlobalTranslation[1],0.0f,KINDA_SMALL_NUMBER);
@@ -64,7 +64,7 @@ GTEST_TEST(AllTraits,GeometryCollection_Initilization_TransformedGeometryCollect
 		EXPECT_LT(X.Z,GlobalTranslation[2]);
 
 		TArray<FTransform> Transform;
-		GeometryCollectionAlgo::GlobalMatrices(TManagedArray<FTransform3f>(Collection->DynamicCollection->GetTransforms()),*Collection->DynamicCollection,Transform);
+		GeometryCollectionAlgo::Private::GlobalMatrices(*Collection->DynamicCollection,Transform);
 		EXPECT_EQ(Collection->DynamicCollection->GetParent(0),FGeometryCollection::Invalid); // is not a child
 		EXPECT_NEAR(Transform[0].GetTranslation().X - GlobalTranslation[0],0.0f,KINDA_SMALL_NUMBER);
 		EXPECT_NEAR(Transform[0].GetTranslation().Y - GlobalTranslation[1],0.0f,KINDA_SMALL_NUMBER);
@@ -95,7 +95,7 @@ GTEST_TEST(AllTraits,GeometryCollection_Initilization_TransformedGeometryCollect
 		EXPECT_LT(X.Z,GlobalTranslation[2]);
 
 		TArray<FTransform> Transform;
-		GeometryCollectionAlgo::GlobalMatrices(TManagedArray<FTransform3f>(Collection->DynamicCollection->GetTransforms()),*Collection->DynamicCollection,Transform);
+		GeometryCollectionAlgo::Private::GlobalMatrices(*Collection->DynamicCollection,Transform);
 		EXPECT_EQ(Collection->DynamicCollection->GetParent(0),FGeometryCollection::Invalid); // is not a child
 		EXPECT_NEAR(Transform[0].GetTranslation().X - GlobalTranslation[0],0.0f,KINDA_SMALL_NUMBER);
 		EXPECT_NEAR(Transform[0].GetTranslation().Y - GlobalTranslation[1],0.0f,KINDA_SMALL_NUMBER);
@@ -146,7 +146,7 @@ GTEST_TEST(AllTraits,GeometryCollection_Initilization_TransformedGeometryCollect
 		FQuat ParticleR = UnitTest.Solver->GetParticles().GetGeometryCollectionParticles().R(0);
 
 		TArray<FTransform> Transform;
-		GeometryCollectionAlgo::GlobalMatrices(TManagedArray<FTransform3f>(Collection->DynamicCollection->GetTransforms()), *Collection->DynamicCollection, Transform);
+		GeometryCollectionAlgo::Private::GlobalMatrices(*Collection->DynamicCollection, Transform);
 		EXPECT_EQ(Collection->DynamicCollection->GetParent(0), FGeometryCollection::Invalid); // is not a child
 		EXPECT_NEAR(ParticlePos.X, GlobalTranslation[0], KINDA_SMALL_NUMBER); // Check particle position from solver 
 		EXPECT_NEAR(ParticlePos.Y, GlobalTranslation[1], KINDA_SMALL_NUMBER);
diff --git a/Engine/Source/Runtime/Experimental/Chaos/Private/GeometryCollection/GeometryCollectionAlgo.cpp b/Engine/Source/Runtime/Experimental/Chaos/Private/GeometryCollection/GeometryCollectionAlgo.cpp
--- a/Engine/Source/Runtime/Experimental/Chaos/Private/GeometryCollection/GeometryCollectionAlgo.cpp
+++ b/Engine/Source/Runtime/Experimental/Chaos/Private/GeometryCollection/GeometryCollectionAlgo.cpp
@@ -347,8 +347,7 @@ namespace GeometryCollectionAlgo
 		return true;
 	}
 
-	template<typename TransformType>
-	void GlobalMatricesHelper(const int32 Index, const FGeometryDynamicCollection& DynamicCollection, const TManagedArray<TransformType>& Transform, TArray<bool>& IsTransformComputed, const TManagedArray<FTransform>* UniformScale, TArray<FTransform>& OutGlobalTransforms)
+	void GlobalMatricesHelper(const int32 Index, const FGeometryDynamicCollection& DynamicCollection, TArray<bool>& IsTransformComputed, const TManagedArray<FTransform>* UniformScale, TArray<FTransform>& OutGlobalTransforms)
 	{
 		if (IsTransformComputed[Index])
 		{
@@ -356,7 +355,7 @@ namespace GeometryCollectionAlgo
 		}
 
 		FIndicesNeedMatricesArray ToProcess;
-		if (!GlobalMatricesGetIndicesToProcessHelper(Index, DynamicCollection, Transform.Num(), IsTransformComputed, ToProcess))
+		if (!GlobalMatricesGetIndicesToProcessHelper(Index, DynamicCollection, DynamicCollection.GetTransforms().Num(), IsTransformComputed, ToProcess))
 		{
 			return;
 		}
@@ -365,7 +364,7 @@ namespace GeometryCollectionAlgo
 		{
 			const int32 ProcessIndex = ToProcess.Pop(false);
 			const int32 ParentIndex = DynamicCollection.GetParent(ProcessIndex);
-			FTransform Result = FTransform(Transform[ProcessIndex]);
+			FTransform Result = FTransform(DynamicCollection.GetTransform(ProcessIndex));
 			if (ParentIndex != FGeometryCollection::Invalid)
 			{
 				Result *= OutGlobalTransforms[ParentIndex];
@@ -384,9 +383,7 @@ namespace GeometryCollectionAlgo
 		}
 	}
 
-
-	template<typename TransformType>
-	void GlobalMatricesHelper(const int32 Index, const TManagedArray<int32>& Parents, const TManagedArray<TransformType>& Transform, TArray<bool>& IsTransformComputed, const TManagedArray<FTransform>* UniformScale, TArray<FMatrix>& OutGlobalTransforms)
+	void GlobalMatricesHelper(const int32 Index, const TManagedArray<int32>& Parents, const TManagedArray<FTransform>& Transform, TArray<bool>& IsTransformComputed, const TManagedArray<FTransform>* UniformScale, TArray<FMatrix>& OutGlobalTransforms)
 	{
 		if (IsTransformComputed[Index])
 		{
@@ -423,8 +420,7 @@ namespace GeometryCollectionAlgo
 	}
 
 	// #note: this version outputs FTransforms to support functionality for getting global matrices for an array of indices.
-	template<typename TransformType>
-	void GlobalMatricesHelper(const int32 Index, const TManagedArray<int32>& Parents, const TManagedArray<TransformType>& Transform, TArray<bool>& IsTransformComputed, const TManagedArray<FTransform>* UniformScale, TArray<FTransform>& OutGlobalTransforms)
+	void GlobalMatricesHelper(const int32 Index, const TManagedArray<int32>& Parents, const TManagedArray<FTransform>& Transform, TArray<bool>& IsTransformComputed, const TManagedArray<FTransform>* UniformScale, TArray<FTransform>& OutGlobalTransforms)
 	{
 		if (IsTransformComputed[Index])
 		{
@@ -460,7 +456,7 @@ namespace GeometryCollectionAlgo
 		}
 	}
 
-	FTransform GlobalMatricesHelperForIndices(const int32 Index, const FGeometryDynamicCollection& DynamicCollection, const TManagedArray<FTransform3f>& Transform, TArray<bool>& IsTransformComputed, const TManagedArray<FTransform>* UniformScale, TArray<FTransform>& TransformCache)
+	FTransform GlobalMatricesHelperForIndices(const int32 Index, const FGeometryDynamicCollection& DynamicCollection, TArray<bool>& IsTransformComputed, const TManagedArray<FTransform>* UniformScale, TArray<FTransform>& TransformCache)
 	{
 		if (IsTransformComputed[Index])
 		{
@@ -468,7 +464,7 @@ namespace GeometryCollectionAlgo
 		}
 
 		FIndicesNeedMatricesArray ToProcess;
-		if (!GlobalMatricesGetIndicesToProcessHelper(Index, DynamicCollection, Transform.Num(), IsTransformComputed, ToProcess))
+		if (!GlobalMatricesGetIndicesToProcessHelper(Index, DynamicCollection, DynamicCollection.GetTransforms().Num(), IsTransformComputed, ToProcess))
 		{
 			return TransformCache[Index];
 		}
@@ -477,7 +473,7 @@ namespace GeometryCollectionAlgo
 		{
 			const int32 ProcessIndex = ToProcess.Pop(false);
 			const int32 ParentIndex = DynamicCollection.GetParent(ProcessIndex);
-			FTransform Result = FTransform(Transform[ProcessIndex]);
+			FTransform Result = FTransform(DynamicCollection.GetTransform(ProcessIndex));
 			if (ParentIndex != FGeometryCollection::Invalid)
 			{
 				Result *= TransformCache[ParentIndex];
@@ -498,8 +494,7 @@ namespace GeometryCollectionAlgo
 		return TransformCache[Index];
 	}
 
-	template<typename TransformType>
-	FTransform GlobalMatricesHelperForIndices(const int32 Index, const TManagedArray<int32>& Parents, const TManagedArray<TransformType>& Transform, TArray<bool>& IsTransformComputed, const TManagedArray<FTransform>* UniformScale, TArray<FTransform>& TransformCache)
+	FTransform GlobalMatricesHelperForIndices(const int32 Index, const TManagedArray<int32>& Parents, const TManagedArray<FTransform>& Transform, TArray<bool>& IsTransformComputed, const TManagedArray<FTransform>* UniformScale, TArray<FTransform>& TransformCache)
 	{
 		if (IsTransformComputed[Index])
 		{
@@ -537,24 +532,25 @@ namespace GeometryCollectionAlgo
 		return TransformCache[Index];
 	}
 
-
-	FTransform GlobalMatrix(const TManagedArray<FTransform3f>& RelativeTransforms, const FGeometryDynamicCollection& DynamicCollection, int32 Index)
+	namespace Private 
 	{
-		FTransform Transform = FTransform::Identity;
-
-		if (RelativeTransforms.IsValidIndex(Index))
+		FTransform GlobalMatrix(const FGeometryDynamicCollection& DynamicCollection, int32 Index)
 		{
-			while (Index != FGeometryCollection::Invalid)
+			FTransform Transform = FTransform::Identity;
+
+			if (DynamicCollection.GetTransforms().IsValidIndex(Index))
 			{
-				Transform = Transform * FTransform(RelativeTransforms[Index]);
-				Index = DynamicCollection.GetParent(Index);
+				while (Index != FGeometryCollection::Invalid)
+				{
+					Transform = Transform * FTransform(DynamicCollection.GetTransform(Index));
+					Index = DynamicCollection.GetParent(Index);
+				}
 			}
+			return Transform;
 		}
-		return Transform;
 	}
 
-	template<class TransformType>
-	FTransform GlobalMatrixTemplate(const TManagedArray<TransformType>& RelativeTransforms, const TManagedArray<int32>& Parents, int32 Index)
+	FTransform GlobalMatrix(const TManagedArray<FTransform>& RelativeTransforms, const TManagedArray<int32>& Parents, int32 Index)
 	{
 		FTransform Transform = FTransform::Identity;
 
@@ -569,28 +565,22 @@ namespace GeometryCollectionAlgo
 		return Transform;
 	}
 
-	FTransform GlobalMatrix(const TManagedArray<FTransform>& RelativeTransforms, const TManagedArray<int32>& Parents, int32 Index)
-	{
-		return GlobalMatrixTemplate(RelativeTransforms, Parents, Index);
-	}
-
-	FTransform GlobalMatrix(const TManagedArray<FTransform3f>& RelativeTransforms, const TManagedArray<int32>& Parents, int32 Index)
-	{
-		return GlobalMatrixTemplate(RelativeTransforms, Parents, Index);
-	}
-
-	void GlobalMatrices(const TManagedArray<FTransform3f>& RelativeTransforms, const FGeometryDynamicCollection& DynamicCollection, const TArray<int32>& Indices, TArray<FTransform>& OutGlobalTransforms)
+	namespace Private
 	{
-		TArray<bool> IsTransformComputed;
-		IsTransformComputed.AddDefaulted(RelativeTransforms.Num());
+		void GlobalMatrices(const FGeometryDynamicCollection& DynamicCollection, const TArray<int32>& Indices, TArray<FTransform>& OutGlobalTransforms)
+		{
+			TArray<bool> IsTransformComputed;
+			const int32 NumTransform = DynamicCollection.GetTransforms().Num();
+			IsTransformComputed.AddDefaulted(NumTransform);
 
-		TArray<FTransform> TransformCache;
-		TransformCache.SetNumUninitialized(RelativeTransforms.Num(), false);
+			TArray<FTransform> TransformCache;
+			TransformCache.SetNumUninitialized(NumTransform, false);
 
-		OutGlobalTransforms.SetNumUninitialized(Indices.Num(), false);
-		for (int Idx = 0; Idx < Indices.Num(); Idx++)
-		{
-			OutGlobalTransforms[Idx] = GlobalMatricesHelperForIndices(Indices[Idx], DynamicCollection, RelativeTransforms, IsTransformComputed, nullptr, TransformCache);
+			OutGlobalTransforms.SetNumUninitialized(Indices.Num(), false);
+			for (int Idx = 0; Idx < Indices.Num(); Idx++)
+			{
+				OutGlobalTransforms[Idx] = GlobalMatricesHelperForIndices(Indices[Idx], DynamicCollection, IsTransformComputed, nullptr, TransformCache);
+			}
 		}
 	}
 
@@ -615,12 +605,6 @@ namespace GeometryCollectionAlgo
 		GlobalMatricesTemplate<FTransform>(RelativeTransforms, Parents, Indices, Transforms);
 	}
 
-	void GlobalMatrices(const TManagedArray<FTransform3f>& RelativeTransforms, const TManagedArray<int32>& Parents, const TArray<int32>& Indices, TArray<FTransform>& Transforms)
-	{
-		GlobalMatricesTemplate<FTransform3f>(RelativeTransforms, Parents, Indices, Transforms);
-	}
-
-	
 	void GlobalMatricesFromRoot(const int32 ParentTransformIndex, const TManagedArray<FTransform>& RelativeTransforms, const TManagedArray<TSet<int32>>& Children, TArray<FMatrix>& Transforms)
 	{
 		if (Children[ParentTransformIndex].Num() > 0)
@@ -634,8 +618,8 @@ namespace GeometryCollectionAlgo
 		}
 	}
 
-	template<typename MatrixType, typename TransformType>
-	void GlobalMatrices(const TManagedArray<TransformType>& RelativeTransforms, const TManagedArray<int32>& Parents, const TManagedArray<FTransform>& UniformScale, TArray<MatrixType>& OutGlobalTransforms)
+	template<typename MatrixType>
+	void GlobalMatrices(const TManagedArray<FTransform>& RelativeTransforms, const TManagedArray<int32>& Parents, const TManagedArray<FTransform>& UniformScale, TArray<MatrixType>& OutGlobalTransforms)
 	{
 		int32 NumTransforms = RelativeTransforms.Num();
 
@@ -646,27 +630,30 @@ namespace GeometryCollectionAlgo
 
 		for (int BoneIdx = 0; BoneIdx < NumTransforms; ++BoneIdx)
 		{
-			GlobalMatricesHelper<TransformType>(BoneIdx, Parents, RelativeTransforms, IsTransformComputed, &UniformScale, OutGlobalTransforms);
+			GlobalMatricesHelper(BoneIdx, Parents, RelativeTransforms, IsTransformComputed, &UniformScale, OutGlobalTransforms);
 		}
 	}
 
-	void GlobalMatrices(const TManagedArray<FTransform3f>& RelativeTransforms, const FGeometryDynamicCollection& DynamicCollection, TArray<FTransform>& OutGlobalTransforms)
+	namespace Private
 	{
-		int32 NumTransforms = RelativeTransforms.Num();
+		void GlobalMatrices(const FGeometryDynamicCollection& DynamicCollection, TArray<FTransform>& OutGlobalTransforms)
+		{
+			int32 NumTransforms = DynamicCollection.GetTransforms().Num();
 
-		TArray<bool> IsTransformComputed;
-		IsTransformComputed.AddDefaulted(NumTransforms);
+			TArray<bool> IsTransformComputed;
+			IsTransformComputed.AddDefaulted(NumTransforms);
 
-		OutGlobalTransforms.SetNumUninitialized(NumTransforms, false);
+			OutGlobalTransforms.SetNumUninitialized(NumTransforms, false);
 
-		for (int BoneIdx = 0; BoneIdx < NumTransforms; ++BoneIdx)
-		{
-			GlobalMatricesHelper<FTransform3f>(BoneIdx, DynamicCollection, RelativeTransforms, IsTransformComputed, nullptr, OutGlobalTransforms);
+			for (int BoneIdx = 0; BoneIdx < NumTransforms; ++BoneIdx)
+			{
+				GlobalMatricesHelper(BoneIdx, DynamicCollection, IsTransformComputed, nullptr, OutGlobalTransforms);
+			}
 		}
 	}
 
-	template<typename MatrixType, typename TransformType>
-	void GlobalMatrices(const TManagedArray<TransformType>& RelativeTransforms, const TManagedArray<int32>& Parents, TArray<MatrixType>& OutGlobalTransforms)
+	template<typename MatrixType>
+	void GlobalMatrices(const TManagedArray<FTransform>& RelativeTransforms, const TManagedArray<int32>& Parents, TArray<MatrixType>& OutGlobalTransforms)
 	{
 		int32 NumTransforms = RelativeTransforms.Num();
 
@@ -677,20 +664,14 @@ namespace GeometryCollectionAlgo
 
 		for (int BoneIdx = 0; BoneIdx < NumTransforms; ++BoneIdx)
 		{
-			GlobalMatricesHelper<TransformType>(BoneIdx, Parents, RelativeTransforms, IsTransformComputed, nullptr, OutGlobalTransforms);
+			GlobalMatricesHelper(BoneIdx, Parents, RelativeTransforms, IsTransformComputed, nullptr, OutGlobalTransforms);
 		}
 	}
 
-	template void CHAOS_API GlobalMatrices<FTransform, FTransform>(const TManagedArray<FTransform>&, const TManagedArray<int32>&, const TManagedArray<FTransform>&, TArray<FTransform>&);
-	template void CHAOS_API GlobalMatrices<FTransform, FTransform>(const TManagedArray<FTransform>&, const TManagedArray<int32>&, TArray<FTransform>&);
-	template void CHAOS_API GlobalMatrices<FMatrix, FTransform>(const TManagedArray<FTransform>&, const TManagedArray<int32>&, const TManagedArray<FTransform>&, TArray<FMatrix>&);
-	template void CHAOS_API GlobalMatrices<FMatrix, FTransform>(const TManagedArray<FTransform>&, const TManagedArray<int32>&, TArray<FMatrix>&);
-
-	template void CHAOS_API GlobalMatrices<FTransform, FTransform3f>(const TManagedArray<FTransform3f>&, const TManagedArray<int32>&, const TManagedArray<FTransform>&, TArray<FTransform>&);
-	template void CHAOS_API GlobalMatrices<FTransform, FTransform3f>(const TManagedArray<FTransform3f>&, const TManagedArray<int32>&, TArray<FTransform>&);
-	template void CHAOS_API GlobalMatrices<FMatrix, FTransform3f>(const TManagedArray<FTransform3f>&, const TManagedArray<int32>&, const TManagedArray<FTransform>&, TArray<FMatrix>&);
-	template void CHAOS_API GlobalMatrices<FMatrix, FTransform3f>(const TManagedArray<FTransform3f>&, const TManagedArray<int32>&, TArray<FMatrix>&);
-
+	template void CHAOS_API GlobalMatrices<FTransform>(const TManagedArray<FTransform>&, const TManagedArray<int32>&, const TManagedArray<FTransform>&, TArray<FTransform>&);
+	template void CHAOS_API GlobalMatrices<FTransform>(const TManagedArray<FTransform>&, const TManagedArray<int32>&, TArray<FTransform>&);
+	template void CHAOS_API GlobalMatrices<FMatrix>(const TManagedArray<FTransform>&, const TManagedArray<int32>&, const TManagedArray<FTransform>&, TArray<FMatrix>&);
+	template void CHAOS_API GlobalMatrices<FMatrix>(const TManagedArray<FTransform>&, const TManagedArray<int32>&, TArray<FMatrix>&);
 
 	void FloodForOverlappedPairs(int Level, int32 BoneIndex, TMap<int32, int32> &BoneToGroup, const TManagedArray<int32>& Levels, const TMap<int32, FBox>& BoundingBoxes, TSet<TTuple<int32, int32>>& OutOverlappedPairs)
 	{
diff --git a/Engine/Source/Runtime/Experimental/Chaos/Private/PhysicsProxy/GeometryCollectionPhysicsProxy.cpp b/Engine/Source/Runtime/Experimental/Chaos/Private/PhysicsProxy/GeometryCollectionPhysicsProxy.cpp
--- a/Engine/Source/Runtime/Experimental/Chaos/Private/PhysicsProxy/GeometryCollectionPhysicsProxy.cpp
+++ b/Engine/Source/Runtime/Experimental/Chaos/Private/PhysicsProxy/GeometryCollectionPhysicsProxy.cpp
@@ -549,7 +549,7 @@ void FGeometryCollectionPhysicsProxy::Initialize(Chaos::FPBDRigidsEvolutionBase
 	}
 
 	// Initialise GT/External particles
-	const int32 NumTransforms = DynamicCollection.Transform.Num();
+	const int32 NumTransforms = DynamicCollection.GetTransforms().Num();
 
 	// make sure we copy the anchored information over to the physics thread collection
 	const Chaos::Facades::FCollectionAnchoringFacade DynamicCollectionAnchoringFacade(DynamicCollection);
@@ -581,7 +581,7 @@ void FGeometryCollectionPhysicsProxy::Initialize(Chaos::FPBDRigidsEvolutionBase
 				// constraint will query the game thread particle position/rotation for the geometry collection to compute its
 				// reference frame. If that position/rotation does not match up with the physics thread's position/rotation,
 				// the geometry collection particle will have an added velocity computed by the joint constraint solver.
-				const FTransform& T = MassToLocal[Index] * FTransform(GameThreadCollection.Transform[Index]) * Parameters.WorldTransform;
+				const FTransform& T = MassToLocal[Index] * FTransform(GameThreadCollection.GetTransform(Index)) * Parameters.WorldTransform;
 				P->SetX(T.GetTranslation(), false);
 				P->SetR(T.GetRotation(), false);
 				P->SetM(ScaledMass);
@@ -672,7 +672,7 @@ void FGeometryCollectionPhysicsProxy::Initialize(Chaos::FPBDRigidsEvolutionBase
 								const Chaos::FImplicitObjectPtr& ChildImplicit = Implicits[ChildIndex];
 								if (ChildImplicit)
 								{
-									const Chaos::FRigidTransform3 ChildShapeTransform = MassToLocal[ChildIndex] * FTransform(GameThreadCollection.Transform[ChildIndex]);
+									const Chaos::FRigidTransform3 ChildShapeTransform = MassToLocal[ChildIndex] * FTransform(GameThreadCollection.GetTransform(ChildIndex));
 									const Chaos::FRigidTransform3 RelativeShapeTransform = ChildShapeTransform.GetRelativeTransform(ParentShapeTransform);
 
 									Chaos::FImplicitObjectPtr TransformedChildImplicit = MakeTransformImplicitObject(*ChildImplicit, RelativeShapeTransform);
@@ -1068,7 +1068,7 @@ void FGeometryCollectionPhysicsProxy::InitializeBodiesPT(Chaos::FPBDRigidsSolver
 		Parameters.PrevWorldTransform = Parameters.WorldTransform;
 
 		TArray<FTransform> Transform;
-		GeometryCollectionAlgo::GlobalMatrices(DynamicCollection.Transform, DynamicCollection, Transform);
+		GeometryCollectionAlgo::Private::GlobalMatrices(DynamicCollection, Transform);
 
 		// Here Clean up Additional particles
 		TBitArray<> EffectiveParticles;
@@ -1628,14 +1628,13 @@ Chaos::TPBDGeometryCollectionParticleHandle<Chaos::FReal, 3>* FGeometryCollectio
 {
 	FGeometryDynamicCollection& DynamicCollection = PhysicsThreadCollection;
 	TManagedArray<uint8>& DynamicState = DynamicCollection.DynamicState;
-	TManagedArray<FTransform3f>& Transform = DynamicCollection.Transform;
 	const TManagedArray<FTransform>& MassToLocal = Parameters.RestCollection->GetAttribute<FTransform>(MassToLocalAttributeName, FTransformCollection::TransformGroup);
 	const TManagedArray<Chaos::FImplicitObjectPtr>& Implicits = DynamicCollection.GetAttribute<Chaos::FImplicitObjectPtr>(FGeometryDynamicCollection::ImplicitsAttribute, FTransformCollection::TransformGroup);
 	const TManagedArray<TUniquePtr<FCollisionStructureManager::FSimplicial>>& Simplicials = DynamicCollection.Simplicials;
 	Chaos::Facades::FCollectionAnchoringFacade AnchoringFacade(DynamicCollection);
 
 	//If we are a root particle use the world transform, otherwise set the relative transform
-	const FTransform CollectionSpaceTransform = GeometryCollectionAlgo::GlobalMatrix(Transform, DynamicCollection, CollectionClusterIndex);
+	const FTransform CollectionSpaceTransform = GeometryCollectionAlgo::Private::GlobalMatrix(DynamicCollection, CollectionClusterIndex);
 	const Chaos::TRigidTransform<Chaos::FReal, 3> ParticleTM = MassToLocal[CollectionClusterIndex] * CollectionSpaceTransform * Parameters.WorldTransform;
 
 	// Gather unique indices from GT to pass into PT handle creation
@@ -1713,7 +1712,7 @@ Chaos::TPBDGeometryCollectionParticleHandle<Chaos::FReal, 3>* FGeometryCollectio
 		Shape->SetMaterial(Parameters.PhysicalMaterialHandle);
 	}
 
-	const FTransform ParentTransform = GeometryCollectionAlgo::GlobalMatrix(DynamicCollection.Transform, DynamicCollection, CollectionClusterIndex);
+	const FTransform ParentTransform = GeometryCollectionAlgo::Private::GlobalMatrix(DynamicCollection, CollectionClusterIndex);
 
 	// Populate bounds as we didn't pass a shared implicit to PopulateSimulatedParticle this will have been skipped, now that we have the full cluster we can build it
 	if (Handle->GetGeometry() && Handle->GetGeometry()->HasBoundingBox())
@@ -1746,13 +1745,11 @@ FGeometryCollectionPhysicsProxy::BuildClusters_Internal(
 
 	FGeometryDynamicCollection& DynamicCollection = PhysicsThreadCollection;
 	TManagedArray<uint8>& DynamicState = DynamicCollection.DynamicState;
-	TManagedArray<FTransform3f>& Transform = DynamicCollection.Transform;
 	const TManagedArray<FTransform>& MassToLocal = Parameters.RestCollection->GetAttribute<FTransform>(MassToLocalAttributeName, FTransformCollection::TransformGroup);
-	//TManagedArray<TSharedPtr<FCollisionStructureManager::FSimplicial> >& Simplicials = DynamicCollection.Simplicials;
 	const TManagedArray<Chaos::FImplicitObjectPtr>& Implicits = DynamicCollection.GetAttribute<Chaos::FImplicitObjectPtr>(FGeometryDynamicCollection::ImplicitsAttribute, FTransformCollection::TransformGroup);
 
 	//If we are a root particle use the world transform, otherwise set the relative transform
-	const FTransform CollectionSpaceTransform = GeometryCollectionAlgo::GlobalMatrix(Transform, DynamicCollection, CollectionClusterIndex);
+	const FTransform CollectionSpaceTransform = GeometryCollectionAlgo::Private::GlobalMatrix(DynamicCollection, CollectionClusterIndex);
 	const Chaos::TRigidTransform<Chaos::FReal, 3> ParticleTM = MassToLocal[CollectionClusterIndex] * CollectionSpaceTransform * Parameters.WorldTransform;
 
 	//create new cluster particle
@@ -1847,7 +1844,7 @@ FGeometryCollectionPhysicsProxy::BuildClusters_Internal(
 		Shape->SetMaterial(Parameters.PhysicalMaterialHandle);
 	}
 
-	const FTransform ParentTransform = GeometryCollectionAlgo::GlobalMatrix(DynamicCollection.Transform, DynamicCollection, CollectionClusterIndex);
+	const FTransform ParentTransform = GeometryCollectionAlgo::Private::GlobalMatrix(DynamicCollection, CollectionClusterIndex);
 
 	int32 MinCollisionGroup = INT_MAX;
 	for(int32 Idx=0; Idx < ChildHandles.Num(); Idx++)
@@ -3180,7 +3177,7 @@ void FGeometryCollectionPhysicsProxy::PushToPhysicsState()
 						if (ClusterUnionIndex != INDEX_NONE)
 						{
 							const FTransform ParentWorldTransform{ ParentHandle->R(), ParentHandle->X() };
-							const FTransform NewWorldTransform = MassToLocal[TransformGroupIndex] * FTransform(PhysicsThreadCollection.Transform[TransformGroupIndex]) * Parameters.WorldTransform;
+							const FTransform NewWorldTransform = MassToLocal[TransformGroupIndex] * FTransform(PhysicsThreadCollection.GetTransform(TransformGroupIndex)) * Parameters.WorldTransform;
 							const FTransform RelativeTransform = NewWorldTransform.GetRelativeTransform(ParentWorldTransform);
 
 							DeferredClusterUnionParticleUpdates.Add(Handle);
@@ -3366,7 +3363,7 @@ void FGeometryCollectionPhysicsProxy::BufferPhysicsResults_Internal(Chaos::FPBDR
 
 	UniqueIdxToInternalClusterHandle.Reset();
 	
-	const int32 NumTransforms = PhysicsThreadCollection.Transform.Num();
+	const int32 NumTransforms = PhysicsThreadCollection.GetTransforms().Num();
 	if(NumTransforms > 0)
 	{ 
 		SCOPE_CYCLE_COUNTER(STAT_CalcParticleToWorld);
@@ -3493,7 +3490,7 @@ void FGeometryCollectionPhysicsProxy::BufferPhysicsResults_Internal(Chaos::FPBDR
 			if (bHasChanged)
 			{
 				// default is what we have on the collection
-				FTransform ParentSpaceTransform = FTransform(PhysicsThreadCollection.Transform[TransformGroupIndex]);
+				FTransform ParentSpaceTransform = FTransform(PhysicsThreadCollection.GetTransform(TransformGroupIndex));
 
 				// recompute parent space transform if there's no more parent or if the parent is an internal cluster 
 				if (!ClusterParent || ClusterParent->InternalCluster())
@@ -3513,7 +3510,7 @@ void FGeometryCollectionPhysicsProxy::BufferPhysicsResults_Internal(Chaos::FPBDR
 				Results.SetPositions(EntryIndex, PositionData);
 				Results.SetVelocities(EntryIndex, VelocityData);
 				// todo(chaos) : we shoudl eventually get rid of the transform in the Physics collection
-				PhysicsThreadCollection.Transform[TransformGroupIndex] = FTransform3f(ParentSpaceTransform);
+				PhysicsThreadCollection.SetTransform(TransformGroupIndex, FTransform3f(ParentSpaceTransform));
 				IsObjectDynamic = true;
 			}
 
@@ -3745,8 +3742,9 @@ bool FGeometryCollectionPhysicsProxy::PullNonInterpolatableDataFromSinglePhysics
 			bIsCollectionDirty = true;
 		}
 
-		if (UpdateValue(GameThreadCollection.HasParent[TransformGroupIndex], StateData.HasParent))
+		if (GameThreadCollection.GetHasParent(TransformGroupIndex) != StateData.HasParent)
 		{
+			GameThreadCollection.SetHasParent(TransformGroupIndex, StateData.HasParent);
 			bIsCollectionDirty = true;
 		}
 
@@ -3833,7 +3831,7 @@ bool FGeometryCollectionPhysicsProxy::PullFromPhysicsState(const Chaos::FDirtyGe
 		InterpolationData.AccumlateErrorXR(Error->ErrorX, Error->ErrorR, SolverSyncTimestamp, RenderInterpErrorCorrectionDurationTicks);
 	}
 
-	const int32 NumTransforms = GameThreadCollection.Transform.Num();
+	const int32 NumTransforms = GameThreadCollection.GetTransforms().Num();
 	const bool bNeedInterpolation = (NextPullData != nullptr);
 
 	bool bIsCollectionDirty = false;
@@ -3958,10 +3956,10 @@ bool FGeometryCollectionPhysicsProxy::PullFromPhysicsState(const Chaos::FDirtyGe
 							{
 								NewTransform = ParticleMassToLocal.Inverse() * ComponentScaleTransform * ParticleMassToLocal * NewTransform;
 							}
-							if (!NewTransform.Equals(FTransform(GameThreadCollection.Transform[TransformIndex]), GeometryCollectionTransformTolerance))
+							if (!NewTransform.Equals(FTransform(GameThreadCollection.GetTransform(TransformIndex)), GeometryCollectionTransformTolerance))
 							{
 								bHasDifferentTransforms = true;
-								GameThreadCollection.Transform[TransformIndex] = FTransform3f(NewTransform);
+								GameThreadCollection.SetTransform(TransformIndex, FTransform3f(NewTransform));
 							}
 						}
 					}
@@ -4005,7 +4003,7 @@ void FGeometryCollectionPhysicsProxy::UpdateFilterData_External(const FCollision
 	check(IsInGameThread());
 
 	// SimFilter/QueryFilter members are read on both threads, these are const after initialization and are not updated here.
-	const int32 NumTransforms = GameThreadCollection.Transform.Num();
+	const int32 NumTransforms = GameThreadCollection.GetTransforms().Num();
 	for (int32 Index = 0; Index < NumTransforms; ++Index)
 	{
 		FParticle* P = GTParticles[Index].Get();
@@ -5058,7 +5056,7 @@ void FGeometryCollectionPhysicsProxy::FieldParameterUpdateCallback(Chaos::FPBDRi
 
 	// Process Particle-Collection commands
 	int32 NumCommands = Commands.Num();
-	if (NumCommands && !RigidSolver->IsShuttingDown() && Collection.Transform.Num())
+	if (NumCommands && !RigidSolver->IsShuttingDown() && Collection.GetTransforms().Num())
 	{
 		TArray<int32> CommandsToRemove;
 		CommandsToRemove.Reserve(NumCommands);
diff --git a/Engine/Source/Runtime/Experimental/Chaos/Public/GeometryCollection/GeometryCollectionAlgo.h b/Engine/Source/Runtime/Experimental/Chaos/Public/GeometryCollection/GeometryCollectionAlgo.h
--- a/Engine/Source/Runtime/Experimental/Chaos/Public/GeometryCollection/GeometryCollectionAlgo.h
+++ b/Engine/Source/Runtime/Experimental/Chaos/Public/GeometryCollection/GeometryCollectionAlgo.h
@@ -134,16 +134,12 @@ namespace GeometryCollectionAlgo
 	*  Global Matrices of the specified index.
 	*/
 	FTransform CHAOS_API GlobalMatrix(const TManagedArray<FTransform>& RelativeTransforms, const TManagedArray<int32>& Parents, int32 Index);
-	FTransform CHAOS_API GlobalMatrix(const TManagedArray<FTransform3f>& RelativeTransforms, const TManagedArray<int32>& Parents, int32 Index);
-	FTransform CHAOS_API GlobalMatrix(const TManagedArray<FTransform3f>& RelativeTransforms, const FGeometryDynamicCollection& DynamicCollection, int32 Index);
 
 
 	/*
 	*  Global Matrices of the collection based on list of indices
 	*/
 	void CHAOS_API GlobalMatrices(const TManagedArray<FTransform>& RelativeTransforms, const TManagedArray<int32>& Parents, const TArray<int32>& Indices, TArray<FTransform>& Transforms);
-	void CHAOS_API GlobalMatrices(const TManagedArray<FTransform3f>& RelativeTransforms, const TManagedArray<int32>& Parents, const TArray<int32>& Indices, TArray<FTransform>& Transforms);
-	void CHAOS_API GlobalMatrices(const TManagedArray<FTransform3f>& RelativeTransforms, const FGeometryDynamicCollection& DynamicCollection, const TArray<int32>& Indices, TArray<FTransform>& Transforms);
 
 	/*
 	 *  Recursively traverse from a root node down
@@ -153,14 +149,11 @@ namespace GeometryCollectionAlgo
 	/*
 	*  Global Matrices of the collection, transforms will be resized to fit
 	*/
-	template<typename MatrixType, typename TransformType>
-	void CHAOS_API GlobalMatrices(const TManagedArray<TransformType>& RelativeTransforms, const TManagedArray<int32>& Parents, const TManagedArray<FTransform>& UniformScale, TArray<MatrixType>& Transforms);
-
-	template<typename MatrixType, typename TransformType>
-	void CHAOS_API GlobalMatrices(const TManagedArray<TransformType>& RelativeTransforms, const TManagedArray<int32>& Parents, TArray<MatrixType>& Transforms);
-
-	void CHAOS_API GlobalMatrices(const TManagedArray<FTransform3f>& RelativeTransforms, const FGeometryDynamicCollection& DynamicCollection, TArray<FTransform>& Transforms);
+	template<typename MatrixType>
+	void CHAOS_API GlobalMatrices(const TManagedArray<FTransform>& RelativeTransforms, const TManagedArray<int32>& Parents, const TManagedArray<FTransform>& UniformScale, TArray<MatrixType>& Transforms);
 
+	template<typename MatrixType>
+	void CHAOS_API GlobalMatrices(const TManagedArray<FTransform>& RelativeTransforms, const TManagedArray<int32>& Parents, TArray<MatrixType>& Transforms);
 
 	/*
 	*  Gets pairs of elements whose bounding boxes overlap.
@@ -291,4 +284,13 @@ namespace GeometryCollectionAlgo
 	CHAOS_API
 	ComputeRecursiveOrder(const FManagedArrayCollection& Collection);
 
+
+	// For internal use only 
+	namespace Private
+	{
+		void CHAOS_API GlobalMatrices(const FGeometryDynamicCollection& DynamicCollection, TArray<FTransform>& Transforms);
+		void CHAOS_API GlobalMatrices(const FGeometryDynamicCollection& DynamicCollection, const TArray<int32>& Indices, TArray<FTransform>& OutGlobalTransforms);
+		FTransform CHAOS_API GlobalMatrix(const FGeometryDynamicCollection& DynamicCollection, int32 Index);
+	}
+
 }
diff --git a/Engine/Source/Runtime/Experimental/Chaos/Public/GeometryCollectionProxyData.h b/Engine/Source/Runtime/Experimental/Chaos/Public/GeometryCollectionProxyData.h
--- a/Engine/Source/Runtime/Experimental/Chaos/Public/GeometryCollectionProxyData.h
+++ b/Engine/Source/Runtime/Experimental/Chaos/Public/GeometryCollectionProxyData.h
@@ -65,9 +65,6 @@ private:
 
 	/** Construct */
 	CHAOS_API void Construct();
-
-	friend class FGeometryCollectionPhysicsProxy;
-	friend class UGeometryCollectionComponent;
 };
 
 
diff --git a/Engine/Source/Runtime/Experimental/GeometryCollectionEngine/Private/GeometryCollection/GeometryCollectionComponent.cpp b/Engine/Source/Runtime/Experimental/GeometryCollectionEngine/Private/GeometryCollection/GeometryCollectionComponent.cpp
--- a/Engine/Source/Runtime/Experimental/GeometryCollectionEngine/Private/GeometryCollection/GeometryCollectionComponent.cpp
+++ b/Engine/Source/Runtime/Experimental/GeometryCollectionEngine/Private/GeometryCollection/GeometryCollectionComponent.cpp
@@ -502,28 +502,14 @@ private:
 // Define the methods
 COPY_ON_WRITE_ATTRIBUTES
 
-const TManagedArray<FTransform3f>& UGeometryCollectionComponent::GetTransformArray() const
-{
-	check(IndirectTransformArray != nullptr);
-	return *IndirectTransformArray;
-}
 
-TManagedArray<FTransform3f>& UGeometryCollectionComponent::GetTransformArrayCopyOnWrite()
+void UGeometryCollectionComponent::GetTransformArrayCopyOnWrite()
 {
-	if (!IndirectTransformArray)
-	{
-		DynamicCollection->AddAttribute<FTransform>(FTransformCollection::TransformAttribute, FTransformCollection::TransformGroup);
-		DynamicCollection->CopyAttribute(*RestCollection->GetGeometryCollection(), FTransformCollection::TransformAttribute, FTransformCollection::TransformGroup);
-		IndirectTransformArray = &DynamicCollection->ModifyAttribute<FTransform3f>(FTransformCollection::TransformAttribute, FTransformCollection::TransformGroup);
-		CopyOnWriteAttributeList.Add(reinterpret_cast<FManagedArrayBase**>(&IndirectTransformArray));
-	}
-	return *IndirectTransformArray;
+	DynamicCollection->AddAttribute<FTransform3f>(FTransformCollection::TransformAttribute, FTransformCollection::TransformGroup);
+	DynamicCollection->CopyAttribute(*RestCollection->GetGeometryCollection(), FTransformCollection::TransformAttribute, FTransformCollection::TransformGroup);
+	DynamicCollection->ModifyAttribute<FTransform3f>(FTransformCollection::TransformAttribute, FTransformCollection::TransformGroup);
 }
 
-void UGeometryCollectionComponent::ResetTransformArrayDynamic()
-{
-	IndirectTransformArray = nullptr;
-}
 const TManagedArray<FTransform>& UGeometryCollectionComponent::GetTransformArrayRest() const
 {
 	return RestCollection->GetGeometryCollection()->Transform;
@@ -1340,7 +1326,7 @@ bool UGeometryCollectionComponent::DoCustomNavigableGeometryExport(FNavigableGeo
 			{
 				const int32 RootIndex = GetRootIndex();
 				const FTransform& CompToWorld = GetComponentToWorld(); 
-				const FTransform FinalTransform = (DynamicCollection ? FTransform(DynamicCollection->Transform[RootIndex]) : FTransform::Identity) * CompToWorld;
+				const FTransform FinalTransform = (DynamicCollection ? FTransform(DynamicCollection->GetTransform(RootIndex)) : FTransform::Identity) * CompToWorld;
 				const FVector Scale3D = FinalTransform.GetScale3D();
 				if (!Scale3D.IsZero())
 				{
@@ -1404,7 +1390,7 @@ bool UGeometryCollectionComponent::DoCustomNavigableGeometryExport(FNavigableGeo
 	TArray<FTransform> GeomToComponent;
 	if (DynamicCollection)
 	{
-		GeometryCollectionAlgo::GlobalMatrices(GetTransformArray(), *DynamicCollection, TransformIndexBuffer, GeomToComponent);
+		GeometryCollectionAlgo::Private::GlobalMatrices(*DynamicCollection, TransformIndexBuffer, GeomToComponent);
 	}
 	else
 	{
@@ -2933,11 +2919,10 @@ void UGeometryCollectionComponent::SetInitialTransforms(const TArray<FTransform>
 {
 	if (DynamicCollection)
 	{
-		TManagedArray<FTransform3f>& Transform = DynamicCollection->Transform;
-		int32 MaxIdx = FMath::Min(Transform.Num(), InitialTransforms.Num());
+		int32 MaxIdx = FMath::Min(DynamicCollection->GetTransforms().Num(), InitialTransforms.Num());
 		for (int32 Idx = 0; Idx < MaxIdx; ++Idx)
 		{
-			Transform[Idx] = FTransform3f(InitialTransforms[Idx]);
+			DynamicCollection->SetTransform(Idx, FTransform3f(InitialTransforms[Idx]));
 		}
 	}
 }
@@ -3586,7 +3571,7 @@ void UGeometryCollectionComponent::OnCreatePhysicsState()
 		}
 #endif
 		const bool bValidWorld = GetWorld() && (GetWorld()->IsGameWorld() || GetWorld()->IsPreviewWorld() || GeometryCollectionCreatePhysicsStateInEditor);
-		const bool bValidCollection = DynamicCollection && DynamicCollection->Transform.Num() > 0;
+		const bool bValidCollection = DynamicCollection && DynamicCollection->GetTransforms().Num() > 0;
 		if (bValidWorld && bValidCollection)
 		{
 			FChaosUserData::Set<UPrimitiveComponent>(&PhysicsUserData, this);
@@ -3908,7 +3893,7 @@ void UGeometryCollectionComponent::MoveComponentToRootTransform()
 			if (bIsRootActive || bHasDynamicOPrClusterUnionParent)
 			{
 				const FTransform& OriginalComponentSpaceRootTransformOffset = AssetCollection->Transform[RootIndex];
-				DynamicCollection->Transform[RootIndex] = FTransform3f(OriginalComponentSpaceRootTransformOffset);
+				DynamicCollection->SetTransform(RootIndex, FTransform3f(OriginalComponentSpaceRootTransformOffset));
 				const Chaos::FPBDRigidParticle* RootParticle = PhysicsProxy->GetExternalParticles()[RootIndex].Get();
 				const FTransform ParticleWorldPosition(RootParticle->R(), RootParticle->X());
 				const FTransform MassToLocal = MassToLocalAttribute[RootIndex];
@@ -4015,7 +4000,9 @@ void UGeometryCollectionComponent::UpdateRemovalIfNeeded()
 					const float Scale = 1.0 - Decay;
 					if (Scale < UE_SMALL_NUMBER)
 					{
-						DynamicCollection->Transform[TransformIndex].SetScale3D(FVector3f::ZeroVector);
+						FTransform3f Transform = DynamicCollection->GetTransform(TransformIndex);
+						Transform.SetScale3D(FVector3f::ZeroVector);
+						DynamicCollection->SetTransform(TransformIndex, Transform);
 					}
 					// do not try to get this condition out of the loop as this may cause some optimizer related issues
 					else if (RestCollection->bScaleOnRemoval && MassToLocal && CompSpaceTransform)
@@ -4033,7 +4020,7 @@ void UGeometryCollectionComponent::UpdateRemovalIfNeeded()
 						const FVector CenterOfMass = (*MassToLocal)[TransformIndex].GetTranslation();
 						const FVector ScaleCenter = LocalDown + CenterOfMass;
 						const FTransform ScaleTransform(FQuat::Identity, ScaleCenter * FVector::FReal(1.f - Scale), FVector(Scale));
-						DynamicCollection->Transform[TransformIndex] = FTransform3f(ScaleTransform) * DynamicCollection->Transform[TransformIndex];
+						DynamicCollection->SetTransform(TransformIndex, FTransform3f(ScaleTransform) * DynamicCollection->GetTransform(TransformIndex));
 					}
 				}
 			}
@@ -5408,7 +5395,7 @@ const FTransform& UGeometryCollectionComponent::FComponentSpaceTransforms::Reque
 		{
 			if (Component->DynamicCollection)
 			{
-				Transforms[RootIndex] = FTransform(Component->GetTransformArray()[RootIndex]);
+				Transforms[RootIndex] = FTransform(Component->DynamicCollection->GetTransform(RootIndex));
 			}
 			else
 			{
@@ -5442,7 +5429,7 @@ const TArray<FTransform>& UGeometryCollectionComponent::FComponentSpaceTransform
 	int32 CurrentTransformNum;
 	if (Component->DynamicCollection)
 	{
-		CurrentTransformNum = Component->GetTransformArray().Num();
+		CurrentTransformNum = Component->DynamicCollection->GetTransforms().Num();
 	}
 	else
 	{
@@ -5469,7 +5456,7 @@ const TArray<FTransform>& UGeometryCollectionComponent::FComponentSpaceTransform
 			FTransform CurrentTransform;
 			if (Component->DynamicCollection)
 			{
-				CurrentTransform = FTransform(Component->GetTransformArray()[TransformIndex]);
+				CurrentTransform = FTransform(Component->DynamicCollection->GetTransform(TransformIndex));
 			}
 			else
 			{
@@ -5491,7 +5478,7 @@ const TArray<FTransform>& UGeometryCollectionComponent::FComponentSpaceTransform
 	{
 		if (Component->DynamicCollection)
 		{
-			GeometryCollectionAlgo::GlobalMatrices(Component->GetTransformArray(), *Component->GetDynamicCollection(), Transforms);
+			GeometryCollectionAlgo::Private::GlobalMatrices(*Component->DynamicCollection, Transforms);
 		}
 		else
 		{
@@ -5923,7 +5910,7 @@ void UGeometryCollectionComponent::IncrementSleepTimer(float DeltaTime)
 						bool ShouldUpdateTimer = (DecayFacade.GetDecay(TransformIdx) > 0);
 						if (!ShouldUpdateTimer && RestCollection->bSlowMovingAsSleeping)
 						{
-							const FVector3f CurrentPosition = DynamicCollection->Transform[TransformIdx].GetTranslation();
+							const FVector3f CurrentPosition = DynamicCollection->GetTransform(TransformIdx).GetTranslation();
 							ShouldUpdateTimer |= RemoveOnSleepFacade.ComputeSlowMovingState(TransformIdx, FVector(CurrentPosition), DeltaTime, RestCollection->SlowMovingVelocityThreshold);
 						}
 						if (ShouldUpdateTimer || DynamicStateFacade.IsSleeping(TransformIdx))
diff --git a/Engine/Source/Runtime/Experimental/GeometryCollectionEngine/Public/GeometryCollection/GeometryCollectionComponent.h b/Engine/Source/Runtime/Experimental/GeometryCollectionEngine/Public/GeometryCollection/GeometryCollectionComponent.h
--- a/Engine/Source/Runtime/Experimental/GeometryCollectionEngine/Public/GeometryCollection/GeometryCollectionComponent.h
+++ b/Engine/Source/Runtime/Experimental/GeometryCollectionEngine/Public/GeometryCollection/GeometryCollectionComponent.h
@@ -894,13 +894,9 @@ public:
 	// Declare all the methods
 	COPY_ON_WRITE_ATTRIBUTES
 
-	GEOMETRYCOLLECTIONENGINE_API const TManagedArray<FTransform3f>& GetTransformArray() const;
-	GEOMETRYCOLLECTIONENGINE_API TManagedArray<FTransform3f>& GetTransformArrayCopyOnWrite();
-	GEOMETRYCOLLECTIONENGINE_API void ResetTransformArrayDynamic();
+	GEOMETRYCOLLECTIONENGINE_API void GetTransformArrayCopyOnWrite();
 	GEOMETRYCOLLECTIONENGINE_API const TManagedArray<FTransform>& GetTransformArrayRest() const;
-	private:
-		TManagedArray<FTransform3f>* IndirectTransformArray;
-	public:
+
 	GEOMETRYCOLLECTIONENGINE_API TManagedArray<int32>& GetParentArrayCopyOnWrite();
 	GEOMETRYCOLLECTIONENGINE_API int32 GetParent(int32 Index) const;
 	GEOMETRYCOLLECTIONENGINE_API const TManagedArray<int32>& GetParentArrayRest() const;
