Decrease memory consumption when creating a building. Now the transforms are not copied in dynamic collections. It will used the rest collection as long the dynamic collection is not modified. Lazy initializer on transforms.

diff --git a/Engine/Source/Programs/HeadlessChaos/Private/GeometryCollection/GeometryCollectionTestCollisionResolution.cpp b/Engine/Source/Programs/HeadlessChaos/Private/GeometryCollection/GeometryCollectionTestCollisionResolution.cpp
--- a/Engine/Source/Programs/HeadlessChaos/Private/GeometryCollection/GeometryCollectionTestCollisionResolution.cpp
+++ b/Engine/Source/Programs/HeadlessChaos/Private/GeometryCollection/GeometryCollectionTestCollisionResolution.cpp
@@ -49,7 +49,7 @@ namespace GeometryCollectionTest
 		}
 		{
 			// validate that Simplicials are null when CollisionType==Chaos_Volumetric
-			EXPECT_EQ(Collection->DynamicCollection->GetTransforms().Num(), 4);
+			EXPECT_EQ(Collection->DynamicCollection->GetNumTransforms(), 4);
 			EXPECT_EQ(Collection->DynamicCollection->Simplicials[0], nullptr);
 			EXPECT_EQ(Collection->DynamicCollection->Simplicials[1], nullptr);
 			EXPECT_EQ(Collection->DynamicCollection->Simplicials[2], nullptr);
@@ -96,7 +96,7 @@ namespace GeometryCollectionTest
 		}
 		{
 			// validate that Simplicials are null when CollisionType==Chaos_Volumetric
-			EXPECT_EQ(Collection->DynamicCollection->GetTransforms().Num(), 4);
+			EXPECT_EQ(Collection->DynamicCollection->GetNumTransforms(), 4);
 			EXPECT_TRUE(Collection->DynamicCollection->Simplicials[0]!=nullptr);
 			EXPECT_TRUE(UnitTest.Solver->GetParticles().GetGeometryCollectionParticles().CollisionParticles(0)!=nullptr);
 			EXPECT_TRUE(Collection->DynamicCollection->Simplicials[0]->Size() == UnitTest.Solver->GetParticles().GetGeometryCollectionParticles().CollisionParticles(0)->Size());
@@ -328,7 +328,7 @@ namespace GeometryCollectionTest
 		}
 		{
 			// validate simplicials and implicits are configured correctly
-			EXPECT_EQ(Collection->DynamicCollection->GetTransforms().Num(), 4);
+			EXPECT_EQ(Collection->DynamicCollection->GetNumTransforms(), 4);
 			EXPECT_TRUE(Collection->DynamicCollection->Simplicials[0] != nullptr);
 			EXPECT_TRUE(UnitTest.Solver->GetParticles().GetGeometryCollectionParticles().CollisionParticles(0) != nullptr);
 			EXPECT_TRUE(Collection->DynamicCollection->Simplicials[0]->Size() == UnitTest.Solver->GetParticles().GetGeometryCollectionParticles().CollisionParticles(0)->Size());
@@ -336,7 +336,7 @@ namespace GeometryCollectionTest
 			// The following test has been disabled because from now we remove the Implicits after initialization to free up some memory
 			// EXPECT_TRUE(Collection->DynamicCollection->GetAttribute<Chaos::FImplicitObjectPtr>(FGeometryDynamicCollection::ImplicitsAttribute, FTransformCollection::TransformGroup)[0]->GetType() == (int32)Chaos::ImplicitObjectType::LevelSet);
 
-			EXPECT_EQ(CollectionStaticSphere->DynamicCollection->GetTransforms().Num(), 4);
+			EXPECT_EQ(CollectionStaticSphere->DynamicCollection->GetNumTransforms(), 4);
 			EXPECT_TRUE(CollectionStaticSphere->DynamicCollection->Simplicials[0] == nullptr);
 			// The following test has been disabled because from now we remove the Implicits after initialization to free up some memory
 			// EXPECT_TRUE(CollectionStaticSphere->DynamicCollection->GetAttribute<Chaos::FImplicitObjectPtr>(FGeometryDynamicCollection::ImplicitsAttribute, FTransformCollection::TransformGroup)[0]->GetType() == (int32)Chaos::ImplicitObjectType::Sphere);
diff --git a/Engine/Source/Programs/HeadlessChaos/Private/GeometryCollection/GeometryCollectionTestSimulation.cpp b/Engine/Source/Programs/HeadlessChaos/Private/GeometryCollection/GeometryCollectionTestSimulation.cpp
--- a/Engine/Source/Programs/HeadlessChaos/Private/GeometryCollection/GeometryCollectionTestSimulation.cpp
+++ b/Engine/Source/Programs/HeadlessChaos/Private/GeometryCollection/GeometryCollectionTestSimulation.cpp
@@ -40,7 +40,7 @@ namespace GeometryCollectionTest
 
 		{ // test results
 			EXPECT_LT(FMath::Abs(Collection->RestCollection->Transform[0].GetTranslation().Z), SMALL_THRESHOLD); // rest never touched
-			EXPECT_EQ(Collection->DynamicCollection->GetTransforms().Num(), 1); // simulated is falling
+			EXPECT_EQ(Collection->DynamicCollection->GetNumTransforms(), 1); // simulated is falling
 			EXPECT_LT(Collection->DynamicCollection->GetTransform(0).GetTranslation().Z, 0.f);
 			EXPECT_NEAR(Collection->DynamicCollection->GetTransform(0).GetTranslation().Z, -980.f * UnitTest.Dt * UnitTest.Dt, 1e-2);// we seem to be twice gravity
 		}
@@ -70,7 +70,7 @@ namespace GeometryCollectionTest
 
 		{
 			EXPECT_LT(FMath::Abs(Collection->RestCollection->Transform[0].GetTranslation().Z), SMALL_THRESHOLD);
-			EXPECT_EQ(Collection->DynamicCollection->GetTransforms().Num(), 1);
+			EXPECT_EQ(Collection->DynamicCollection->GetNumTransforms(), 1);
 			EXPECT_LT(FMath::Abs(Collection->DynamicCollection->GetTransform(0).GetTranslation().Z - 0.1f * Scale), MEDIUM_THRESHOLD * Scale);
 		}
 	}
@@ -98,7 +98,7 @@ namespace GeometryCollectionTest
 
 		{ // test results
 			EXPECT_LT(FMath::Abs(Collection->RestCollection->Transform[0].GetTranslation().Z), SMALL_THRESHOLD);
-			EXPECT_EQ(Collection->DynamicCollection->GetTransforms().Num(), 1);
+			EXPECT_EQ(Collection->DynamicCollection->GetNumTransforms(), 1);
 			EXPECT_LT(FMath::Abs(Collection->DynamicCollection->GetTransform(0).GetTranslation().Z) - Scale[0], SMALL_THRESHOLD);
 		}
 	}
@@ -126,7 +126,7 @@ namespace GeometryCollectionTest
 
 		{
 			EXPECT_LT(FMath::Abs(Collection->RestCollection->Transform[0].GetTranslation().Z), SMALL_THRESHOLD);
-			EXPECT_EQ(Collection->DynamicCollection->GetTransforms().Num(), 1);
+			EXPECT_EQ(Collection->DynamicCollection->GetNumTransforms(), 1);
 			EXPECT_LT(FMath::Abs(Collection->DynamicCollection->GetTransform(0).GetTranslation().Z - 0.1f * Scale[0]), MEDIUM_THRESHOLD * Scale[0]);
 		}
 	}
@@ -145,7 +145,7 @@ namespace GeometryCollectionTest
 			UnitTest.Advance();
 
 		{
-			EXPECT_EQ(Collection->DynamicCollection->GetTransforms().Num(), 1);
+			EXPECT_EQ(Collection->DynamicCollection->GetNumTransforms(), 1);
 			//UE_LOG(LogTest, Verbose, TEXT("Position : (%3.5f,%3.5f,%3.5f)"), Transform[0].GetTranslation().X, Transform[0].GetTranslation().Y, Transform[0].GetTranslation().Z);
 			EXPECT_EQ(Collection->DynamicCollection->GetTransform(0).GetTranslation().Z, 0.f);
 			EXPECT_EQ(Collection->DynamicCollection->DynamicState[0], (int32)EObjectStateTypeEnum::Chaos_Object_Kinematic);
@@ -204,12 +204,12 @@ namespace GeometryCollectionTest
 		UnitTest.AddSimulationObject(MovingCollection);
 		UnitTest.Initialize();
 
-		const auto& Transform0 = MovingCollection->DynamicCollection->GetTransform(0);
-		const auto& Transform1 = SleepingCollection->DynamicCollection->GetTransform(0);
 		for (int i = 0; i < 15; i++)
 		{
 			UnitTest.Advance();
 
+			//const auto& Transform0 = MovingCollection->DynamicCollection->GetTransform(0);
+			//const auto& Transform1 = SleepingCollection->DynamicCollection->GetTransform(0);
 			//UE_LOG(LogTest, Verbose, TEXT("Position[0] : (%3.5f,%3.5f,%3.5f)"), Transform0.GetTranslation().X, Transform0.GetTranslation().Y, Transform0.GetTranslation().Z);
 			//UE_LOG(LogTest, Verbose, TEXT("Position[1] : (%3.5f,%3.5f,%3.5f)"), Transform1.GetTranslation().X, Transform1.GetTranslation().Y, Transform1.GetTranslation().Z);
 		}
@@ -217,7 +217,7 @@ namespace GeometryCollectionTest
 		{
 			// Is now dynamic and has moved from initial position
 			EXPECT_EQ(SleepingCollection->DynamicCollection->DynamicState[0], (int32)EObjectStateTypeEnum::Chaos_Object_Dynamic);
-			EXPECT_LT(Transform0.GetTranslation().Z, InitialStartHeight - 2.0f);
+			EXPECT_LT(MovingCollection->DynamicCollection->GetTransform(0).GetTranslation().Z, InitialStartHeight - 2.0f);
 		}
 
 	}
@@ -236,21 +236,20 @@ namespace GeometryCollectionTest
 		UnitTest.AddSimulationObject(MovingCollection);
 		UnitTest.Initialize();
 
-		const auto& Transform0 = MovingCollection->DynamicCollection->GetTransform(0);
 		for (int i = 0; i < 5; i++)
 		{
 			UnitTest.Advance();
 			EXPECT_EQ(MovingCollection->DynamicCollection->DynamicState[0], (int32)EObjectStateTypeEnum::Chaos_Object_Dynamic);
-			EXPECT_LT(Transform0.GetTranslation().Z, 15.0f);
+			EXPECT_LT(MovingCollection->DynamicCollection->GetTransform(0).GetTranslation().Z, 15.0f);
 		}
 		// Disabled particle
 		MovingCollection->PhysObject->DisableParticles_External({ 0 });
-		FReal CurrentPosition = Transform0.GetTranslation().Z;
+		FReal CurrentPosition = MovingCollection->DynamicCollection->GetTransform(0).GetTranslation().Z;
 		for (int i = 0; i < 5; i++)
 		{
 			UnitTest.Advance();
 			EXPECT_EQ(MovingCollection->DynamicCollection->DynamicState[0], (int32)EObjectStateTypeEnum::Chaos_Object_Dynamic);
-			EXPECT_EQ(Transform0.GetTranslation().Z, CurrentPosition);
+			EXPECT_EQ(MovingCollection->DynamicCollection->GetTransform(0).GetTranslation().Z, CurrentPosition);
 		}
 	}
 
diff --git a/Engine/Source/Programs/HeadlessChaos/Private/GeometryCollection/GeometryCollectionTestSimulationField.cpp b/Engine/Source/Programs/HeadlessChaos/Private/GeometryCollection/GeometryCollectionTestSimulationField.cpp
--- a/Engine/Source/Programs/HeadlessChaos/Private/GeometryCollection/GeometryCollectionTestSimulationField.cpp
+++ b/Engine/Source/Programs/HeadlessChaos/Private/GeometryCollection/GeometryCollectionTestSimulationField.cpp
@@ -90,7 +90,7 @@ namespace GeometryCollectionTest
 		EXPECT_EQ(DynamicState[0], (int32)EObjectStateTypeEnum::Chaos_Object_Kinematic);
 
 		// simulated
-		EXPECT_EQ(Collection->DynamicCollection->GetTransforms().Num(), 1);
+		EXPECT_EQ(Collection->DynamicCollection->GetNumTransforms(), 1);
 		const FVector Translation1 = FVector(Collection->DynamicCollection->GetTransform(0).GetTranslation());
 		EXPECT_NEAR((Translation0 - Translation1).Size(), 0.f, KINDA_SMALL_NUMBER);
 		EXPECT_NEAR(Collection->DynamicCollection->GetTransform(0).GetTranslation().Z, 1.f, KINDA_SMALL_NUMBER);
@@ -134,7 +134,7 @@ namespace GeometryCollectionTest
 		}
 
 		// simulated
-		EXPECT_EQ(Collection->DynamicCollection->GetTransforms().Num(), 1);
+		EXPECT_EQ(Collection->DynamicCollection->GetNumTransforms(), 1);
 		const FVector Translation1 = FVector(Collection->DynamicCollection->GetTransform(0).GetTranslation());
 		EXPECT_NEAR((Translation0 - Translation1).Size(), 0.f, KINDA_SMALL_NUMBER);
 		EXPECT_NEAR(Collection->DynamicCollection->GetTransform(0).GetTranslation().Z, 1.f, KINDA_SMALL_NUMBER);
diff --git a/Engine/Source/Programs/HeadlessChaos/Private/GeometryCollection/GeometryCollectionTestSimulationSolver.cpp b/Engine/Source/Programs/HeadlessChaos/Private/GeometryCollection/GeometryCollectionTestSimulationSolver.cpp
--- a/Engine/Source/Programs/HeadlessChaos/Private/GeometryCollection/GeometryCollectionTestSimulationSolver.cpp
+++ b/Engine/Source/Programs/HeadlessChaos/Private/GeometryCollection/GeometryCollectionTestSimulationSolver.cpp
@@ -62,7 +62,7 @@ using namespace ChaosTest;
 			EXPECT_LT(FMath::Abs(RestTransform[0].GetTranslation().Z), SMALL_THRESHOLD);
 
 			// simulated
-			EXPECT_EQ(Collection->DynamicCollection->GetTransforms().Num(), 1);
+			EXPECT_EQ(Collection->DynamicCollection->GetNumTransforms(), 1);
 			EXPECT_LT(FMath::Abs(Collection->DynamicCollection->GetTransform(0).GetTranslation().Z), SMALL_THRESHOLD);
 		}
 	}
diff --git a/Engine/Source/Programs/HeadlessChaos/Private/GeometryCollection/GeometryCollectionTestUtility.cpp b/Engine/Source/Programs/HeadlessChaos/Private/GeometryCollection/GeometryCollectionTestUtility.cpp
--- a/Engine/Source/Programs/HeadlessChaos/Private/GeometryCollection/GeometryCollectionTestUtility.cpp
+++ b/Engine/Source/Programs/HeadlessChaos/Private/GeometryCollection/GeometryCollectionTestUtility.cpp
@@ -22,7 +22,7 @@ namespace GeometryCollectionTest {
 	TSharedPtr<FGeometryDynamicCollection> GeometryCollectionToGeometryDynamicCollection(const FGeometryCollection* InputCollection, EObjectStateTypeEnum DynamicStateDefault)
 	{
 		TSharedPtr<FGeometryDynamicCollection> NewCollection(new FGeometryDynamicCollection(InputCollection));
-		NewCollection->CopyAttribute(*InputCollection, FTransformCollection::TransformAttribute, FGeometryCollection::TransformGroup);
+		// TransformAttribute will be initialized by a lazy initializer pattern
 		NewCollection->CopyAttribute(*InputCollection, FTransformCollection::ParentAttribute, FGeometryCollection::TransformGroup);
 		NewCollection->CopyAttribute(*InputCollection, FTransformCollection::ChildrenAttribute, FGeometryCollection::TransformGroup);
 		NewCollection->CopyAttribute(*InputCollection, FGeometryCollection::SimulationTypeAttribute, FGeometryCollection::TransformGroup);
diff --git a/Engine/Source/Runtime/Experimental/Chaos/Private/GeometryCollection/GeometryCollectionAlgo.cpp b/Engine/Source/Runtime/Experimental/Chaos/Private/GeometryCollection/GeometryCollectionAlgo.cpp
--- a/Engine/Source/Runtime/Experimental/Chaos/Private/GeometryCollection/GeometryCollectionAlgo.cpp
+++ b/Engine/Source/Runtime/Experimental/Chaos/Private/GeometryCollection/GeometryCollectionAlgo.cpp
@@ -355,7 +355,7 @@ namespace GeometryCollectionAlgo
 		}
 
 		FIndicesNeedMatricesArray ToProcess;
-		if (!GlobalMatricesGetIndicesToProcessHelper(Index, DynamicCollection, DynamicCollection.GetTransforms().Num(), IsTransformComputed, ToProcess))
+		if (!GlobalMatricesGetIndicesToProcessHelper(Index, DynamicCollection, DynamicCollection.GetNumTransforms(), IsTransformComputed, ToProcess))
 		{
 			return;
 		}
@@ -464,7 +464,7 @@ namespace GeometryCollectionAlgo
 		}
 
 		FIndicesNeedMatricesArray ToProcess;
-		if (!GlobalMatricesGetIndicesToProcessHelper(Index, DynamicCollection, DynamicCollection.GetTransforms().Num(), IsTransformComputed, ToProcess))
+		if (!GlobalMatricesGetIndicesToProcessHelper(Index, DynamicCollection, DynamicCollection.GetNumTransforms(), IsTransformComputed, ToProcess))
 		{
 			return TransformCache[Index];
 		}
@@ -538,13 +538,11 @@ namespace GeometryCollectionAlgo
 		{
 			FTransform Transform = FTransform::Identity;
 
-			if (DynamicCollection.GetTransforms().IsValidIndex(Index))
+			check(Index < DynamicCollection.GetNumTransforms())
+			while (Index != FGeometryCollection::Invalid)
 			{
-				while (Index != FGeometryCollection::Invalid)
-				{
-					Transform = Transform * FTransform(DynamicCollection.GetTransform(Index));
-					Index = DynamicCollection.GetParent(Index);
-				}
+				Transform = Transform * FTransform(DynamicCollection.GetTransform(Index));
+				Index = DynamicCollection.GetParent(Index);
 			}
 			return Transform;
 		}
@@ -570,7 +568,7 @@ namespace GeometryCollectionAlgo
 		void GlobalMatrices(const FGeometryDynamicCollection& DynamicCollection, const TArray<int32>& Indices, TArray<FTransform>& OutGlobalTransforms)
 		{
 			TArray<bool> IsTransformComputed;
-			const int32 NumTransform = DynamicCollection.GetTransforms().Num();
+			const int32 NumTransform = DynamicCollection.GetNumTransforms();
 			IsTransformComputed.AddDefaulted(NumTransform);
 
 			TArray<FTransform> TransformCache;
@@ -638,7 +636,7 @@ namespace GeometryCollectionAlgo
 	{
 		void GlobalMatrices(const FGeometryDynamicCollection& DynamicCollection, TArray<FTransform>& OutGlobalTransforms)
 		{
-			int32 NumTransforms = DynamicCollection.GetTransforms().Num();
+			int32 NumTransforms = DynamicCollection.GetNumTransforms();
 
 			TArray<bool> IsTransformComputed;
 			IsTransformComputed.AddDefaulted(NumTransforms);
diff --git a/Engine/Source/Runtime/Experimental/Chaos/Private/GeometryCollectionProxyData.cpp b/Engine/Source/Runtime/Experimental/Chaos/Private/GeometryCollectionProxyData.cpp
--- a/Engine/Source/Runtime/Experimental/Chaos/Private/GeometryCollectionProxyData.cpp
+++ b/Engine/Source/Runtime/Experimental/Chaos/Private/GeometryCollectionProxyData.cpp
@@ -15,6 +15,7 @@ GeometryCollectionProxyData.cpp:
 FTransformDynamicCollection::FTransformDynamicCollection(const FGeometryCollection* InRestCollection)
 	: FManagedArrayCollection()
 	, RestCollection(InRestCollection)
+	, bTransformHasChanged(false)
 {
 	check(RestCollection != nullptr);
 	Construct();
@@ -25,23 +26,39 @@ void FTransformDynamicCollection::Construct()
 	FManagedArrayCollection::FConstructionParameters TransformDependency(FTransformCollection::TransformGroup);
 
 	// Transform Group
-	AddExternalAttribute<FTransform3f>(FTransformCollection::TransformAttribute, FTransformCollection::TransformGroup, Transform);
 	AddExternalAttribute<bool>(FTransformCollection::ParentAttribute, FTransformCollection::TransformGroup, HasParent);
 }
 
-const FTransform3f& FTransformDynamicCollection::GetTransform(int32 Index) const
+void FTransformDynamicCollection::InitializeTransforms()
 {
+	if (bTransformHasChanged == false)
+	{
+		AddExternalAttribute<FTransform3f>(FTransformCollection::TransformAttribute, FTransformCollection::TransformGroup, Transform);
+		CopyAttribute(*RestCollection, FTransformCollection::TransformAttribute, FTransformCollection::TransformGroup);
+
+		bTransformHasChanged = true;
+	}
+}
+
+FTransform3f FTransformDynamicCollection::GetTransform(int32 Index) const
+{
+	if (bTransformHasChanged == false)
+	{
+		return FTransform3f(RestCollection->Transform[Index]);
+	}
 	return Transform[Index];
 }
 
 void FTransformDynamicCollection::SetTransform(int32 Index, const FTransform3f& InTransform)
 {
+	InitializeTransforms();
 	Transform[Index] = InTransform;
 }
 
-const TArray<FTransform3f>& FTransformDynamicCollection::GetTransforms() const
+int32 FTransformDynamicCollection::GetNumTransforms() const
 {
-	return Transform.GetConstArray();
+	ensure(!bTransformHasChanged || RestCollection->Transform.Num() == Transform.Num());
+	return RestCollection->Transform.Num();
 }
 
 const TManagedArray<bool>& FTransformDynamicCollection::GetHasParent() const
@@ -56,6 +73,7 @@ bool FTransformDynamicCollection::GetHasParent(int32 Index) const
 
 void FTransformDynamicCollection::SetHasParent(int32 Index, bool Value)
 {
+	InitializeTransforms();
 	HasParent[Index] = Value;
 }
 
diff --git a/Engine/Source/Runtime/Experimental/Chaos/Private/PhysicsProxy/GeometryCollectionPhysicsProxy.cpp b/Engine/Source/Runtime/Experimental/Chaos/Private/PhysicsProxy/GeometryCollectionPhysicsProxy.cpp
--- a/Engine/Source/Runtime/Experimental/Chaos/Private/PhysicsProxy/GeometryCollectionPhysicsProxy.cpp
+++ b/Engine/Source/Runtime/Experimental/Chaos/Private/PhysicsProxy/GeometryCollectionPhysicsProxy.cpp
@@ -492,7 +492,7 @@ void FGeometryCollectionPhysicsProxy::Initialize(Chaos::FPBDRigidsEvolutionBase
 {
 	check(IsInGameThread());
 	//
-	// Game thread initilization. 
+	// Game thread initialization. 
 	//
 	//  1) Create a input buffer to store all game thread side data. 
 	//  2) Populate the buffer with the necessary data.
@@ -549,7 +549,7 @@ void FGeometryCollectionPhysicsProxy::Initialize(Chaos::FPBDRigidsEvolutionBase
 	}
 
 	// Initialise GT/External particles
-	const int32 NumTransforms = DynamicCollection.GetTransforms().Num();
+	const int32 NumTransforms = DynamicCollection.GetNumTransforms();
 
 	// make sure we copy the anchored information over to the physics thread collection
 	const Chaos::Facades::FCollectionAnchoringFacade DynamicCollectionAnchoringFacade(DynamicCollection);
@@ -3363,7 +3363,7 @@ void FGeometryCollectionPhysicsProxy::BufferPhysicsResults_Internal(Chaos::FPBDR
 
 	UniqueIdxToInternalClusterHandle.Reset();
 	
-	const int32 NumTransforms = PhysicsThreadCollection.GetTransforms().Num();
+	const int32 NumTransforms = PhysicsThreadCollection.GetNumTransforms();
 	if(NumTransforms > 0)
 	{ 
 		SCOPE_CYCLE_COUNTER(STAT_CalcParticleToWorld);
@@ -3831,7 +3831,7 @@ bool FGeometryCollectionPhysicsProxy::PullFromPhysicsState(const Chaos::FDirtyGe
 		InterpolationData.AccumlateErrorXR(Error->ErrorX, Error->ErrorR, SolverSyncTimestamp, RenderInterpErrorCorrectionDurationTicks);
 	}
 
-	const int32 NumTransforms = GameThreadCollection.GetTransforms().Num();
+	const int32 NumTransforms = GameThreadCollection.GetNumTransforms();
 	const bool bNeedInterpolation = (NextPullData != nullptr);
 
 	bool bIsCollectionDirty = false;
@@ -4003,7 +4003,7 @@ void FGeometryCollectionPhysicsProxy::UpdateFilterData_External(const FCollision
 	check(IsInGameThread());
 
 	// SimFilter/QueryFilter members are read on both threads, these are const after initialization and are not updated here.
-	const int32 NumTransforms = GameThreadCollection.GetTransforms().Num();
+	const int32 NumTransforms = GameThreadCollection.GetNumTransforms();
 	for (int32 Index = 0; Index < NumTransforms; ++Index)
 	{
 		FParticle* P = GTParticles[Index].Get();
@@ -5056,7 +5056,7 @@ void FGeometryCollectionPhysicsProxy::FieldParameterUpdateCallback(Chaos::FPBDRi
 
 	// Process Particle-Collection commands
 	int32 NumCommands = Commands.Num();
-	if (NumCommands && !RigidSolver->IsShuttingDown() && Collection.GetTransforms().Num())
+	if (NumCommands && !RigidSolver->IsShuttingDown() && Collection.GetNumTransforms())
 	{
 		TArray<int32> CommandsToRemove;
 		CommandsToRemove.Reserve(NumCommands);
diff --git a/Engine/Source/Runtime/Experimental/Chaos/Public/GeometryCollectionProxyData.h b/Engine/Source/Runtime/Experimental/Chaos/Public/GeometryCollectionProxyData.h
--- a/Engine/Source/Runtime/Experimental/Chaos/Public/GeometryCollectionProxyData.h
+++ b/Engine/Source/Runtime/Experimental/Chaos/Public/GeometryCollectionProxyData.h
@@ -28,9 +28,9 @@ public:
 	FTransformDynamicCollection(FTransformDynamicCollection&&) = delete;
 	FTransformDynamicCollection& operator=(FTransformDynamicCollection&&) = delete;
 
-	CHAOS_API const FTransform3f& GetTransform(int32 Index) const;
+	CHAOS_API FTransform3f GetTransform(int32 Index) const;
 	CHAOS_API void SetTransform(int32 Index, const FTransform3f& Transform);
-	CHAOS_API const TArray<FTransform3f>& GetTransforms() const;
+	CHAOS_API int32 GetNumTransforms() const;
 
 	CHAOS_API const TManagedArray<bool>& GetHasParent() const;
 	CHAOS_API bool GetHasParent(int32 Index) const;
@@ -62,9 +62,11 @@ protected:
 private:
 	TManagedArray<bool>         HasParent;
 	TManagedArray<FTransform3f> Transform;
+	bool bTransformHasChanged;
 
 	/** Construct */
 	CHAOS_API void Construct();
+	void InitializeTransforms();
 };
 
 
diff --git a/Engine/Source/Runtime/Experimental/GeometryCollectionEngine/Private/GeometryCollection/GeometryCollectionComponent.cpp b/Engine/Source/Runtime/Experimental/GeometryCollectionEngine/Private/GeometryCollection/GeometryCollectionComponent.cpp
--- a/Engine/Source/Runtime/Experimental/GeometryCollectionEngine/Private/GeometryCollection/GeometryCollectionComponent.cpp
+++ b/Engine/Source/Runtime/Experimental/GeometryCollectionEngine/Private/GeometryCollection/GeometryCollectionComponent.cpp
@@ -503,13 +503,6 @@ private:
 COPY_ON_WRITE_ATTRIBUTES
 
 
-void UGeometryCollectionComponent::GetTransformArrayCopyOnWrite()
-{
-	DynamicCollection->AddAttribute<FTransform3f>(FTransformCollection::TransformAttribute, FTransformCollection::TransformGroup);
-	DynamicCollection->CopyAttribute(*RestCollection->GetGeometryCollection(), FTransformCollection::TransformAttribute, FTransformCollection::TransformGroup);
-	DynamicCollection->ModifyAttribute<FTransform3f>(FTransformCollection::TransformAttribute, FTransformCollection::TransformGroup);
-}
-
 const TManagedArray<FTransform>& UGeometryCollectionComponent::GetTransformArrayRest() const
 {
 	return RestCollection->GetGeometryCollection()->Transform;
@@ -2919,7 +2912,7 @@ void UGeometryCollectionComponent::SetInitialTransforms(const TArray<FTransform>
 {
 	if (DynamicCollection)
 	{
-		int32 MaxIdx = FMath::Min(DynamicCollection->GetTransforms().Num(), InitialTransforms.Num());
+		int32 MaxIdx = FMath::Min(DynamicCollection->GetNumTransforms(), InitialTransforms.Num());
 		for (int32 Idx = 0; Idx < MaxIdx; ++Idx)
 		{
 			DynamicCollection->SetTransform(Idx, FTransform3f(InitialTransforms[Idx]));
@@ -2931,7 +2924,7 @@ void UGeometryCollectionComponent::SetInitialClusterBreaks(const TArray<int32>&
 {
 	if (DynamicCollection)
 	{
-		const int32 NumTransforms = DynamicCollection->GetTransforms().Num();
+		const int32 NumTransforms = DynamicCollection->GetNumTransforms();
 
 		for (int32 ReleaseIndex : ReleaseIndices)
 		{
@@ -3487,7 +3480,6 @@ void UGeometryCollectionComponent::ResetDynamicCollection()
 			*DynamicArray = nullptr;
 		}
 
-		GetTransformArrayCopyOnWrite();
 		GetParentArrayCopyOnWrite();
 
 		if (bStoreVelocities || bNotifyTrailing)
@@ -3571,7 +3563,7 @@ void UGeometryCollectionComponent::OnCreatePhysicsState()
 		}
 #endif
 		const bool bValidWorld = GetWorld() && (GetWorld()->IsGameWorld() || GetWorld()->IsPreviewWorld() || GeometryCollectionCreatePhysicsStateInEditor);
-		const bool bValidCollection = DynamicCollection && DynamicCollection->GetTransforms().Num() > 0;
+		const bool bValidCollection = DynamicCollection && DynamicCollection->GetNumTransforms() > 0;
 		if (bValidWorld && bValidCollection)
 		{
 			FChaosUserData::Set<UPrimitiveComponent>(&PhysicsUserData, this);
@@ -5429,7 +5421,7 @@ const TArray<FTransform>& UGeometryCollectionComponent::FComponentSpaceTransform
 	int32 CurrentTransformNum;
 	if (Component->DynamicCollection)
 	{
-		CurrentTransformNum = Component->DynamicCollection->GetTransforms().Num();
+		CurrentTransformNum = Component->DynamicCollection->GetNumTransforms();
 	}
 	else
 	{
diff --git a/Engine/Source/Runtime/Experimental/GeometryCollectionEngine/Private/GeometryCollection/GeometryCollectionEngineRemoval.cpp b/Engine/Source/Runtime/Experimental/GeometryCollectionEngine/Private/GeometryCollection/GeometryCollectionEngineRemoval.cpp
--- a/Engine/Source/Runtime/Experimental/GeometryCollectionEngine/Private/GeometryCollection/GeometryCollectionEngineRemoval.cpp
+++ b/Engine/Source/Runtime/Experimental/GeometryCollectionEngine/Private/GeometryCollection/GeometryCollectionEngineRemoval.cpp
@@ -58,7 +58,7 @@ void FGeometryCollectionRemoveOnBreakDynamicFacade::SetAttributeValues(const Geo
 
 		// make sure we generate random value consistently between client and server 
 		// we can use the length of the transform group for that
-		const int32 RandomSeed = DynamicCollection.GetTransforms().Num();
+		const int32 RandomSeed = DynamicCollection.GetNumTransforms();
 		FRandomStream Random(RandomSeed);
 
 		TManagedArray<float>& PostBreakDuration = PostBreakDurationAttribute.Modify();
diff --git a/Engine/Source/Runtime/Experimental/GeometryCollectionEngine/Public/GeometryCollection/GeometryCollectionComponent.h b/Engine/Source/Runtime/Experimental/GeometryCollectionEngine/Public/GeometryCollection/GeometryCollectionComponent.h
--- a/Engine/Source/Runtime/Experimental/GeometryCollectionEngine/Public/GeometryCollection/GeometryCollectionComponent.h
+++ b/Engine/Source/Runtime/Experimental/GeometryCollectionEngine/Public/GeometryCollection/GeometryCollectionComponent.h
@@ -894,7 +894,6 @@ public:
 	// Declare all the methods
 	COPY_ON_WRITE_ATTRIBUTES
 
-	GEOMETRYCOLLECTIONENGINE_API void GetTransformArrayCopyOnWrite();
 	GEOMETRYCOLLECTIONENGINE_API const TManagedArray<FTransform>& GetTransformArrayRest() const;
 
 	GEOMETRYCOLLECTIONENGINE_API TManagedArray<int32>& GetParentArrayCopyOnWrite();
