Chaos - fix FChaosEngineInterface::SetGlobalPose_AssumesLocked for kinematics

Set the kinematic target in SetGlobalPose so that this is a teleport and velocity reset for kinematics as intended (if a teleport is not desired, the user should be calling SetKinematicTarget, not SetGlobalPose). Previously there could be problems with interpolation in async physics.

diff --git a/Engine/Source/Programs/HeadlessChaos/Private/HeadlessChaosTestPhysicsScene.cpp b/Engine/Source/Programs/HeadlessChaos/Private/HeadlessChaosTestPhysicsScene.cpp
--- a/Engine/Source/Programs/HeadlessChaos/Private/HeadlessChaosTestPhysicsScene.cpp
+++ b/Engine/Source/Programs/HeadlessChaos/Private/HeadlessChaosTestPhysicsScene.cpp
@@ -1333,7 +1333,14 @@ namespace ChaosTest {
 		EXPECT_NEAR(Particle.V()[2], ZVel, 1e-2);
 	}
 
-	GTEST_TEST(EngineInterface, SetKinematicTarget)
+	void ExpectVectorEqual(const FVec3& V0, const FVec3& V1)
+	{
+		EXPECT_EQ(V0.X, V1.X);
+		EXPECT_EQ(V0.Y, V1.Y);
+		EXPECT_EQ(V0.Z, V1.Z);
+	}
+
+	void TestKinematicTarget(const bool bInUpdateKinematicFromSimulation)
 	{
 		// Need to test:
 		// GT particle position is immediately updated after calling SetKinematicTarget_AssumesLocked
@@ -1362,12 +1369,14 @@ namespace ChaosTest {
 		TArray<FPhysicsActorHandle> Proxys = { Proxy };
 		Scene.AddActorsToScene_AssumesLocked(Proxys);
 		Particle.SetObjectState(EObjectStateType::Kinematic);
+		Particle.SetUpdateKinematicFromSimulation(bInUpdateKinematicFromSimulation);
 
 		struct FDummyInput : FSimCallbackInput
 		{
 			FSingleParticlePhysicsProxy* Proxy;
 			FVec3 CorrectX;
 			FVec3 CorrectV;
+			bool bKinematicWritebackEnabled;
 			void Reset() {}
 		};
 
@@ -1375,15 +1384,19 @@ namespace ChaosTest {
 		{
 			virtual void OnPreSimulate_Internal() override
 			{
+				const FVec3 ExpectedX = GetConsumerInput_Internal()->CorrectX;
+				const FVec3 ExpectedV = GetConsumerInput_Internal()->CorrectV;
+
 				auto Handle = GetConsumerInput_Internal()->Proxy->GetPhysicsThreadAPI();
-				EXPECT_EQ(Handle->X(), GetConsumerInput_Internal()->CorrectX);
-				EXPECT_EQ(Handle->V(), GetConsumerInput_Internal()->CorrectV);
+				ExpectVectorEqual(Handle->X(), ExpectedX);
+				ExpectVectorEqual(Handle->V(), ExpectedV);
 			}
 		};
 
 		auto Callback = Scene.GetSolver()->CreateAndRegisterSimCallbackObject_External<FCallback>();
 
 		Callback->GetProducerInputData_External()->Proxy = Proxy;
+		Callback->GetProducerInputData_External()->bKinematicWritebackEnabled = bInUpdateKinematicFromSimulation;
 
 		FVec3 Grav(0, 0, 0);
 		float Dt = 1;
@@ -1394,12 +1407,12 @@ namespace ChaosTest {
 			Scene.StartFrame();
 			Scene.EndFrame();
 			// Test X and V on GT
-			EXPECT_EQ(Particle.X().X, CorrectX.X);
-			EXPECT_EQ(Particle.X().Y, CorrectX.Y);
-			EXPECT_EQ(Particle.X().Z, CorrectX.Z);
-			EXPECT_EQ(Particle.V().X, CorrectV.X);
-			EXPECT_EQ(Particle.V().Y, CorrectV.Y);
-			EXPECT_EQ(Particle.V().Z, CorrectV.Z);
+			// NOTE: GT velocity will not be updated if kinematic writeback from the physics thread is disabled
+			ExpectVectorEqual(Particle.X(), CorrectX);
+			if (Callback->GetProducerInputData_External()->bKinematicWritebackEnabled)
+			{
+				ExpectVectorEqual(Particle.V(), CorrectV);
+			}
 			// Test X and V on PT, this is going to be used in OnPreSimulate_Internal in next frame.
 			Callback->GetProducerInputData_External()->CorrectX = CorrectX;
 			Callback->GetProducerInputData_External()->CorrectV = CorrectV;
@@ -1419,9 +1432,13 @@ namespace ChaosTest {
 		CurrentV = FVec3(1, 1, 1);
 		FChaosEngineInterface::SetKinematicTarget_AssumesLocked(Proxy, FTransform(CurrentX));
 
-		// Test if position is immediately updated on GT after SetKinematicTarget_AssumesLocked
-		EXPECT_EQ(Particle.X(), CurrentX);
+		// Test if position is immediately updated on GT after SetKinematicTarget_AssumesLocked (if we aren't reading data back from PT)
+		if (!bInUpdateKinematicFromSimulation)
+		{
+			ExpectVectorEqual(Particle.X(), CurrentX);
+		}
 
+		// This will fail when bInUpdateKinematicFromSimulation is false becasuse GT and PT disagree on velocity
 		AdvanceFrameAndRunTest(CurrentX, CurrentV);
 
 		// Test if velocity becomes zero when no kinematic target is set
@@ -1444,6 +1461,7 @@ namespace ChaosTest {
 		FChaosEngineInterface::SetGlobalPose_AssumesLocked(Proxy, FTransform(CurrentX));
 
 		Callback->GetProducerInputData_External()->CorrectX = CurrentX;
+		Callback->GetProducerInputData_External()->CorrectV = CurrentV;
 		AdvanceFrameAndRunTest(CurrentX, CurrentV);
 
 		// Test if particle positions and velocities are correct after SetGlobalPose_AssumesLocked, SetKinematicTarget_AssumesLocked
@@ -1475,6 +1493,18 @@ namespace ChaosTest {
 		AdvanceFrameAndRunTest(CurrentX, CurrentV);
 	}
 
+	// Test SetKinematicTarget when writeback from PT is enabled
+	GTEST_TEST(EngineInterface, SetKinematicTargetWriteBackEnabled)
+	{
+		TestKinematicTarget(true);
+	}
+
+	// Test SetKinematicTarget when writeback from PT is disabled
+	GTEST_TEST(EngineInterface, SetKinematicTargetWriteBackDisabled)
+	{
+		TestKinematicTarget(false);
+	}
+
 	GTEST_TEST(EngineInterface, PerPropertySetOnGT)
 	{
 		//Need to test:
diff --git a/Engine/Source/Runtime/PhysicsCore/Private/ChaosEngineInterface.cpp b/Engine/Source/Runtime/PhysicsCore/Private/ChaosEngineInterface.cpp
--- a/Engine/Source/Runtime/PhysicsCore/Private/ChaosEngineInterface.cpp
+++ b/Engine/Source/Runtime/PhysicsCore/Private/ChaosEngineInterface.cpp
@@ -2168,22 +2168,16 @@ void FChaosEngineInterface::SetGlobalPose_AssumesLocked(const FPhysicsActorHandl
 
 	if (IsKinematic(InActorReference))
 	{
-		// If we have already set a kinematic target this tick it must be cleared. NOTE: because we are clearing the target, 
-		// the Velocity will be zeroed on the next simulation tick unless there is a subsequent call to SetKinematicTarget 
-		// with an appropriate position delta. 
-		Body_External.ClearKinematicTarget();
-		Body_External.SetX(InNewPose.GetLocation());
-		Body_External.SetR(InNewPose.GetRotation());
-		// @todo(chaos): Should we also explicitly set the velocity to zero here?
-		//Body_External.SetV(FVector::Zero());
-		//Body_External.SetW(FVector::Zero());
-	}
-	else
-	{
-		Body_External.SetX(InNewPose.GetLocation());
-		Body_External.SetR(InNewPose.GetRotation());
+		// NOTE: SetGlobalPose is a teleport for kinematics. Use SetKinematicTarget_AssumesLocked
+		// if the kinematic should calculate its velocity from the transform delta.
+		Body_External.SetKinematicTarget(InNewPose);
+		Body_External.SetV(FVector::Zero());
+		Body_External.SetW(FVector::Zero());
 	}
 
+	Body_External.SetX(InNewPose.GetLocation());
+	Body_External.SetR(InNewPose.GetRotation());
+
 	Body_External.UpdateShapeBounds();
 
 	FChaosScene* Scene = GetCurrentScene(InActorReference);
@@ -2219,11 +2213,15 @@ void FChaosEngineInterface::SetKinematicTarget_AssumesLocked(const FPhysicsActor
 	const Chaos::FKinematicTarget NewKinematicTarget = Chaos::FKinematicTarget::MakePositionTarget(InNewTarget);
 	InActorReference->GetGameThreadAPI().SetKinematicTarget(NewKinematicTarget);
 
+	// If enabled for this body, immediately update the body transforms to match the kinematic target.
+	// @todo(chaos): Velocity is not updated here and never will be because we don't read back from the physics thread.
+	// We should fix this, but it is awkward to handle multiple calls to SetKinematicTarget on the same frame if we
+	// don't have a "previous transform" from which to calculate the velocity and we have overwritten X/R already.
 	if (ShouldSetKinematicTargetSetGameTransform(InActorReference))
 	{
 		// IMPORTANT : we do not invalidate X and R as they will be properly computed using the kinematic target information 
-		InActorReference->GetGameThreadAPI().SetX(InNewTarget.GetLocation(), false); 
-		InActorReference->GetGameThreadAPI().SetR(InNewTarget.GetRotation(), false); 
+		InActorReference->GetGameThreadAPI().SetX(InNewTarget.GetLocation(), false);
+		InActorReference->GetGameThreadAPI().SetR(InNewTarget.GetRotation(), false);
 		InActorReference->GetGameThreadAPI().UpdateShapeBounds();
 
 		FChaosScene* Scene = GetCurrentScene(InActorReference);
