Overlap bug fix: Heightfields should Normalize Initial direction vector before sending it to GJK

#[fyi] Michael.Lentine, Benn.Gallagher

diff --git a/Engine/Source/Programs/HeadlessChaos/Private/HeadlessChaosTestHeightfield.cpp b/Engine/Source/Programs/HeadlessChaos/Private/HeadlessChaosTestHeightfield.cpp
--- a/Engine/Source/Programs/HeadlessChaos/Private/HeadlessChaosTestHeightfield.cpp
+++ b/Engine/Source/Programs/HeadlessChaos/Private/HeadlessChaosTestHeightfield.cpp
@@ -895,7 +895,7 @@ namespace ChaosTest {
 				Heights[Row * Columns + Col] = 5.0;
 			}
 		}
-		// Submit
+		// Summit
 		Heights[(Rows / 2) * Columns + Columns / 2] = 10.0;
 
 		return Heights;
@@ -1238,6 +1238,42 @@ namespace ChaosTest {
 			EXPECT_TRUE(BoxResult);
 			EXPECT_TRUE(CapsuleResult);
 		}
+
+		HeightsCopy = Heights;
+		const float UniformScale = 10.0f;
+		FHeightField HeightfieldScaled(MoveTemp(HeightsCopy), TArray<uint8>(), Rows, Columns, FVec3(UniformScale, UniformScale, UniformScale));
+		const auto& Bounds2 = HeightfieldScaled.BoundingBox();	//Current API forces us to do this to cache the bounds
+		// Small sphere intersecting plane area of mountain with a scaled heightfield
+		{
+			Chaos::FSphere Sphere1(FVec3(0.0, 0.0, 0.0), 0.2);
+			int32 Row = 2;
+			int32 Col = 2;
+			
+			const FVec3 Translation(Col* UniformScale, Row * UniformScale, 0.1f);
+			FRigidTransform3 QueryTM(Translation, TRotation<FReal, 3>::Identity);
+			bool SphereResult = HeightfieldScaled.OverlapGeom(Sphere1, QueryTM, 0.0, nullptr);
+			EXPECT_TRUE(SphereResult);
+		}
+
+		// Small convex intersecting plane area of mountain with a scaled heightfield
+		{
+			//Tetrahedron
+			TArray<FConvex::FVec3Type> HullParticles;
+			HullParticles.SetNum(4);
+			HullParticles[0] = { -1,-1,-1 };
+			HullParticles[1] = { 1,-1,-1 };
+			HullParticles[2] = { 0,1,-1 };
+			HullParticles[3] = { 0,0,1 };
+			FConvex Tet(HullParticles, 0.0f);
+			
+			int32 Row = 2;
+			int32 Col = 2;
+
+			const FVec3 Translation(Col * UniformScale, Row * UniformScale, 1.1f);
+			FRigidTransform3 QueryTM(Translation, TRotation<FReal, 3>::Identity);
+			bool TetResult = HeightfieldScaled.OverlapGeom(Tet, QueryTM, 0.2, nullptr);
+			EXPECT_TRUE(TetResult);
+		}
 	}
 	
 
diff --git a/Engine/Source/Runtime/Experimental/Chaos/Private/Chaos/HeightField.cpp b/Engine/Source/Runtime/Experimental/Chaos/Private/Chaos/HeightField.cpp
--- a/Engine/Source/Runtime/Experimental/Chaos/Private/Chaos/HeightField.cpp
+++ b/Engine/Source/Runtime/Experimental/Chaos/Private/Chaos/HeightField.cpp
@@ -1967,7 +1967,7 @@ namespace Chaos
 			const VectorRegister4Float Offset = VectorCross(AB, AC);
 
 			FTriangleRegister TriangleConvex(A, B, C);
-			return GJKIntersectionSameSpaceSimd(TriangleConvex, QueryGeom, Thickness, Offset);
+			return GJKIntersectionSameSpaceSimd(TriangleConvex, QueryGeom, Thickness, VectorNormalizeSafe(Offset, GlobalVectorConstants::Float1000));
 		};
 
 		bool bResult = false;
diff --git a/Engine/Source/Runtime/Experimental/Chaos/Public/Chaos/GJK.h b/Engine/Source/Runtime/Experimental/Chaos/Public/Chaos/GJK.h
--- a/Engine/Source/Runtime/Experimental/Chaos/Public/Chaos/GJK.h
+++ b/Engine/Source/Runtime/Experimental/Chaos/Public/Chaos/GJK.h
@@ -239,7 +239,7 @@ namespace Chaos
 		@param A The first geometry
 		@param B The second geometry
 		@param ThicknessA The amount of geometry inflation for Geometry A(for example if the surface distance of two geometries with thickness 0 would be 2, a thickness of 0.5 would give a distance of 1.5)
-		@param InitialDir The first direction we use to search the CSO
+		@param InitialDir The first direction we use to search the CSO (Must be normalized)
 		@param ThicknessB The amount of geometry inflation for Geometry B(for example if the surface distance of two geometries with thickness 0 would be 2, a thickness of 0.5 would give a distance of 1.5)
 		@return True if the geometries overlap, False otherwise 
 	*/
@@ -249,6 +249,8 @@ namespace Chaos
 	{
 		VectorRegister4Float V = VectorNegate(InitialDir);
 
+		//ensure(FMath::Abs(VectorDot3Scalar(V, V) - 1.0f) < 1e-3f ); // Check that the vector is normalized, comment out for performance
+
 		FSimplex SimplexIDs;
 		VectorRegister4Float Simplex[4] = { VectorZero(), VectorZero(), VectorZero(), VectorZero() };
 		VectorRegister4Int NumVerts = GlobalVectorConstants::IntZero;
