Add modular unit tests

diff --git a/Engine/Plugins/Experimental/ChaosModularVehicle/Source/ChaosModularVehicleEngine/Private/ChaosModularVehicle/VehicleSimEngineComponent.cpp b/Engine/Plugins/Experimental/ChaosModularVehicle/Source/ChaosModularVehicleEngine/Private/ChaosModularVehicle/VehicleSimEngineComponent.cpp
--- a/Engine/Plugins/Experimental/ChaosModularVehicle/Source/ChaosModularVehicleEngine/Private/ChaosModularVehicle/VehicleSimEngineComponent.cpp
+++ b/Engine/Plugins/Experimental/ChaosModularVehicle/Source/ChaosModularVehicleEngine/Private/ChaosModularVehicle/VehicleSimEngineComponent.cpp
@@ -15,8 +15,8 @@ UVehicleSimEngineComponent::UVehicleSimEngineComponent()
 	MaxTorque = 200;
 	MaxRPM = 5000;
 	EngineIdleRPM = 1200;
-	EngineBrakeEffect = 0.5f;
-	MomentOfInertia = 10.0f; //??		// [Kg.m-2] How hard it is to turn the engine
+	EngineBrakeEffect = 150.0f;
+	EngineInertia = 1000.0f; 	// [Kg.m-2] How hard it is to turn the engine
 }
 
 Chaos::ISimulationModuleBase* UVehicleSimEngineComponent::CreateNewCoreModule() const
@@ -38,7 +38,7 @@ Chaos::ISimulationModuleBase* UVehicleSimEngineComponent::CreateNewCoreModule()
 	Settings.MaxRPM = MaxRPM;
 	Settings.IdleRPM = EngineIdleRPM;
 	Settings.EngineBrakeEffect = EngineBrakeEffect;
-	Settings.EngineInertia = MomentOfInertia;
+	Settings.EngineInertia = EngineInertia;
 
 	Chaos::ISimulationModuleBase* Engine = new Chaos::FEngineSimModule(Settings);
 	return Engine;
diff --git a/Engine/Plugins/Experimental/ChaosModularVehicle/Source/ChaosModularVehicleEngine/Public/ChaosModularVehicle/VehicleSimEngineComponent.h b/Engine/Plugins/Experimental/ChaosModularVehicle/Source/ChaosModularVehicleEngine/Public/ChaosModularVehicle/VehicleSimEngineComponent.h
--- a/Engine/Plugins/Experimental/ChaosModularVehicle/Source/ChaosModularVehicleEngine/Public/ChaosModularVehicle/VehicleSimEngineComponent.h
+++ b/Engine/Plugins/Experimental/ChaosModularVehicle/Source/ChaosModularVehicleEngine/Public/ChaosModularVehicle/VehicleSimEngineComponent.h
@@ -34,7 +34,7 @@ public:
 	float EngineBrakeEffect;	// [0..1] How much the engine slows the vehicle when the throttle is released
 
 	UPROPERTY(EditAnywhere, BlueprintReadWrite, Category = Attributes)
-	float MomentOfInertia;		// [Kg.m-2] How hard it is to turn the engine
+	float EngineInertia;		// [Kg.m-2] How hard it is to turn the engine
 
 	virtual ESimModuleType GetModuleType() const override { return ESimModuleType::Engine; }
 
diff --git a/Engine/Source/Programs/HeadlessChaos/Private/HeadlessChaosTestModularVehicles.cpp b/Engine/Source/Programs/HeadlessChaos/Private/HeadlessChaosTestModularVehicles.cpp
new file mode 100644
--- /dev/null
+++ b/Engine/Source/Programs/HeadlessChaos/Private/HeadlessChaosTestModularVehicles.cpp
@@ -0,0 +1,1032 @@
+// Copyright Epic Games, Inc. All Rights Reserved.
+
+#include "HeadlessChaos.h"
+#include "HeadlessChaosTestUtility.h"
+
+// for Module Unit Tests
+#include "SimModule/ChassisModule.h"
+#include "SimModule/AerofoilModule.h"
+#include "SimModule/EngineModule.h"
+#include "SimModule/ClutchModule.h"
+#include "SimModule/TransmissionModule.h"
+#include "SimModule/WheelModule.h"
+#include "SimModule/SuspensionModule.h"
+#include "SimModule/SimModuleTree.h"
+
+// for Simulation Tests
+#include "Chaos/PBDRigidsEvolutionGBF.h"
+#include "Chaos/Box.h"
+#include "Chaos/Sphere.h"
+#include "Chaos/Utilities.h"
+
+//////////////////////////////////////////////////////////////////////////
+// These tests are mostly working in real word units rather than Unreal 
+// units as it's easier to tell if the simulations are working close to 
+// reality. i.e. Google stopping distance @ 30MPH ==> typically 15 metres
+//////////////////////////////////////////////////////////////////////////
+
+namespace ChaosTest
+{
+	using namespace Chaos;
+
+	GTEST_TEST(AllTraits, ModularVehicleTest_Aerofoil)
+	{
+		FAerofoilSettings RWingSetup;
+		RWingSetup.Offset.Set(-0.8f, 3.0f, 0.0f);
+		RWingSetup.ForceAxis.Set(0.0f, 0.f, 1.0f);
+		RWingSetup.ControlRotationAxis.Set(0.f, 1.f, 0.f);
+		RWingSetup.Area = 8.2f;
+		RWingSetup.Camber = 3.0f;
+		RWingSetup.MaxControlAngle = 1.0f;
+		RWingSetup.StallAngle = 16.0f;
+		RWingSetup.Type = EAerofoil::Wing;
+
+		FAerofoilSimModule RWing(RWingSetup);
+
+		RWing.SetControlSurface(0.0f);
+		RWing.SetDensityOfMedium(RealWorldConsts::AirDensity());
+
+		float Altitude = 100.0f;
+		float DeltaTime = 1.0f / 30.0f;
+
+		//////////////////////////////////////////////////////////////////////////
+
+		FVector Velocity(1.0f, 0.0f, 0.0f);
+
+		float AOAFlat = RWing.CalcAngleOfAttackDegrees(FVector(0, 0, 1), FVector(-1, 0, 0));
+		EXPECT_LT(AOAFlat, SMALL_NUMBER);
+
+		float AOAFlat2 = RWing.CalcAngleOfAttackDegrees(FVector(0, 0, 1), FVector(1, 0, 0));
+		EXPECT_LT(AOAFlat2, SMALL_NUMBER);
+
+		float AOA90 = RWing.CalcAngleOfAttackDegrees(FVector(0, 0, 1), FVector(0, 0, 1));
+		EXPECT_LT(AOA90 - 90.0f, SMALL_NUMBER);
+
+		float AOA45 = RWing.CalcAngleOfAttackDegrees(FVector(0, 0, 1), FVector(0, 0.707, 0.707));
+		EXPECT_LT(AOA45 - 45.0f, SMALL_NUMBER);
+
+		//////////////////////////////////////////////////////////////////////////
+		// Lift
+		{
+			float Zero = RWing.CalcLiftCoefficient(0, 0);
+			EXPECT_LT(Zero, SMALL_NUMBER);
+
+			float Two = RWing.CalcLiftCoefficient(2, 0);
+			float NegTwo = RWing.CalcLiftCoefficient(-2, 0);
+			EXPECT_GT(Two, SMALL_NUMBER);
+			EXPECT_LT(NegTwo, SMALL_NUMBER);
+			EXPECT_LT(Two - FMath::Abs(NegTwo), SMALL_NUMBER);
+
+			float Three = RWing.CalcLiftCoefficient(0, 3);
+			float NegThree = RWing.CalcLiftCoefficient(0, -3);
+			EXPECT_GT(Three, SMALL_NUMBER);
+			EXPECT_LT(NegThree, SMALL_NUMBER);
+			EXPECT_LT(Three - FMath::Abs(NegThree), SMALL_NUMBER);
+
+			float Nine = RWing.CalcLiftCoefficient(6, 3);
+			float NegNine = RWing.CalcLiftCoefficient(-6, -3);
+			EXPECT_GT(Nine, SMALL_NUMBER);
+			EXPECT_LT(NegNine, SMALL_NUMBER);
+			EXPECT_LT(Nine - FMath::Abs(NegNine), SMALL_NUMBER);
+
+			float Stall = RWing.CalcLiftCoefficient(RWingSetup.StallAngle, 0);
+			float StallPlus = RWing.CalcLiftCoefficient(RWingSetup.StallAngle, 5);
+			EXPECT_GT(Stall, Nine);
+			EXPECT_GT(Stall, Three);
+			EXPECT_GT(Stall, Two);
+			EXPECT_GT(Stall, StallPlus);
+		}
+
+		// Drag
+		{
+			float Two = RWing.CalcDragCoefficient(2, 0);
+			float NegTwo = RWing.CalcDragCoefficient(-2, 0);
+			EXPECT_GT(Two, SMALL_NUMBER);
+			EXPECT_GT(NegTwo, SMALL_NUMBER);
+			EXPECT_LT(Two - NegTwo, SMALL_NUMBER);
+
+			float Six = RWing.CalcDragCoefficient(4, 2);
+			float NegSix = RWing.CalcDragCoefficient(-4, -2);
+			EXPECT_GT(Six, SMALL_NUMBER);
+			EXPECT_GT(NegSix, SMALL_NUMBER);
+			EXPECT_LT(Six - NegSix, SMALL_NUMBER);
+
+			float AltNegTwo = RWing.CalcDragCoefficient(2, -4);
+			EXPECT_GT(AltNegTwo, SMALL_NUMBER);
+			EXPECT_LT(AltNegTwo - NegTwo, SMALL_NUMBER);
+		}
+
+
+		////////////////////////////////////////////////////////////////////////////
+
+		FVector Velocity1(0.0f, 0.0f, 10.0f);
+		FVector RWForceZero = RWing.GetForce(Velocity1, Altitude, DeltaTime);
+		EXPECT_LT(FMath::Abs(RWForceZero.X), SMALL_NUMBER);
+		EXPECT_LT(FMath::Abs(RWForceZero.Y), SMALL_NUMBER);
+		EXPECT_LT(RWForceZero.Z, 0.f); // drag value opposes velocity direction
+
+		FVector Velocity2(0.0f, 10.0f, 10.0f);
+		FVector RWForce3 = RWing.GetForce(Velocity2, Altitude, DeltaTime);
+		EXPECT_LT(FMath::Abs(RWForce3.X), SMALL_NUMBER);
+		EXPECT_LT(RWForce3.Y, 0.0f);
+		EXPECT_LT(RWForce3.Z, 0.0f);
+
+		FVector Velocity3(10.0f, 0.0f, 0.0f);
+		FVector RWForce4 = RWing.GetForce(Velocity3, Altitude, DeltaTime);
+		EXPECT_LT(RWForce4.X, 0.0f);
+		EXPECT_LT(FMath::Abs(RWForce4.Y), SMALL_NUMBER);
+		EXPECT_GT(RWForce4.Z, 0.0f);
+
+	}
+
+
+	// Transmission
+
+	class FTransmissionTestClass : public FTransmissionSimModule
+	{
+		public:
+
+		FTransmissionTestClass(const FTransmissionSettings& Settings)
+		: FTransmissionSimModule(Settings)
+		{}
+
+		void Test_TransmissionManualGearSelection()
+		{
+			FAllInputs Inputs;
+			FSimModuleTree Tree;
+
+			EXPECT_EQ(GetCurrentGear(), 1);
+
+			// Immediate Gear Change, since Setup.GearChangeTime = 0.0f
+			ChangeUp();
+
+			EXPECT_EQ(GetCurrentGear(), 2);
+			ChangeUp();
+			ChangeUp();
+			ChangeUp();
+			EXPECT_EQ(GetCurrentGear(), 5);
+
+			ChangeUp();
+			EXPECT_EQ(GetCurrentGear(), 5);
+
+			SetGear(1);
+			EXPECT_EQ(GetCurrentGear(), 1);
+
+			ChangeDown();
+			EXPECT_EQ(GetCurrentGear(), 0);
+
+			ChangeDown();
+			EXPECT_EQ(GetCurrentGear(), -1);
+
+			ChangeDown();
+			EXPECT_EQ(GetCurrentGear(), -2);
+
+			ChangeDown();
+			EXPECT_EQ(GetCurrentGear(), -2);
+
+			ChangeUp();
+			EXPECT_EQ(GetCurrentGear(), -1);
+
+			ChangeUp();
+			EXPECT_EQ(GetCurrentGear(), 0);
+
+			SetGear(1);
+
+			// Now change settings so we have a delay in the gear changing
+			AccessSetup().GearChangeTime = 0.5f;
+
+			ChangeUp();
+			EXPECT_EQ(GetCurrentGear(), 0);
+			Simulate(0.25f, Inputs, Tree);
+			EXPECT_EQ(GetCurrentGear(), 0);
+			Simulate(0.25f, Inputs, Tree);
+			EXPECT_EQ(GetCurrentGear(), 2);
+			Simulate(0.25f, Inputs, Tree);
+			EXPECT_EQ(GetCurrentGear(), 2);
+
+			SetGear(4);
+			EXPECT_EQ(GetCurrentGear(), 0);
+			Simulate(0.25f, Inputs, Tree);
+			EXPECT_EQ(GetCurrentGear(), 0);
+			Simulate(0.25f, Inputs, Tree);
+			EXPECT_EQ(GetCurrentGear(), 4);
+		}
+
+		void Test_TransmissionAutoGearSelection()
+		{
+			FAllInputs Inputs;
+			FSimModuleTree Tree;
+
+			SetGear(1, true);
+
+			SetRPM(1400);
+			Simulate(0.25f, Inputs, Tree);
+			EXPECT_EQ(GetCurrentGear(), 1);
+
+			SetRPM(2000);
+			Simulate(0.25f, Inputs, Tree);
+			EXPECT_EQ(GetCurrentGear(), 1);
+
+			SetRPM(3000);
+			Simulate(0.25f, Inputs, Tree);
+			EXPECT_EQ(GetCurrentGear(), 2);
+
+			SetRPM(2000);
+			Simulate(0.25f, Inputs, Tree);
+			EXPECT_EQ(GetCurrentGear(), 2);
+
+			SetRPM(1000);
+			Simulate(0.25f, Inputs, Tree);
+			EXPECT_EQ(GetCurrentGear(), 1);
+
+			// stays in first, doesn't change to neutral
+			SetRPM(1000);
+			Simulate(0.25f, Inputs, Tree);
+			EXPECT_EQ(GetCurrentGear(), 1);
+
+			SetGear(-2, true);
+
+			SetRPM(3000);
+			Simulate(0.25f, Inputs, Tree);
+			EXPECT_EQ(GetCurrentGear(), -2);
+
+			SetRPM(1000);
+			Simulate(0.25f, Inputs, Tree);
+			EXPECT_EQ(GetCurrentGear(), -1);
+
+			// stays in reverse first, doesn't change to neutral
+			SetRPM(1000);
+			Simulate(0.25f, Inputs, Tree);
+			EXPECT_EQ(GetCurrentGear(), -1);
+
+			// changes to next reverse gear
+			SetRPM(3000);
+			Simulate(0.25f, Inputs, Tree);
+			EXPECT_EQ(GetCurrentGear(), -2);
+		}
+
+		void Test_TransmissionGearRatios()
+		{
+			float Ratio = 0;
+			Ratio = GetGearRatio(-1);
+			EXPECT_LT(-12.f - Ratio, SMALL_NUMBER); // -ve output for reverse gears
+
+			Ratio = GetGearRatio(0);
+			EXPECT_LT(Ratio, SMALL_NUMBER);
+
+			Ratio = GetGearRatio(1);
+			EXPECT_LT(16.f - Ratio, SMALL_NUMBER);
+
+			Ratio = GetGearRatio(2);
+			EXPECT_LT(12.f - Ratio, SMALL_NUMBER);
+
+			Ratio = GetGearRatio(3);
+			EXPECT_LT(8.f - Ratio, SMALL_NUMBER);
+
+			Ratio = GetGearRatio(4);
+			EXPECT_LT(4.f - Ratio, SMALL_NUMBER);
+		}
+	};
+
+	GTEST_TEST(AllTraits, ModularVehicleTest_TransmissionManualGearSelection)
+	{
+		// done this way so we can access protected function calls
+
+		FTransmissionSettings Setup;
+		{
+			Setup.ForwardRatios.Empty();
+			Setup.ReverseRatios.Empty();
+			Setup.ForwardRatios.Add(4.f);
+			Setup.ForwardRatios.Add(3.f);
+			Setup.ForwardRatios.Add(2.f);
+			Setup.ForwardRatios.Add(1.f);
+			Setup.ForwardRatios.Add(0.8f);
+			Setup.ReverseRatios.Add(3.f);
+			Setup.ReverseRatios.Add(6.f);
+			Setup.FinalDriveRatio = 4.f;
+			Setup.ChangeUpRPM = 3000;
+			Setup.ChangeDownRPM = 1200;
+			Setup.GearChangeTime = 0.0f;
+			Setup.TransmissionType = FTransmissionSettings::ETransType::ManualType;
+			Setup.AutoReverse = true;
+			Setup.TransmissionEfficiency = 1.0f;
+		}
+
+		FTransmissionTestClass Transmission(Setup);
+		Transmission.Test_TransmissionManualGearSelection();
+	}
+
+	GTEST_TEST(AllTraits, ModularVehicleTest_TransmissionAutoGearSelection)
+	{
+		FTransmissionSettings Setup;
+		{
+			Setup.ForwardRatios.Empty();
+			Setup.ReverseRatios.Empty();
+			Setup.ForwardRatios.Add(4.f);
+			Setup.ForwardRatios.Add(3.f);
+			Setup.ForwardRatios.Add(2.f);
+			Setup.ForwardRatios.Add(1.f);
+			Setup.ReverseRatios.Add(3.f);
+			Setup.ReverseRatios.Add(6.f);
+			Setup.FinalDriveRatio = 4.f;
+			Setup.ChangeUpRPM = 3000;
+			Setup.ChangeDownRPM = 1200;
+			Setup.GearChangeTime = 0.0f;
+			Setup.TransmissionType = FTransmissionSettings::ETransType::AutomaticType;
+			Setup.AutoReverse = false;
+			Setup.TransmissionEfficiency = 1.0f;
+		}
+
+		FTransmissionTestClass Transmission(Setup);
+		Transmission.Test_TransmissionAutoGearSelection();
+	}
+
+	GTEST_TEST(AllTraits, ModularVehicleTest_TransmissionGearRatios)
+	{
+		FTransmissionSettings Setup;
+		{
+			Setup.ForwardRatios.Empty();
+			Setup.ReverseRatios.Empty();
+			Setup.ForwardRatios.Add(4.f);
+			Setup.ForwardRatios.Add(3.f);
+			Setup.ForwardRatios.Add(2.f);
+			Setup.ForwardRatios.Add(1.f);
+			Setup.ReverseRatios.Add(3.f);
+			Setup.FinalDriveRatio = 4.f;
+			Setup.ChangeUpRPM = 3000;
+			Setup.ChangeDownRPM = 1200;
+			Setup.GearChangeTime = 0.0f;
+			Setup.TransmissionType = FTransmissionSettings::ETransType::AutomaticType;
+			Setup.AutoReverse = true;
+			Setup.TransmissionEfficiency = 1.0f;
+		}
+
+		FTransmissionTestClass Transmission(Setup);
+		Transmission.Test_TransmissionGearRatios();
+	}
+
+	// Wheel
+	void SimulateBraking(FWheelSimModule& Wheel
+		, const float Gravity
+		, float VehicleSpeed
+		, float DeltaTime
+		, float& StoppingDistanceOut
+		, float& SimulationTimeOut
+		, bool bLoggingEnabled = false
+		)
+	{
+		FAllInputs Inputs;
+		Inputs.ControlInputs.Brake = 1.0f; // apply full braking force
+		Inputs.ControlInputs.Throttle = 0.0f;
+		FSimModuleTree Tree;
+
+		StoppingDistanceOut = 0.f;
+		SimulationTimeOut = 0.f;
+		
+		float MaxSimTime = 15.0f;
+		float VehicleMass = 1300.f;
+		float VehicleMassPerWheel = 1300.f / 4.f;
+
+		Wheel.SetForceIntoSurface(VehicleMassPerWheel * Gravity);
+
+		// Road speed
+		FVector Velocity = FVector(VehicleSpeed, 0.f, 0.f);
+
+		// wheel rolling speed matches road speed
+		Wheel.SetLinearSpeed(Velocity.X);
+
+		if (bLoggingEnabled)
+		{
+			UE_LOG(LogChaos, Warning, TEXT("--------------------START---------------------"));
+		}
+
+		while (SimulationTimeOut < MaxSimTime)
+		{
+			// rolling speed matches road speed
+			Wheel.SetLocalLinearVelocity(Velocity);
+			Wheel.Simulate(DeltaTime, Inputs, Tree);
+
+			// deceleration from brake, F = m * a, a = F / m, v = dt * F / m
+			Velocity += DeltaTime * Wheel.GetForceFromFriction() / VehicleMassPerWheel;
+			StoppingDistanceOut += Velocity.X * DeltaTime;
+
+			if (bLoggingEnabled)
+			{
+				UE_LOG(LogChaos, Warning, TEXT("Wheel.GetForceFromFriction() %s"), *Wheel.GetForceFromFriction().ToString());
+			}
+
+			if (FMath::Abs(Velocity.X) < 0.05f)
+			{
+				Velocity.X = 0.f;
+				break; // break out early if already stopped
+			}
+
+			SimulationTimeOut += DeltaTime;
+		}
+
+		if (bLoggingEnabled)
+		{
+			UE_LOG(LogChaos, Warning, TEXT("---------------------END----------------------"));
+		}
+
+
+	}
+
+	void SimulateAccelerating(FWheelSimModule& Wheel
+		, const float Gravity
+		, const float DriveTorque
+		, float FinalVehicleSpeed
+		, float DeltaTime
+		, float& DistanceTravelledOut
+		, float& SimulationTimeOut
+		)
+	{
+		FAllInputs Inputs;
+		Inputs.ControlInputs.Throttle = 1.0f; // apply full throttle
+		Inputs.ControlInputs.Brake = 0.0f;
+		FSimModuleTree Tree;
+
+		DistanceTravelledOut = 0.f;
+		SimulationTimeOut = 0.f;
+
+		float MaxSimTime = 15.0f;
+		float VehicleMass = 1300.f;
+		float VehicleMassPerWheel = VehicleMass / 4.f;
+
+		Wheel.SetForceIntoSurface(VehicleMassPerWheel * Gravity);
+
+		// Road speed
+		FVector Velocity = FVector(0.f, 0.f, 0.f);
+
+		// start from stationary
+		Wheel.SetLinearSpeed(Velocity.X);
+
+		while (SimulationTimeOut < MaxSimTime)
+		{
+			Wheel.SetLocalLinearVelocity(Velocity);
+			Wheel.SetDriveTorque(DriveTorque);
+			Wheel.Simulate(DeltaTime, Inputs, Tree);
+
+			Velocity += DeltaTime * Wheel.GetForceFromFriction() / VehicleMassPerWheel;
+			DistanceTravelledOut += Velocity.X * DeltaTime;
+
+			SimulationTimeOut += DeltaTime;
+
+			if (FMath::Abs(Velocity.X) >= FinalVehicleSpeed)
+			{
+				break; // time is up
+			}
+
+		}
+	}
+
+	GTEST_TEST(AllTraits, ModularVehicleTest_WheelBrakingLongitudinalSlip)
+	{
+		FWheelSettings Setup;
+		Setup.ABSEnabled = false;
+		Setup.TractionControlEnabled = false;
+		Setup.SteeringEnabled = true;
+		Setup.HandbrakeEnabled = true;
+		Setup.Radius = 0.3f;
+		Setup.FrictionMultiplier = 1.0f;
+		Setup.CorneringStiffness = 1000.0f;
+		Setup.LateralSlipGraphMultiplier = 0.7f;
+
+		FWheelSimModule Wheel(Setup);
+
+		// Google braking distance at 30mph says 14m (not interested in the thinking distance part)
+		// So using a range 10-20 to ensure we are in the correct ballpark.
+		// If specified more accurately in the test, then modifying the code would break the test all the time.
+
+		// units meters
+		float Gravity = 9.8f;
+		float StoppingDistanceTolerance = 0.5f;
+		float DeltaTime = 1.f / 30.f;
+		float StoppingDistanceA = 0.f;
+		float SimulationTime = 0.0f;
+		Wheel.SetSurfaceFriction(RealWorldConsts::DryRoadFriction());
+
+		// reasonably ideal stopping distance - traveling forwards
+		Wheel.AccessSetup().MaxBrakeTorque = 650.0f;
+ 		SimulateBraking(Wheel, Gravity, MPHToMS(30.f), DeltaTime, StoppingDistanceA, SimulationTime);
+		EXPECT_GT(StoppingDistanceA, 10.f);
+		EXPECT_LT(StoppingDistanceA, 20.f);
+
+		// traveling backwards stops just the same
+		float StoppingDistanceReverseDir = 0.f;
+		Wheel.AccessSetup().MaxBrakeTorque = 650.0f;
+		SimulateBraking(Wheel, Gravity, MPHToMS(-30.f), DeltaTime, StoppingDistanceReverseDir, SimulationTime);
+		EXPECT_GT(StoppingDistanceReverseDir, -20.f);
+		EXPECT_LT(StoppingDistanceReverseDir, -10.f);
+		EXPECT_LT(StoppingDistanceA - FMath::Abs(StoppingDistanceReverseDir), StoppingDistanceTolerance);
+
+		// Changing to units of Cm should yield the same results
+		float MToCm = 100.0f;
+		float StoppingDistanceCm = 0.f;
+		Wheel.AccessSetup().MaxBrakeTorque = 650.0f * MToCm * MToCm;
+		Wheel.AccessSetup().Radius = (0.3f * MToCm);
+		SimulateBraking(Wheel, Gravity * MToCm, MPHToCmS(30.f), DeltaTime, StoppingDistanceCm, SimulationTime);
+		EXPECT_NEAR(StoppingDistanceCm, StoppingDistanceA * MToCm, 1.0f);
+
+		// Similar results with different delta time
+		float StoppingDistanceDiffDT = 0.f;
+		Wheel.AccessSetup().MaxBrakeTorque = 650.0f;
+		Wheel.AccessSetup().Radius = 0.3f;
+		SimulateBraking(Wheel, Gravity, MPHToMS(30.f), DeltaTime * 0.25f, StoppingDistanceDiffDT, SimulationTime);
+		EXPECT_LT(StoppingDistanceA - StoppingDistanceDiffDT, StoppingDistanceTolerance);
+
+		// barely touching the brake - going to take longer to stop
+		float StoppingDistanceLightBraking = 0.f;
+		Wheel.AccessSetup().MaxBrakeTorque = 150.0f;
+		SimulateBraking(Wheel, Gravity, MPHToMS(30.f), DeltaTime, StoppingDistanceLightBraking, SimulationTime);
+		EXPECT_GT(StoppingDistanceLightBraking, StoppingDistanceA);
+
+		// locking the wheels / too much brake torque -> dynamic friction rather than static friction -> going to take longer to stop
+		float StoppingDistanceTooHeavyBreaking = 0.f;
+		Wheel.AccessSetup().MaxBrakeTorque = 5000.0f;
+		SimulateBraking(Wheel, Gravity, MPHToMS(30.f), DeltaTime, StoppingDistanceTooHeavyBreaking, SimulationTime);
+		EXPECT_GT(StoppingDistanceTooHeavyBreaking, StoppingDistanceA);
+
+		// Would have locked the wheels but ABS prevents skidding
+		Wheel.AccessSetup().ABSEnabled = true;
+		float StoppingDistanceTooHeavyBreakingABS = 0.f;
+		Wheel.AccessSetup().MaxBrakeTorque = 5000.0f;
+		SimulateBraking(Wheel, Gravity, MPHToMS(30.f), DeltaTime, StoppingDistanceTooHeavyBreakingABS, SimulationTime);
+		EXPECT_LT(StoppingDistanceTooHeavyBreakingABS, StoppingDistanceTooHeavyBreaking);
+		Wheel.AccessSetup().ABSEnabled = false;
+
+		// lower initial speed - stops more quickly
+		float StoppingDistanceLowerSpeed = 0.f;
+		Wheel.AccessSetup().MaxBrakeTorque = 650.0f;
+		SimulateBraking(Wheel, Gravity, MPHToMS(20.f), DeltaTime, StoppingDistanceLowerSpeed, SimulationTime);
+		EXPECT_LT(StoppingDistanceLowerSpeed, StoppingDistanceA);
+
+		// higher initial speed - stops more slowly
+		float StoppingDistanceHigherSpeed = 0.f;
+		Wheel.AccessSetup().MaxBrakeTorque = 650.0f;
+		SimulateBraking(Wheel, Gravity, MPHToMS(60.f), DeltaTime, StoppingDistanceHigherSpeed, SimulationTime);
+		EXPECT_GT(StoppingDistanceHigherSpeed, StoppingDistanceA);
+
+		// slippy surface - stops more slowly
+		float StoppingDistanceLowFriction = 0.f;
+		Wheel.SetSurfaceFriction(0.3f);
+		Wheel.AccessSetup().MaxBrakeTorque = 650.0f;
+		SimulateBraking(Wheel, Gravity, MPHToMS(30.f), DeltaTime, StoppingDistanceLowFriction, SimulationTime);
+		EXPECT_GT(StoppingDistanceLowFriction, StoppingDistanceA);
+	}
+
+	GTEST_TEST(AllTraits, ModularVehicleTest_WheelAcceleratingLongitudinalSlip)
+	{
+		FWheelSettings Setup;
+		Setup.ABSEnabled = false;
+		Setup.TractionControlEnabled = false;
+		Setup.SteeringEnabled = false;
+		Setup.HandbrakeEnabled = false;
+		Setup.Radius = 0.3f;
+		Setup.FrictionMultiplier = 1.0f;
+		Setup.CorneringStiffness = 1000.0f;
+		Setup.LateralSlipGraphMultiplier = 0.7f;
+
+		FWheelSimModule Wheel(Setup);
+
+		// There could be one frame extra computation on the acceleration since the last frame of brake is not using the full 
+		// amount of torque, it's clearing the last remaining velocity without pushing the vehicle back in the opposite direction
+		// Hence a slightly larger tolerance for the result
+		float AccelerationResultsTolerance = 1.0f; // meters
+
+		// units meters
+		float Gravity = 9.8f;
+		float DeltaTime = 1.f / 30.f;
+		float DriveTorque = 0.0;
+
+		float StoppingDistanceA = 0.f;
+		float SimulationTimeBrake = 0.0f;
+		Wheel.SetSurfaceFriction(RealWorldConsts::DryRoadFriction());
+
+		// How far & what time does it take to stop from 30MPH to rest
+		Wheel.AccessSetup().MaxBrakeTorque = 650.0f;
+		SimulateBraking(Wheel, Gravity, MPHToMS(30.0f), DeltaTime, StoppingDistanceA, SimulationTimeBrake);
+
+		// How far and what time does it take to accelerate from rest to 30MPH
+		float SimulationTimeAccel = 0.0f;
+		float DrivingDistanceA = 0.f;
+		DriveTorque = 650.0f;
+		SimulateAccelerating(Wheel, Gravity, DriveTorque, MPHToMS(30.0f), DeltaTime, DrivingDistanceA, SimulationTimeAccel);
+
+		// 0-30 MPH and 30-0 MPH should be the same if there's no slipping and accel torque was same as the brake torque run
+		EXPECT_LT(DrivingDistanceA - StoppingDistanceA, AccelerationResultsTolerance);
+		EXPECT_LT(SimulationTimeAccel - SimulationTimeBrake, AccelerationResultsTolerance);
+
+		// same range as braking from 30MPH
+		EXPECT_GT(DrivingDistanceA, 10.f);
+		EXPECT_LT(DrivingDistanceA, 20.f);
+
+		// Unreal units cm - Note for the same results the radius needs to remain at 0.3m and not also be scaled to 30(cm)
+		float SimulationTimeAccelCM = 0.0f;
+		float MToCm = 100.0f;
+		float DrivingDistanceCM = 0.f;
+		DriveTorque = 650.0f * MToCm * MToCm;
+		Wheel.AccessSetup().Radius = (0.3f * MToCm);
+		SimulateAccelerating(Wheel, Gravity * MToCm, DriveTorque, MPHToCmS(30.0f), DeltaTime, DrivingDistanceCM, SimulationTimeAccelCM);
+		EXPECT_GT(DrivingDistanceCM, 10.f * MToCm);
+		EXPECT_LT(DrivingDistanceCM, 20.f * MToCm);
+		EXPECT_NEAR(DrivingDistanceCM, DrivingDistanceA * MToCm, AccelerationResultsTolerance);
+
+		float SimulationTimeAccelSpin = 0.0f;
+		float DrivingDistanceWheelspin = 0.f;
+		Wheel.AccessSetup().Radius = 0.3f;
+		DriveTorque = 5000; // definitely cause wheel spin
+		SimulateAccelerating(Wheel, Gravity, DriveTorque, 30.0f, DeltaTime, DrivingDistanceWheelspin, SimulationTimeAccelSpin);
+
+		// drives further to reach the same speed
+		EXPECT_GT(DrivingDistanceWheelspin, DrivingDistanceA);
+
+		// takes longer to reach the same speed
+		EXPECT_GT(SimulationTimeAccelSpin, SimulationTimeAccel);
+
+		// Enable traction control should be better than both of the above
+		float SimulationTimeAccelTC = 0.0f;
+		float DrivingDistanceTC = 0.f;
+		Wheel.AccessSetup().TractionControlEnabled = true;
+		DriveTorque = 5000; // definitely cause wheel spin
+		SimulateAccelerating(Wheel, Gravity, DriveTorque, MPHToMS(30.0f), DeltaTime, DrivingDistanceTC, SimulationTimeAccelTC);
+
+		// reaches target speed in a shorter distance
+		EXPECT_LT(DrivingDistanceTC, DrivingDistanceWheelspin);
+
+		// reaches speed quicker with TC on when wheel would be slipping from drive torque
+		EXPECT_LT(SimulationTimeAccelTC, SimulationTimeAccelSpin);
+
+	}
+
+	GTEST_TEST(AllTraits, ModularVehicleTest_WheelRolling)
+	{
+		FWheelSettings Setup;
+		Setup.Radius = 0.3f;
+		FWheelSimModule Wheel(Setup);
+
+		FAllInputs Inputs;
+		Inputs.ControlInputs.Throttle = 0.0f;
+		Inputs.ControlInputs.Brake = 0.0f;
+		FSimModuleTree Tree;
+
+		float DeltaTime = 1.f / 30.f;
+		float MaxSimTime = 10.0f;
+		float Tolerance = 0.1f; // wheel friction losses slow wheel speed
+
+		//------------------------------------------------------------------
+		// Car is moving FORWARDS - with AMPLE friction we would expect an initially 
+		// static rolling wheel to speed up and match the vehicle speed
+		FVector VehicleGroundSpeed(10.0f, 0.0f, 0.0f);
+		Wheel.SetSurfaceFriction(1.0f);	// Some wheel/ground friction
+		Wheel.SetForceIntoSurface(250.f); // wheel pressed into the ground, to give it grip
+		Wheel.SetAngularVelocity(0.f);
+		Wheel.SetLocalLinearVelocity(VehicleGroundSpeed);
+
+		// initially wheel is static
+		EXPECT_LT(Wheel.GetAngularVelocity(), SMALL_NUMBER);
+
+		// after some time, the wheel picks up speed to match the vehicle speed
+		float SimulatedTime = 0.f;
+		while (SimulatedTime < MaxSimTime)
+		{
+			Wheel.Simulate(DeltaTime, Inputs, Tree);
+			SimulatedTime += DeltaTime;
+		}
+
+		// there's enough grip to cause the wheel to spin and match the vehivle speed
+		float WheelGroundSpeed = Wheel.GetAngularVelocity() * Wheel.GetEffectiveRadius();
+		EXPECT_LT(VehicleGroundSpeed.X - WheelGroundSpeed, Tolerance);
+		EXPECT_LT(VehicleGroundSpeed.X - WheelGroundSpeed, Tolerance);
+		EXPECT_GT(Wheel.GetAngularVelocity(), 0.f); // +ve spin on it
+
+		//------------------------------------------------------------------
+		// Car is moving BACKWARDS - with AMPLE friction we would expect an initially 
+		// static rolling wheel to speed up and match the vehicle speed
+		VehicleGroundSpeed.Set(-10.0f, 0.0f, 0.0f);
+		Wheel.SetSurfaceFriction(1.0f);	// Some wheel/ground friction
+		Wheel.SetForceIntoSurface(250.f); // wheel pressed into the ground, to give it grip
+		Wheel.SetAngularVelocity(0.f);
+		Wheel.SetLocalLinearVelocity(VehicleGroundSpeed);
+
+		// initially wheel is static
+		EXPECT_LT(Wheel.GetAngularVelocity(), SMALL_NUMBER);
+
+		// after some time, the wheel picks up speed to match the vehicle speed
+		SimulatedTime = 0.f;
+		while (SimulatedTime < MaxSimTime)
+		{
+			Wheel.Simulate(DeltaTime, Inputs, Tree);
+			SimulatedTime += DeltaTime;
+		}
+
+		// there's enough grip to cause the wheel to spin and match the vehicle speed
+		WheelGroundSpeed = Wheel.GetAngularVelocity() * Wheel.GetEffectiveRadius();
+		EXPECT_LT(VehicleGroundSpeed.X - WheelGroundSpeed, Tolerance);
+		EXPECT_LT(VehicleGroundSpeed.X - Wheel.GetLinearSpeed(), Tolerance);
+		EXPECT_LT(Wheel.GetAngularVelocity(), 0.f); // -ve spin on it
+
+		//------------------------------------------------------------------
+		// Car is moving FORWARDS - with NO friction we would expect an initially 
+		// static wheel to NOT speed up to match the vehicle speed
+		Wheel.SetSurfaceFriction(0.0f);	// No wheel/ground friction
+		Wheel.SetForceIntoSurface(250.f); // wheel pressed into the ground, to give it grip
+		Wheel.SetAngularVelocity(0.f);
+		Wheel.SetLocalLinearVelocity(VehicleGroundSpeed);
+
+		// initially wheel is static
+		EXPECT_LT(Wheel.GetAngularVelocity(), SMALL_NUMBER);
+
+		// after some time, the wheel picks up speed to match the vehicle speed
+		SimulatedTime = 0.f;
+		while (SimulatedTime < MaxSimTime)
+		{
+			Wheel.Simulate(DeltaTime, Inputs, Tree);
+			SimulatedTime += DeltaTime;
+		}
+
+		WheelGroundSpeed = Wheel.GetAngularVelocity() * Wheel.GetEffectiveRadius();
+
+		// wheel is just sliding there's no friction to make it spin
+		EXPECT_LT(WheelGroundSpeed, SMALL_NUMBER);
+
+	}
+
+	// Engine
+	GTEST_TEST(AllTraits, ModularVehicleTest_EngineRPM)
+	{
+		FAllInputs Inputs;
+		FSimModuleTree Tree;
+
+		FEngineSettings Setup;
+		{
+			Setup.MaxRPM = 5000;
+			Setup.IdleRPM = 1000;
+			Setup.MaxTorque = 400.f;
+			Setup.EngineBrakeEffect = 200.0f;
+			Setup.EngineInertia = 100.0f;
+
+			Setup.TorqueCurve.Empty();
+
+			Setup.TorqueCurve.AddNormalized(0.5f);
+			Setup.TorqueCurve.AddNormalized(0.5f);
+			Setup.TorqueCurve.AddNormalized(0.5f);
+			Setup.TorqueCurve.AddNormalized(0.5f);
+			Setup.TorqueCurve.AddNormalized(0.6f);
+			Setup.TorqueCurve.AddNormalized(0.7f);
+			Setup.TorqueCurve.AddNormalized(0.8f);
+			Setup.TorqueCurve.AddNormalized(0.9f);
+			Setup.TorqueCurve.AddNormalized(1.0f);
+			Setup.TorqueCurve.AddNormalized(0.9f);
+			Setup.TorqueCurve.AddNormalized(0.7f);
+			Setup.TorqueCurve.AddNormalized(0.5f);
+		}
+
+		FEngineSimModule Engine(Setup);
+
+		Inputs.ControlInputs.Throttle = 0.f;
+
+		float DeltaTime = 1.0f / 30.0f;
+		float TOLERANCE = 0.1f;
+
+		// engine idle - no throttle
+		for (int i = 0; i < 200; i++)
+		{
+			Engine.Simulate(DeltaTime, Inputs, Tree);
+		}
+
+		EXPECT_LT(Engine.GetRPM() - Engine.Setup().IdleRPM, TOLERANCE);
+
+		// apply half throttle
+		Inputs.ControlInputs.Throttle = 0.5f;
+
+		for (int i = 0; i < 100; i++)
+		{
+			Engine.Simulate(DeltaTime, Inputs, Tree);
+			//UE_LOG(LogChaos, Warning, TEXT("EngineSpeed %.2f rad/sec (%.1f RPM)"), Engine.GetAngularVelocity(), Engine.GetRPM());
+		}
+
+		EXPECT_GT(Engine.GetRPM(), Engine.Setup().IdleRPM);
+
+		Inputs.ControlInputs.Throttle = 0.0f;
+
+		// engine idle - no throttle
+		for (int i = 0; i < 200; i++)
+		{
+			Engine.Simulate(DeltaTime, Inputs, Tree);
+			//UE_LOG(LogChaos, Warning, TEXT("EngineSpeed %.2f rad/sec (%.1f RPM)"), Engine.GetAngularVelocity(), Engine.GetRPM());
+		}
+
+		EXPECT_LT(Engine.GetRPM() - Engine.Setup().IdleRPM, TOLERANCE);
+
+	}
+
+
+	GTEST_TEST(AllTraits, ModularVehicleTest_SimModuleTree_EngineDrivingWheelsThroughClutch)
+	{
+		FAllInputs Inputs;
+		FSimModuleTree Tree;
+		float TOLERANCE = 0.1f;
+
+		FEngineSettings Setup;
+		{
+			Setup.MaxRPM = 5000;
+			Setup.IdleRPM = 1000;
+			Setup.MaxTorque = 400.f;
+			Setup.EngineBrakeEffect = 20.0f;
+			Setup.EngineInertia = 100.0f;
+
+			Setup.TorqueCurve.Empty();
+
+			Setup.TorqueCurve.AddNormalized(0.5f);
+			Setup.TorqueCurve.AddNormalized(0.5f);
+			Setup.TorqueCurve.AddNormalized(0.5f);
+			Setup.TorqueCurve.AddNormalized(0.5f);
+			Setup.TorqueCurve.AddNormalized(0.6f);
+			Setup.TorqueCurve.AddNormalized(0.7f);
+			Setup.TorqueCurve.AddNormalized(0.8f);
+			Setup.TorqueCurve.AddNormalized(0.9f);
+			Setup.TorqueCurve.AddNormalized(1.0f);
+			Setup.TorqueCurve.AddNormalized(0.9f);
+			Setup.TorqueCurve.AddNormalized(0.7f);
+			Setup.TorqueCurve.AddNormalized(0.5f);
+		}
+
+		FChassisSettings ChassisSettings;
+		int RootNodeIndex = Tree.AddRoot(new FChassisSimModule(ChassisSettings));
+
+		FEngineSettings EngineSettings;
+		int NodeIndex = Tree.AddNodeBelow(RootNodeIndex, new FEngineSimModule(EngineSettings));
+
+		FClutchSettings ClutchSettings;
+		NodeIndex = Tree.AddNodeBelow(NodeIndex, new FClutchSimModule(ClutchSettings));
+
+		FTransmissionSettings TransmissionSettings;
+		int TransmissionNodeIndex = Tree.AddNodeBelow(NodeIndex, new FTransmissionSimModule(TransmissionSettings));
+
+		FWheelSettings WheelSettingsDriven;
+		FWheelSettings WheelSettingsRolling;
+
+		int Wheel0Idx = Tree.AddNodeBelow(TransmissionNodeIndex, new FWheelSimModule(WheelSettingsDriven));
+		int Wheel1Idx = Tree.AddNodeBelow(TransmissionNodeIndex, new FWheelSimModule(WheelSettingsDriven));
+		int Wheel2Idx = Tree.AddNodeBelow(RootNodeIndex, new FWheelSimModule(WheelSettingsRolling));
+		int Wheel3Idx = Tree.AddNodeBelow(RootNodeIndex, new FWheelSimModule(WheelSettingsRolling));
+		
+		const FWheelSimModule& Wheel0 = *static_cast<const FWheelSimModule*>(Tree.GetSimModule(Wheel0Idx));
+		const FWheelSimModule& Wheel1 = *static_cast<const FWheelSimModule*>(Tree.GetSimModule(Wheel1Idx));
+		const FWheelSimModule& Wheel2 = *static_cast<const FWheelSimModule*>(Tree.GetSimModule(Wheel2Idx));
+		const FWheelSimModule& Wheel3 = *static_cast<const FWheelSimModule*>(Tree.GetSimModule(Wheel3Idx));
+
+		FClusterUnionPhysicsProxy* Proxy = nullptr;
+		float DeltaTime = 1.0f / 30.0f;
+
+		// We need to setup reverse pointer to tree
+		for (int I = 0; I < Tree.GetNumNodes(); I++)
+		{
+			Tree.AccessSimModule(I)->SetSimModuleTree(&Tree);
+		}
+
+		// Throttle ON
+		Inputs.ControlInputs.Throttle = 1.0f;
+		Inputs.ControlInputs.Brake = 0.0f;
+		Inputs.ControlInputs.Clutch = 0.0f;
+
+		// wheels not moving initially
+		EXPECT_NEAR(Wheel0.GetRPM(), 0, TOLERANCE);
+		EXPECT_NEAR(Wheel1.GetRPM(), 0, TOLERANCE);
+		EXPECT_NEAR(Wheel2.GetRPM(), 0, TOLERANCE);
+		EXPECT_NEAR(Wheel3.GetRPM(), 0, TOLERANCE);
+		
+		// Simulate
+		for (int I=0; I<50; I++)
+		{
+			Tree.Simulate(DeltaTime, Inputs, Proxy);
+		}
+
+		// driven wheels are turning
+		EXPECT_GT(Wheel0.GetRPM(), TOLERANCE);
+		EXPECT_GT(Wheel1.GetRPM(), TOLERANCE);
+
+		// non-driven wheels are still static
+		EXPECT_NEAR(Wheel2.GetRPM(), 0, TOLERANCE);
+		EXPECT_NEAR(Wheel3.GetRPM(), 0, TOLERANCE);
+
+		// Brake ON 
+		Inputs.ControlInputs.Throttle = 0.0f;
+		Inputs.ControlInputs.Brake = 1.0f;
+		Inputs.ControlInputs.Clutch = 0.0f;
+
+		// Simulate
+		for (int I = 0; I < 50; I++)
+		{
+			Tree.Simulate(DeltaTime, Inputs, Proxy);
+		}
+
+		// all wheels stop spinning
+		EXPECT_NEAR(Wheel0.GetRPM(), 0, TOLERANCE);
+		EXPECT_NEAR(Wheel1.GetRPM(), 0, TOLERANCE);
+		EXPECT_NEAR(Wheel2.GetRPM(), 0, TOLERANCE);
+		EXPECT_NEAR(Wheel3.GetRPM(), 0, TOLERANCE);
+
+
+		// Throttle ON, Clutch depressed
+		Inputs.ControlInputs.Throttle = 1.0f;
+		Inputs.ControlInputs.Brake = 0.0f;
+		Inputs.ControlInputs.Clutch = 1.0f;
+
+		// Simulate
+		for (int I = 0; I < 50; I++)
+		{
+			Tree.Simulate(DeltaTime, Inputs, Proxy);
+		}
+
+		// driven wheels are not turning due to disengaged clutch
+		EXPECT_NEAR(Wheel0.GetRPM(), 0, TOLERANCE);
+		EXPECT_NEAR(Wheel1.GetRPM(), 0, TOLERANCE);
+
+	}
+
+	GTEST_TEST(AllTraits, ModularVehicleTest_SimModuleTree_WheelsSpinEngineShouldSpin)
+	{
+		FAllInputs Inputs;
+		FSimModuleTree Tree;
+		float TOLERANCE = 0.1f;
+
+		FChassisSettings ChassisSettings;
+		int RootNodeIndex = Tree.AddRoot(new FChassisSimModule(ChassisSettings));
+
+		Inputs.ControlInputs.Brake = 0;
+		Inputs.ControlInputs.Throttle = 0;
+
+		FEngineSettings EngineSettings;
+		{
+			EngineSettings.MaxRPM = 5000;
+			EngineSettings.IdleRPM = 0;
+			EngineSettings.MaxTorque = 400.f;
+			EngineSettings.EngineBrakeEffect = 0.0f;
+			EngineSettings.EngineInertia = 100.0f;
+			EngineSettings.TorqueCurve.Empty();
+			EngineSettings.TorqueCurve.AddNormalized(0.5f);
+			EngineSettings.TorqueCurve.AddNormalized(0.5f);
+			EngineSettings.TorqueCurve.AddNormalized(0.5f);
+			EngineSettings.TorqueCurve.AddNormalized(0.5f);
+			EngineSettings.TorqueCurve.AddNormalized(0.6f);
+			EngineSettings.TorqueCurve.AddNormalized(0.7f);
+			EngineSettings.TorqueCurve.AddNormalized(0.8f);
+			EngineSettings.TorqueCurve.AddNormalized(0.9f);
+			EngineSettings.TorqueCurve.AddNormalized(1.0f);
+			EngineSettings.TorqueCurve.AddNormalized(0.9f);
+			EngineSettings.TorqueCurve.AddNormalized(0.7f);
+			EngineSettings.TorqueCurve.AddNormalized(0.5f);
+		}
+		int EngineNodeIndex = Tree.AddNodeBelow(RootNodeIndex, new FEngineSimModule(EngineSettings));
+
+		FWheelSettings WheelSettingsDriven;
+		FWheelSettings WheelSettingsRolling;
+		WheelSettingsDriven.AutoHandbrakeEnabled = false;
+		WheelSettingsRolling.AutoHandbrakeEnabled = false;
+
+		int Wheel0Idx = Tree.AddNodeBelow(EngineNodeIndex, new FWheelSimModule(WheelSettingsDriven));
+		int Wheel1Idx = Tree.AddNodeBelow(EngineNodeIndex, new FWheelSimModule(WheelSettingsDriven));
+
+		FEngineSimModule& Engine = *static_cast<FEngineSimModule*>(Tree.AccessSimModule(EngineNodeIndex));
+		FWheelSimModule& Wheel0 = *static_cast<FWheelSimModule*>(Tree.AccessSimModule(Wheel0Idx));
+		FWheelSimModule& Wheel1 = *static_cast<FWheelSimModule*>(Tree.AccessSimModule(Wheel1Idx));
+
+		FClusterUnionPhysicsProxy* Proxy = nullptr;
+		float DeltaTime = 1.0f / 30.0f;
+
+		// We need to setup reverse pointer to tree
+		for (int I = 0; I < Tree.GetNumNodes(); I++)
+		{
+			Tree.AccessSimModule(I)->SetSimModuleTree(&Tree);
+		}
+
+
+		// Simulate - engine takes speed from wheels spin
+		for (int I = 0; I < 5; I++)
+		{
+			Wheel0.SetLinearSpeed(100);
+			Wheel1.SetLinearSpeed(100);
+			Tree.Simulate(DeltaTime, Inputs, Proxy);
+		}
+
+		EXPECT_NEAR(Engine.GetRPM(), Wheel0.GetRPM(), TOLERANCE);
+
+
+		// Simulate - engine will just need to take average of connected wheels speeds
+		for (int I = 0; I < 5; I++)
+		{
+			Wheel0.SetLinearSpeed(300);
+			Wheel1.SetLinearSpeed(100);
+			Tree.Simulate(DeltaTime, Inputs, Proxy);
+		}
+
+		EXPECT_NEAR(Engine.GetRPM(), (Wheel0.GetRPM() + Wheel1.GetRPM())*0.5f, TOLERANCE);
+
+	}
+
+
+}
+
+
diff --git a/Engine/Source/Runtime/Experimental/ChaosVehicles/ChaosVehiclesCore/Private/SimModule/ClutchModule.cpp b/Engine/Source/Runtime/Experimental/ChaosVehicles/ChaosVehiclesCore/Private/SimModule/ClutchModule.cpp
--- a/Engine/Source/Runtime/Experimental/ChaosVehicles/ChaosVehiclesCore/Private/SimModule/ClutchModule.cpp
+++ b/Engine/Source/Runtime/Experimental/ChaosVehicles/ChaosVehiclesCore/Private/SimModule/ClutchModule.cpp
@@ -34,17 +34,20 @@ namespace Chaos
 		FTorqueSimModule* Parent = static_cast<FTorqueSimModule*>(GetParent());
 		FTorqueSimModule* Child = static_cast<FTorqueSimModule*>(GetFirstChild());
 
-		float EngineSpeed = Parent->GetAngularVelocity();
-		float TransmissionSpeed = Child->GetAngularVelocity();
+		if (Parent && Child)
+		{
+			float EngineSpeed = Parent->GetAngularVelocity();
+			float TransmissionSpeed = Child->GetAngularVelocity();
 
-		// difference in speed between the two plates
-		float AngularVelocityDifference = EngineSpeed - TransmissionSpeed;
+			// difference in speed between the two plates
+			float AngularVelocityDifference = EngineSpeed - TransmissionSpeed;
 
-		Parent->AddAngularVelocity(-AngularVelocityDifference * 0.1f);
-		Child->AddAngularVelocity(AngularVelocityDifference * 0.1f);
+			Parent->AddAngularVelocity(-AngularVelocityDifference * 0.1f);
+			Child->AddAngularVelocity(AngularVelocityDifference * 0.1f);
 
-		float BrakeTorque = 0.0f;
-		TransmitTorque(VehicleModuleSystem, DriveTorque, BrakeTorque, 1.0f, 1.0f/*ClutchValue*/);
+			float BrakeTorque = 0.0f;
+			TransmitTorque(VehicleModuleSystem, DriveTorque, BrakeTorque, 1.0f, ClutchValue);
+		}
 	}
 
 	bool FClutchSimModule::GetDebugString(FString& StringOut) const
diff --git a/Engine/Source/Runtime/Experimental/ChaosVehicles/ChaosVehiclesCore/Private/SimModule/EngineModule.cpp b/Engine/Source/Runtime/Experimental/ChaosVehicles/ChaosVehiclesCore/Private/SimModule/EngineModule.cpp
--- a/Engine/Source/Runtime/Experimental/ChaosVehicles/ChaosVehiclesCore/Private/SimModule/EngineModule.cpp
+++ b/Engine/Source/Runtime/Experimental/ChaosVehicles/ChaosVehiclesCore/Private/SimModule/EngineModule.cpp
@@ -21,8 +21,15 @@ namespace Chaos
 		// TODO: Engine braking effect
 		DriveTorque = GetEngineTorque(Inputs.ControlInputs.Throttle, GetRPM());
 
-		float BrakeTorque = 0.f;
-		TransmitTorque(VehicleModuleSystem, DriveTorque, BrakeTorque);
+		if (DriveTorque < SMALL_NUMBER)
+		{
+			BrakingTorque = Setup().EngineBrakeEffect;
+		}
+		else
+		{
+			BrakingTorque = 0.0f;
+		}
+		TransmitTorque(VehicleModuleSystem, DriveTorque, BrakingTorque);
 		IntegrateAngularVelocity(DeltaTime, Setup().EngineInertia);
 
 		// clamp RPM
@@ -56,8 +63,6 @@ namespace Chaos
 
 	float FEngineSimModule::GetTorqueFromRPM(float RPM, bool LimitToIdle)
 	{
-		//return Setup().MaxTorque; // TODO: Fix this - engine RPM running rampant
-
 		if (!EngineStarted || (FMath::Abs(RPM - Setup().MaxRPM) < 1.0f) || Setup().MaxRPM == 0)
 		{
 			return 0.0f;
diff --git a/Engine/Source/Runtime/Experimental/ChaosVehicles/ChaosVehiclesCore/Private/SimModule/TorqueSimModule.cpp b/Engine/Source/Runtime/Experimental/ChaosVehicles/ChaosVehiclesCore/Private/SimModule/TorqueSimModule.cpp
--- a/Engine/Source/Runtime/Experimental/ChaosVehicles/ChaosVehiclesCore/Private/SimModule/TorqueSimModule.cpp
+++ b/Engine/Source/Runtime/Experimental/ChaosVehicles/ChaosVehiclesCore/Private/SimModule/TorqueSimModule.cpp
@@ -22,6 +22,8 @@ void FTorqueSimModule::TransmitTorque(const FSimModuleTree& ModuleTree, float Pu
 
 		LoadTorque = 0; // clear torques prior to summation from child nodes
 		BrakingTorque = BrakeTorque;
+		float AverageChildVelocity = 0.0f;
+		int NumVelocities = 0;
 		for (int ChildIndex : Children)
 		{
 			// currently doesn't make sense to transmit from one wheel to next in hierarchy - only because hierarchy is 1 deep at present, this may change in future
@@ -41,11 +43,18 @@ void FTorqueSimModule::TransmitTorque(const FSimModuleTree& ModuleTree, float Pu
 				BrakingTorque += Interface->GetBrakingTorque() * ClutchSlip;
 
 				// diff velocity 
-				float DiffVelocity = (Interface->AngularVelocity * GearingRatio) - AngularVelocity;
-				AngularVelocity += DiffVelocity * ClutchSlip / Children.Num();
+				AverageChildVelocity += (Interface->AngularVelocity * GearingRatio);
+
+				NumVelocities++;
 			}
 		}
 
+		if (NumVelocities > 0)
+		{
+			float DiffVelocity = (AverageChildVelocity / NumVelocities) - AngularVelocity;
+			AngularVelocity += DiffVelocity * ClutchSlip;
+		}
+
 		BrakingTorque /= FMath::Abs(GearingRatio);
 
 	}
@@ -60,7 +69,7 @@ void FTorqueSimModule::IntegrateAngularVelocity(float DeltaTime, float Inertia,
 	ensure(BrakingTorque >= 0.0f);
 
 	// drive torque taken into account
-	AngularVelocity += (DriveTorque + LoadTorque * DeltaTime) / Inertia;
+	AngularVelocity += ((DriveTorque + LoadTorque) * DeltaTime) / Inertia;
 
 	// braking resists velocity no matter what way we are spinning
 	// also has check that we are not overshooting and starting to accelerating in the opposite direction
@@ -90,8 +99,6 @@ void FTorqueSimModule::IntegrateAngularVelocity(float DeltaTime, float Inertia,
 	int ExcessRotations = (int)(AngularPosition / TWO_PI);
 	AngularPosition -= ExcessRotations * TWO_PI;
 
-	AngularVelocity *= 0.98f; // TEMP
-
 	UE_LOG(LogSimulationModule, Log, TEXT("%s: DriveTorque %4.2f, BrakingTorque %4.2f, LoadTorque %4.2f, Speed %4.2f rad/sec, RPM %4.2f, AngularPosition %4.2f, Inertia %4.2f")
 		, *GetDebugName(), DriveTorque, BrakingTorque, LoadTorque, AngularVelocity, GetRPM(), AngularPosition, Inertia);
 }
diff --git a/Engine/Source/Runtime/Experimental/ChaosVehicles/ChaosVehiclesCore/Private/SimModule/WheelModule.cpp b/Engine/Source/Runtime/Experimental/ChaosVehicles/ChaosVehiclesCore/Private/SimModule/WheelModule.cpp
--- a/Engine/Source/Runtime/Experimental/ChaosVehicles/ChaosVehiclesCore/Private/SimModule/WheelModule.cpp
+++ b/Engine/Source/Runtime/Experimental/ChaosVehicles/ChaosVehiclesCore/Private/SimModule/WheelModule.cpp
@@ -35,10 +35,13 @@ namespace Chaos
 		SteerAngleDegrees = Setup().SteeringEnabled ? Inputs.ControlInputs.Steering * Setup().MaxSteeringAngle : 0.0f;
 		BrakeTorque = Inputs.ControlInputs.Brake * Setup().MaxBrakeTorque + HandbrakeTorque;
 		LoadTorque = 0.0f;
+		ForceFromFriction = FVector::ZeroVector;
+		float TorqueFromGroundInteraction = 0.0f;
+		float AvailableGrip = 0.0f;
 
 		// TODO: think about doing this properly, stops vehicles rolling around on their own too much
 		// i.e. an auto handbrake feature
-		if (Inputs.ControlInputs.Brake < SMALL_NUMBER && Inputs.ControlInputs.Throttle < SMALL_NUMBER && LocalLinearVelocity.X < 10.0f)
+		if (Setup().AutoHandbrakeEnabled && LocalLinearVelocity.X < Setup().AutoHandbrakeVelocityThreshold && (Inputs.ControlInputs.Brake < SMALL_NUMBER && Inputs.ControlInputs.Throttle < SMALL_NUMBER))
 		{
 			BrakeTorque = Setup().HandbrakeTorque;
 		}
@@ -54,7 +57,7 @@ namespace Chaos
 
 			float GroundAngularVelocity = LocalWheelVelocity.X / Re;
 			float Delta = GroundAngularVelocity - AngularVelocity;
-			float TorqueFromGroundInteraction = Delta * Setup().WheelInertia / DeltaTime; // torque from wheels moving over terrain
+			TorqueFromGroundInteraction = Delta * Setup().WheelInertia / DeltaTime; // torque from wheels moving over terrain
 
 			// X is longitudinal direction, Y is lateral
 			SlipAngle = FVehicleUtility::CalculateSlipAngle(LocalWheelVelocity.Y, LocalWheelVelocity.X);
@@ -63,7 +66,7 @@ namespace Chaos
 			float AppliedLinearBrakeForce = FMath::Abs(BrakeTorque) / Re;
 
 			// Longitudinal multiplier now affecting both brake and steering equally
-			float AvailableGrip = ForceIntoSurface * SurfaceFriction * Setup().FrictionMultiplier;
+			AvailableGrip = ForceIntoSurface * SurfaceFriction * Setup().FrictionMultiplier;
 
 			float FinalLongitudinalForce = 0.f;
 			float FinalLateralForce = 0.f;
@@ -155,12 +158,16 @@ namespace Chaos
 						bClipping = true;
 						FinalLongitudinalForce *= Clip;
 						FinalLateralForce *= Clip;
+
+						// make the resulting forces less than ideal since there is slippage
+						FinalLongitudinalForce *= Setup().SlipModifier;
+						FinalLateralForce *= Setup().SlipModifier;
+
 					}
 				}
 
 			}
 
-			ForceFromFriction = FVector::ZeroVector;
 			// Potential Axis Swap
 			if (Setup().Axis == EWheelAxis::X)
 			{
@@ -180,40 +187,44 @@ namespace Chaos
 			}
 
 			AddLocalForce(SteeringRotator.RotateVector(ForceFromFriction));
-			TransmitTorque(VehicleModuleSystem, DriveTorque, BrakeTorque);
-
-			DriveTorque -= AvailableGrip;
-			if (DriveTorque < 0.0f)
-			{
-				DriveTorque = 0.0f;
-			}
+		}
 
-			DriveTorque *= TorqueScaling;
+		TransmitTorque(VehicleModuleSystem, DriveTorque, BrakeTorque);
 
-			BrakingTorque -= AvailableGrip;
-			if (BrakingTorque < 0.0f)
-			{
-				BrakingTorque = 0.0f;
-			}
+		DriveTorque -= AvailableGrip;
+		if (DriveTorque < 0.0f)
+		{
+			DriveTorque = 0.0f;
+		}
 
-			BrakingTorque *= TorqueScaling;
+		DriveTorque *= TorqueScaling;
 
-			LoadTorque = TorqueFromGroundInteraction;
+		BrakingTorque -= AvailableGrip;
+		if (BrakingTorque < 0.0f)
+		{
+			BrakingTorque = 0.0f;
 		}
 
+		BrakingTorque *= TorqueScaling;
+
+		LoadTorque = TorqueFromGroundInteraction;
+
 		IntegrateAngularVelocity(DeltaTime, Setup().WheelInertia, Setup().MaxRotationVel);
 	}
 
 	void FWheelSimModule::Animate(Chaos::FClusterUnionPhysicsProxy* Proxy)
 	{
-		if (FPBDRigidClusteredParticleHandle* ClusterChild = GetClusterParticle(Proxy))
+		if (Proxy)
 		{
-			float Direction = Setup().ReverseDirection ? -1.0f : 1.0f;
-			FQuat Rot = (Setup().Axis == Chaos::EWheelAxis::Y) ? FQuat(FVector(1, 0, 0), -GetAngularPosition() * Direction) : FQuat(FVector(1, 0, 0), GetAngularPosition() * Direction);
-			FQuat Steer = FQuat(FVector(0, 0, 1), FMath::DegreesToRadians(GetSteerAngleDegrees()));
+			if (FPBDRigidClusteredParticleHandle* ClusterChild = GetClusterParticle(Proxy))
+			{
+				float Direction = Setup().ReverseDirection ? -1.0f : 1.0f;
+				FQuat Rot = (Setup().Axis == Chaos::EWheelAxis::Y) ? FQuat(FVector(1, 0, 0), -GetAngularPosition() * Direction) : FQuat(FVector(1, 0, 0), GetAngularPosition() * Direction);
+				FQuat Steer = FQuat(FVector(0, 0, 1), FMath::DegreesToRadians(GetSteerAngleDegrees()));
 
-			FTransform InitialTransform = GetInitialParticleTransform();
-			ClusterChild->ChildToParent().SetRotation(InitialTransform.GetRotation() * Steer * Rot);
+				FTransform InitialTransform = GetInitialParticleTransform();
+				ClusterChild->ChildToParent().SetRotation(InitialTransform.GetRotation() * Steer * Rot);
+			}
 		}
 	}
 
diff --git a/Engine/Source/Runtime/Experimental/ChaosVehicles/ChaosVehiclesCore/Public/SimModule/AerofoilModule.h b/Engine/Source/Runtime/Experimental/ChaosVehicles/ChaosVehiclesCore/Public/SimModule/AerofoilModule.h
--- a/Engine/Source/Runtime/Experimental/ChaosVehicles/ChaosVehiclesCore/Public/SimModule/AerofoilModule.h
+++ b/Engine/Source/Runtime/Experimental/ChaosVehicles/ChaosVehiclesCore/Public/SimModule/AerofoilModule.h
@@ -83,7 +83,6 @@ namespace Chaos
 			CurrentAirDensity = InDensity;
 		}
 
-	protected:
 		void SetControlSurface(float CtrlSurfaceInput)
 		{
 			ControlSurfaceAngle = CtrlSurfaceInput * Setup().MaxControlAngle;
diff --git a/Engine/Source/Runtime/Experimental/ChaosVehicles/ChaosVehiclesCore/Public/SimModule/EngineModule.h b/Engine/Source/Runtime/Experimental/ChaosVehicles/ChaosVehiclesCore/Public/SimModule/EngineModule.h
--- a/Engine/Source/Runtime/Experimental/ChaosVehicles/ChaosVehiclesCore/Public/SimModule/EngineModule.h
+++ b/Engine/Source/Runtime/Experimental/ChaosVehicles/ChaosVehiclesCore/Public/SimModule/EngineModule.h
@@ -56,8 +56,8 @@ namespace Chaos
 			: MaxTorque(300.f)
 			, MaxRPM(6000)
 			, IdleRPM(1200)
-			, EngineBrakeEffect(0.2f)
-			, EngineInertia(100000.0f)
+			, EngineBrakeEffect(50.0f)
+			, EngineInertia(100.0f)
 		{
 			TorqueCurve.AddNormalized(0.5f);
 			TorqueCurve.AddNormalized(0.5f);
@@ -76,8 +76,8 @@ namespace Chaos
 		FNormalisedGraph TorqueCurve;
 		float MaxTorque;			// [N.m] The peak torque Y value in the normalized torque graph
 		uint16 MaxRPM;				// [RPM] The absolute maximum RPM the engine can theoretically reach (last X value in the normalized torque graph)
-		uint16 IdleRPM; 		// [RPM] The RPM at which the throttle sits when the car is not moving			
-		float EngineBrakeEffect;	// [0..1] How much the engine slows the vehicle when the throttle is released
+		uint16 IdleRPM; 			// [RPM] The RPM at which the throttle sits when the car is not moving			
+		float EngineBrakeEffect;	// [N.m] How much the engine slows the vehicle when the throttle is released
 
 		float EngineInertia;
 	};
diff --git a/Engine/Source/Runtime/Experimental/ChaosVehicles/ChaosVehiclesCore/Public/SimModule/TransmissionModule.h b/Engine/Source/Runtime/Experimental/ChaosVehicles/ChaosVehiclesCore/Public/SimModule/TransmissionModule.h
--- a/Engine/Source/Runtime/Experimental/ChaosVehicles/ChaosVehiclesCore/Public/SimModule/TransmissionModule.h
+++ b/Engine/Source/Runtime/Experimental/ChaosVehicles/ChaosVehiclesCore/Public/SimModule/TransmissionModule.h
@@ -139,7 +139,7 @@ namespace Chaos
 
 		virtual void Simulate(float DeltaTime, const FAllInputs& Inputs, FSimModuleTree& VehicleModuleSystem) override;
 
-	private:
+	protected:
 
 		/** set the target gear number to change to, can change gear immediately if specified
 		 *  i.e. rather than waiting for the gear change time to elapse
@@ -172,6 +172,10 @@ namespace Chaos
 			GearIndexInOut = FMath::Clamp(GearIndexInOut, -Setup().ReverseRatios.Num(), Setup().ForwardRatios.Num());
 		}
 
+		int32 GetCurrentGear() { return CurrentGear; }
+		int32 GetTargetGear() { return TargetGear; }
+
+	private:
 		int32 CurrentGear; // <0 reverse gear(s), 0 neutral, >0 forward gears
 		int32 TargetGear;  // <0 reverse gear(s), 0 neutral, >0 forward gears
 		float CurrentGearChangeTime; // Time to change gear, no power transmitted to the wheels during change
diff --git a/Engine/Source/Runtime/Experimental/ChaosVehicles/ChaosVehiclesCore/Public/SimModule/WheelModule.h b/Engine/Source/Runtime/Experimental/ChaosVehicles/ChaosVehiclesCore/Public/SimModule/WheelModule.h
--- a/Engine/Source/Runtime/Experimental/ChaosVehicles/ChaosVehiclesCore/Public/SimModule/WheelModule.h
+++ b/Engine/Source/Runtime/Experimental/ChaosVehicles/ChaosVehiclesCore/Public/SimModule/WheelModule.h
@@ -67,17 +67,20 @@ namespace Chaos
 		FWheelSettings()
 			: Radius(30.0f)
 			, Width(20.0f)
-			, WheelInertia(100000.0f) // TODO: ?? defaults and proper value ??
+			, WheelInertia(100.0f)
 
 			, FrictionMultiplier(3.0f)
 			, LateralSlipGraphMultiplier(1.0f)
 			, CorneringStiffness(1000.0f)
 			, SlipAngleLimit(8.0f)
+			, SlipModifier(0.9f)
 
 			, ABSEnabled(true)
 			, TractionControlEnabled(true)
 			, SteeringEnabled(false)
 			, HandbrakeEnabled(false)
+			, AutoHandbrakeEnabled(false)
+			, AutoHandbrakeVelocityThreshold(10.0f)
 			, MaxSteeringAngle(45)
 			, MaxBrakeTorque(4000)
 			, HandbrakeTorque(3000)
@@ -97,11 +100,14 @@ namespace Chaos
 		float CorneringStiffness;
 		FGraph LateralSlipGraph;
 		float SlipAngleLimit;
+		float SlipModifier;
 
 		bool ABSEnabled;			// Advanced braking system operational
 		bool TractionControlEnabled;// Straight Line Traction Control
 		bool SteeringEnabled;
 		bool HandbrakeEnabled;
+		bool AutoHandbrakeEnabled;
+		float AutoHandbrakeVelocityThreshold; 
 
 		float MaxSteeringAngle;
 		float MaxBrakeTorque;
@@ -154,8 +160,26 @@ namespace Chaos
 
 		void SetForceIntoSurface(float ForceIntoSurfaceIn) { ForceIntoSurface = ForceIntoSurfaceIn; }
 		float GetForceIntoSurface() const { return ForceIntoSurface; }
-
+		FVector GetForceFromFriction() const { return ForceFromFriction; }
 		void SetSurfaceFriction(float FrictionIn) { SurfaceFriction = FrictionIn; }
+		
+		/** set wheel rotational speed to match the specified linear forwards speed */
+		void SetLinearSpeed(float LinearMetersPerSecondIn)
+		{
+			SetAngularVelocity(LinearMetersPerSecondIn / Setup().Radius);
+		}
+
+		/** get linear forwards speed from angluar velocity and wheel radius */
+		float GetLinearSpeed()
+		{
+			return AngularVelocity * Setup().Radius;
+		}
+
+		/** Get the radius of the wheel [cm] */
+		float GetEffectiveRadius() const
+		{
+			return Setup().Radius;
+		}
 
 	private:
 
