Fix Chaos event manager unit test after changes to clear out single-threaded buffers on flip

diff --git a/Engine/Source/Programs/HeadlessChaos/Private/GeometryCollection/GeometryCollectionTestEvents.cpp b/Engine/Source/Programs/HeadlessChaos/Private/GeometryCollection/GeometryCollectionTestEvents.cpp
--- a/Engine/Source/Programs/HeadlessChaos/Private/GeometryCollection/GeometryCollectionTestEvents.cpp
+++ b/Engine/Source/Programs/HeadlessChaos/Private/GeometryCollection/GeometryCollectionTestEvents.cpp
@@ -133,7 +133,6 @@ namespace GeometryCollectionTest
 		TestData.Data2 = FVector(1, 2, 3);
 
 		EventManager.FillProducerData(Solver);
-		EventManager.FlipBuffersIfRequired();
 		EventManager.DispatchEvents();
 		EXPECT_EQ(HandlerTest.ResultFromHandler, TestData);
 
@@ -141,7 +140,6 @@ namespace GeometryCollectionTest
 		TestData.Data2 = FVector(7, 8, 9);
 
 		EventManager.FillProducerData(Solver);
-		EventManager.FlipBuffersIfRequired();
 		EventManager.DispatchEvents();
 		EXPECT_EQ(HandlerTest.ResultFromHandler, TestData);
 
@@ -153,7 +151,6 @@ namespace GeometryCollectionTest
 		TestData.Data2 = FVector(9, 9, 9);
 
 		EventManager.FillProducerData(Solver);
-		EventManager.FlipBuffersIfRequired();
 		EventManager.DispatchEvents();
 		EXPECT_EQ(HandlerTest.ResultFromHandler, OriginalTestData);
 
@@ -163,7 +160,6 @@ namespace GeometryCollectionTest
 		TestArrayData.Push(EventTestData(789, FVector(7, 8, 9)));
 
 		EventManager.FillProducerData(Solver);
-		EventManager.FlipBuffersIfRequired();
 		EventManager.DispatchEvents();
 		// dispatched to multiple handlers
 		EXPECT_EQ(HandlerTest.ResultFromHandler2, TestArrayData);
@@ -177,7 +173,6 @@ namespace GeometryCollectionTest
 		TestArrayData.Push(EventTestData(999, FVector(9, 9, 9)));
 
 		EventManager.FillProducerData(Solver);
-		EventManager.FlipBuffersIfRequired();
 		EventManager.DispatchEvents();
 		EXPECT_EQ(HandlerTest.ResultFromHandler2, OriginalTestArrayData); // Unregistered - data should no longer update
 		EXPECT_EQ(AnotherHandlerTest.ResultFromHandler2, TestArrayData); // Still registered so should get updates
@@ -185,7 +180,6 @@ namespace GeometryCollectionTest
 		HandlerTest.RegisterHandler2();
 
 		EventManager.FillProducerData(Solver);
-		EventManager.FlipBuffersIfRequired();
 		EventManager.DispatchEvents();
 
 		EXPECT_EQ(HandlerTest.ResultFromHandler2, TestArrayData);
