Fix unit test, fallout from memory optimization

diff --git a/Engine/Source/Programs/HeadlessChaos/Private/GeometryCollection/GeometryCollectionTestSimulation.cpp b/Engine/Source/Programs/HeadlessChaos/Private/GeometryCollection/GeometryCollectionTestSimulation.cpp
--- a/Engine/Source/Programs/HeadlessChaos/Private/GeometryCollection/GeometryCollectionTestSimulation.cpp
+++ b/Engine/Source/Programs/HeadlessChaos/Private/GeometryCollection/GeometryCollectionTestSimulation.cpp
@@ -285,7 +285,6 @@ namespace GeometryCollectionTest
 		UnitTest.Advance();
 
 		// testing...
-		auto GCParticles = Collection->PhysObject->GetSolverParticleHandles();
 		for (int Frame = 1; Frame < 200; Frame++)
 		{
 			if (Frame == 1)
@@ -295,11 +294,11 @@ namespace GeometryCollectionTest
 				// Object 3 collides with Object 0 plus the ground
 				// Object 4 collides with nothing
 				// We should end up with 2 stacks on the ground (0,1,2), (0,3) and one free-falling object (4)
-				GCParticles[0]->SetCollisionGroup(0);
-				GCParticles[1]->SetCollisionGroup(1);
-				GCParticles[2]->SetCollisionGroup(1);
-				GCParticles[3]->SetCollisionGroup(3);
-				GCParticles[4]->SetCollisionGroup(-1);
+				Collection->PhysObject->GetSolverClusterHandle_Internal(0)->SetCollisionGroup(0);
+				Collection->PhysObject->GetSolverClusterHandle_Internal(1)->SetCollisionGroup(1);
+				Collection->PhysObject->GetSolverClusterHandle_Internal(2)->SetCollisionGroup(1);
+				Collection->PhysObject->GetSolverClusterHandle_Internal(3)->SetCollisionGroup(3);
+				Collection->PhysObject->GetSolverClusterHandle_Internal(4)->SetCollisionGroup(-1);
 
 				EXPECT_TRUE(Collection->DynamicCollection->GetTransform(0).GetRotation() == FQuat4f::Identity); // Can use defaulted zero rotation to indicate that the
 				EXPECT_TRUE(Collection->DynamicCollection->GetTransform(1).GetRotation() == FQuat4f::Identity); // rigid has not been affected. Should we though??
@@ -330,7 +329,7 @@ namespace GeometryCollectionTest
 		EXPECT_TRUE(Collection->DynamicCollection->GetTransform(4).GetRotation() == FQuat4f::Identity); // Phased through everything, good.
 		EXPECT_LT(Collection->DynamicCollection->GetTransform(4).GetTranslation().Z, -100.0f);
 
-		GCParticles[0]->SetCollisionGroup(-1);
+		Collection->PhysObject->GetSolverClusterHandle_Internal(0)->SetCollisionGroup(-1);
 		for (int i = 0; i < 50; i++) { UnitTest.Advance(); }
 		EXPECT_LT(Collection->DynamicCollection->GetTransform(0).GetTranslation().Z, -100.0f);
 
diff --git a/Engine/Source/Programs/HeadlessChaos/Private/GeometryCollection/GeometryCollectionTestSimulationField.cpp b/Engine/Source/Programs/HeadlessChaos/Private/GeometryCollection/GeometryCollectionTestSimulationField.cpp
--- a/Engine/Source/Programs/HeadlessChaos/Private/GeometryCollection/GeometryCollectionTestSimulationField.cpp
+++ b/Engine/Source/Programs/HeadlessChaos/Private/GeometryCollection/GeometryCollectionTestSimulationField.cpp
@@ -536,48 +536,47 @@ namespace GeometryCollectionTest
 
 		UnitTest.Advance();
 
-		TArray<Chaos::TPBDRigidClusteredParticleHandle<FReal, 3>*>& ParticleHandles = Collection->PhysObject->GetSolverParticleHandles();
 		{
 			FName TargetName = GetFieldPhysicsName(EFieldPhysicsType::Field_ExternalClusterStrain);
 			UnitTest.Solver->GetPerSolverField().AddTransientCommand({ TargetName, FalloffField->NewCopy() });
 
 			EXPECT_EQ(ClusterMap.Num(), 3);
-			EXPECT_EQ(ClusterMap[ParticleHandles[4]].Num(), 2);
-			EXPECT_TRUE(ClusterMap[ParticleHandles[4]].Contains(ParticleHandles[0]));
-			EXPECT_TRUE(ClusterMap[ParticleHandles[4]].Contains(ParticleHandles[1]));
-			EXPECT_EQ(ClusterMap[ParticleHandles[5]].Num(), 2);
-			EXPECT_TRUE(ClusterMap[ParticleHandles[5]].Contains(ParticleHandles[2]));
-			EXPECT_TRUE(ClusterMap[ParticleHandles[5]].Contains(ParticleHandles[3]));
-			EXPECT_EQ(ClusterMap[ParticleHandles[6]].Num(), 2);
-			EXPECT_TRUE(ClusterMap[ParticleHandles[6]].Contains(ParticleHandles[5]));
-			EXPECT_TRUE(ClusterMap[ParticleHandles[6]].Contains(ParticleHandles[4]));
-
-			EXPECT_TRUE(ParticleHandles[0]->Disabled());
-			EXPECT_TRUE(ParticleHandles[1]->Disabled());
-			EXPECT_TRUE(ParticleHandles[2]->Disabled());
-			EXPECT_TRUE(ParticleHandles[3]->Disabled());
-			EXPECT_TRUE(ParticleHandles[4]->Disabled());
-			EXPECT_TRUE(ParticleHandles[5]->Disabled());
-			EXPECT_FALSE(ParticleHandles[6]->Disabled());
+			EXPECT_EQ(ClusterMap[Collection->PhysObject->GetParticle_Internal(4)].Num(), 2);
+			EXPECT_TRUE(ClusterMap[Collection->PhysObject->GetParticle_Internal(4)].Contains(Collection->PhysObject->GetParticle_Internal(0)));
+			EXPECT_TRUE(ClusterMap[Collection->PhysObject->GetParticle_Internal(4)].Contains(Collection->PhysObject->GetParticle_Internal(1)));
+			EXPECT_EQ(ClusterMap[Collection->PhysObject->GetParticle_Internal(5)].Num(), 2);
+			EXPECT_TRUE(ClusterMap[Collection->PhysObject->GetParticle_Internal(5)].Contains(Collection->PhysObject->GetParticle_Internal(2)));
+			EXPECT_TRUE(ClusterMap[Collection->PhysObject->GetParticle_Internal(5)].Contains(Collection->PhysObject->GetParticle_Internal(3)));
+			EXPECT_EQ(ClusterMap[Collection->PhysObject->GetParticle_Internal(6)].Num(), 2);
+			EXPECT_TRUE(ClusterMap[Collection->PhysObject->GetParticle_Internal(6)].Contains(Collection->PhysObject->GetParticle_Internal(5)));
+			EXPECT_TRUE(ClusterMap[Collection->PhysObject->GetParticle_Internal(6)].Contains(Collection->PhysObject->GetParticle_Internal(4)));
+
+			EXPECT_TRUE(Collection->PhysObject->GetParticle_Internal(0)->Disabled());
+			EXPECT_TRUE(Collection->PhysObject->GetParticle_Internal(1)->Disabled());
+			EXPECT_TRUE(Collection->PhysObject->GetParticle_Internal(2)->Disabled());
+			EXPECT_TRUE(Collection->PhysObject->GetParticle_Internal(3)->Disabled());
+			EXPECT_TRUE(Collection->PhysObject->GetParticle_Internal(4)->Disabled());
+			EXPECT_TRUE(Collection->PhysObject->GetParticle_Internal(5)->Disabled());
+			EXPECT_FALSE(Collection->PhysObject->GetParticle_Internal(6)->Disabled());
 
 			UnitTest.Advance();
 
 			// todo: indices here might seem odd, particles 4 & 5 are swapped
 			EXPECT_EQ(ClusterMap.Num(), 2);
-			EXPECT_EQ(ClusterMap[ParticleHandles[4]].Num(), 2);
-			EXPECT_TRUE(ClusterMap[ParticleHandles[4]].Contains(ParticleHandles[0]));
-			EXPECT_TRUE(ClusterMap[ParticleHandles[4]].Contains(ParticleHandles[1]));
-			EXPECT_EQ(ClusterMap[ParticleHandles[5]].Num(), 2);
-			EXPECT_TRUE(ClusterMap[ParticleHandles[5]].Contains(ParticleHandles[2]));
-			EXPECT_TRUE(ClusterMap[ParticleHandles[5]].Contains(ParticleHandles[3]));
-
-			EXPECT_TRUE(ParticleHandles[0]->Disabled());
-			EXPECT_TRUE(ParticleHandles[1]->Disabled());
-			EXPECT_TRUE(ParticleHandles[2]->Disabled());
-			EXPECT_TRUE(ParticleHandles[3]->Disabled());
-			EXPECT_FALSE(ParticleHandles[4]->Disabled());
-			EXPECT_FALSE(ParticleHandles[5]->Disabled());
-			EXPECT_TRUE(ParticleHandles[6]->Disabled());
+			EXPECT_EQ(ClusterMap[Collection->PhysObject->GetParticle_Internal(4)].Num(), 2);
+			EXPECT_TRUE(ClusterMap[Collection->PhysObject->GetParticle_Internal(4)].Contains(Collection->PhysObject->GetParticle_Internal(0)));
+			EXPECT_TRUE(ClusterMap[Collection->PhysObject->GetParticle_Internal(4)].Contains(Collection->PhysObject->GetParticle_Internal(1)));
+			EXPECT_EQ(ClusterMap[Collection->PhysObject->GetParticle_Internal(5)].Num(), 2);
+			EXPECT_TRUE(ClusterMap[Collection->PhysObject->GetParticle_Internal(5)].Contains(Collection->PhysObject->GetParticle_Internal(2)));
+			EXPECT_TRUE(ClusterMap[Collection->PhysObject->GetParticle_Internal(5)].Contains(Collection->PhysObject->GetParticle_Internal(3)));
+
+			EXPECT_TRUE(Collection->PhysObject->GetParticle_Internal(0)->Disabled());
+			EXPECT_TRUE(Collection->PhysObject->GetParticle_Internal(1)->Disabled());
+			EXPECT_TRUE(Collection->PhysObject->GetParticle_Internal(2)->Disabled());
+			EXPECT_TRUE(Collection->PhysObject->GetParticle_Internal(3)->Disabled());
+			EXPECT_FALSE(Collection->PhysObject->GetParticle_Internal(4)->Disabled());
+			EXPECT_FALSE(Collection->PhysObject->GetParticle_Internal(5)->Disabled());
+			EXPECT_TRUE(Collection->PhysObject->GetParticle_Internal(6)->Disabled());
 		}
 
 		delete FalloffField;
@@ -613,7 +612,11 @@ namespace GeometryCollectionTest
 		UnitTest.Initialize();
 		UnitTest.Advance();
 
-		TArray<Chaos::TPBDRigidClusteredParticleHandle<FReal, 3>*>& ParticleHandles = Collection->PhysObject->GetSolverParticleHandles();
+		auto ParticleHandles = [&Collection](int32 Index) -> Chaos::TPBDRigidClusteredParticleHandle<FReal, 3>*
+		{
+			return Collection->PhysObject->GetParticle_Internal(Index);
+		};
+
 		auto& Clustering = UnitTest.Solver->GetEvolution()->GetRigidClustering();
 		const auto& ClusterMap = Clustering.GetChildrenMap();
 
@@ -625,47 +628,47 @@ namespace GeometryCollectionTest
 			UnitTest.Solver->GetPerSolverField().AddTransientCommand(Command);
 
 			EXPECT_EQ(ClusterMap.Num(), 3);
-			EXPECT_EQ(ClusterMap[ParticleHandles[6]].Num(), 3);
-			EXPECT_TRUE(ClusterMap[ParticleHandles[6]].Contains(ParticleHandles[0]));
-			EXPECT_TRUE(ClusterMap[ParticleHandles[6]].Contains(ParticleHandles[1]));
-			EXPECT_TRUE(ClusterMap[ParticleHandles[6]].Contains(ParticleHandles[2]));
-			EXPECT_EQ(ClusterMap[ParticleHandles[7]].Num(), 3);
-			EXPECT_TRUE(ClusterMap[ParticleHandles[7]].Contains(ParticleHandles[3]));
-			EXPECT_TRUE(ClusterMap[ParticleHandles[7]].Contains(ParticleHandles[4]));
-			EXPECT_TRUE(ClusterMap[ParticleHandles[7]].Contains(ParticleHandles[5]));
-			EXPECT_EQ(ClusterMap[ParticleHandles[8]].Num(), 2);
-			EXPECT_TRUE(ClusterMap[ParticleHandles[8]].Contains(ParticleHandles[7]));
-			EXPECT_TRUE(ClusterMap[ParticleHandles[8]].Contains(ParticleHandles[6]));
-
-			EXPECT_TRUE(ParticleHandles[0]->Disabled());
-			EXPECT_TRUE(ParticleHandles[1]->Disabled());
-			EXPECT_TRUE(ParticleHandles[2]->Disabled());
-			EXPECT_TRUE(ParticleHandles[3]->Disabled());
-			EXPECT_TRUE(ParticleHandles[4]->Disabled());
-			EXPECT_TRUE(ParticleHandles[5]->Disabled());
-			EXPECT_TRUE(ParticleHandles[6]->Disabled());
-			EXPECT_TRUE(ParticleHandles[7]->Disabled());
-			EXPECT_FALSE(ParticleHandles[8]->Disabled());
+			EXPECT_EQ(ClusterMap[ParticleHandles(6)].Num(), 3);
+			EXPECT_TRUE(ClusterMap[ParticleHandles(6)].Contains(ParticleHandles(0)));
+			EXPECT_TRUE(ClusterMap[ParticleHandles(6)].Contains(ParticleHandles(1)));
+			EXPECT_TRUE(ClusterMap[ParticleHandles(6)].Contains(ParticleHandles(2)));
+			EXPECT_EQ(ClusterMap[ParticleHandles(7)].Num(), 3);
+			EXPECT_TRUE(ClusterMap[ParticleHandles(7)].Contains(ParticleHandles(3)));
+			EXPECT_TRUE(ClusterMap[ParticleHandles(7)].Contains(ParticleHandles(4)));
+			EXPECT_TRUE(ClusterMap[ParticleHandles(7)].Contains(ParticleHandles(5)));
+			EXPECT_EQ(ClusterMap[ParticleHandles(8)].Num(), 2);
+			EXPECT_TRUE(ClusterMap[ParticleHandles(8)].Contains(ParticleHandles(7)));
+			EXPECT_TRUE(ClusterMap[ParticleHandles(8)].Contains(ParticleHandles(6)));
+
+			EXPECT_TRUE(ParticleHandles(0)->Disabled());
+			EXPECT_TRUE(ParticleHandles(1)->Disabled());
+			EXPECT_TRUE(ParticleHandles(2)->Disabled());
+			EXPECT_TRUE(ParticleHandles(3)->Disabled());
+			EXPECT_TRUE(ParticleHandles(4)->Disabled());
+			EXPECT_TRUE(ParticleHandles(5)->Disabled());
+			EXPECT_TRUE(ParticleHandles(6)->Disabled());
+			EXPECT_TRUE(ParticleHandles(7)->Disabled());
+			EXPECT_FALSE(ParticleHandles(8)->Disabled());
 
 			UnitTest.Advance();
 			UnitTest.Solver->GetPerSolverField().AddTransientCommand(Command);
 			UnitTest.Advance();
 
 			EXPECT_EQ(ClusterMap.Num(), 1);
-			EXPECT_EQ(ClusterMap[ParticleHandles[7]].Num(), 3);
-			EXPECT_TRUE(ClusterMap[ParticleHandles[7]].Contains(ParticleHandles[3]));
-			EXPECT_TRUE(ClusterMap[ParticleHandles[7]].Contains(ParticleHandles[4]));
-			EXPECT_TRUE(ClusterMap[ParticleHandles[7]].Contains(ParticleHandles[5]));
-
-			EXPECT_FALSE(ParticleHandles[0]->Disabled());
-			EXPECT_FALSE(ParticleHandles[1]->Disabled());
-			EXPECT_FALSE(ParticleHandles[2]->Disabled());
-			EXPECT_TRUE(ParticleHandles[3]->Disabled());
-			EXPECT_TRUE(ParticleHandles[4]->Disabled());
-			EXPECT_TRUE(ParticleHandles[5]->Disabled());
-			EXPECT_TRUE(ParticleHandles[6]->Disabled());
-			EXPECT_FALSE(ParticleHandles[7]->Disabled());
-			EXPECT_TRUE(ParticleHandles[8]->Disabled());
+			EXPECT_EQ(ClusterMap[ParticleHandles(7)].Num(), 3);
+			EXPECT_TRUE(ClusterMap[ParticleHandles(7)].Contains(ParticleHandles(3)));
+			EXPECT_TRUE(ClusterMap[ParticleHandles(7)].Contains(ParticleHandles(4)));
+			EXPECT_TRUE(ClusterMap[ParticleHandles(7)].Contains(ParticleHandles(5)));
+
+			EXPECT_FALSE(ParticleHandles(0)->Disabled());
+			EXPECT_FALSE(ParticleHandles(1)->Disabled());
+			EXPECT_FALSE(ParticleHandles(2)->Disabled());
+			EXPECT_TRUE(ParticleHandles(3)->Disabled());
+			EXPECT_TRUE(ParticleHandles(4)->Disabled());
+			EXPECT_TRUE(ParticleHandles(5)->Disabled());
+			EXPECT_TRUE(ParticleHandles(6)->Disabled());
+			EXPECT_FALSE(ParticleHandles(7)->Disabled());
+			EXPECT_TRUE(ParticleHandles(8)->Disabled());
 		}
 
 		delete FalloffField;
@@ -704,7 +707,10 @@ namespace GeometryCollectionTest
 		auto& Clustering = UnitTest.Solver->GetEvolution()->GetRigidClustering();
 		const auto& ClusterMap = Clustering.GetChildrenMap();
 
-		TArray<Chaos::TPBDRigidClusteredParticleHandle<FReal, 3>*>& ParticleHandles = Collection->PhysObject->GetSolverParticleHandles();
+		auto ParticleHandles = [&Collection](int32 Index) -> Chaos::TPBDRigidClusteredParticleHandle<FReal, 3>*
+		{
+			return Collection->PhysObject->GetParticle_Internal(Index);
+		};
 
 		UnitTest.Advance();
 
@@ -715,38 +721,38 @@ namespace GeometryCollectionTest
 			Command.MetaData.Add(FFieldSystemMetaData::EMetaType::ECommandData_ProcessingResolution, TUniquePtr< FFieldSystemMetaDataProcessingResolution >(ResolutionData));
 			UnitTest.Solver->GetPerSolverField().AddTransientCommand(Command);
 
-			EXPECT_TRUE(ParticleHandles[0]->Disabled());
-			EXPECT_TRUE(ParticleHandles[1]->Disabled());
-			EXPECT_TRUE(ParticleHandles[2]->Disabled());
-			EXPECT_TRUE(ParticleHandles[3]->Disabled());
-			EXPECT_TRUE(ParticleHandles[4]->Disabled());
-			EXPECT_TRUE(ParticleHandles[5]->Disabled());
-			EXPECT_TRUE(ParticleHandles[6]->Disabled());
-			EXPECT_TRUE(ParticleHandles[7]->Disabled());
-			EXPECT_FALSE(ParticleHandles[8]->Disabled());
+			EXPECT_TRUE(ParticleHandles(0)->Disabled());
+			EXPECT_TRUE(ParticleHandles(1)->Disabled());
+			EXPECT_TRUE(ParticleHandles(2)->Disabled());
+			EXPECT_TRUE(ParticleHandles(3)->Disabled());
+			EXPECT_TRUE(ParticleHandles(4)->Disabled());
+			EXPECT_TRUE(ParticleHandles(5)->Disabled());
+			EXPECT_TRUE(ParticleHandles(6)->Disabled());
+			EXPECT_TRUE(ParticleHandles(7)->Disabled());
+			EXPECT_FALSE(ParticleHandles(8)->Disabled());
 
 			UnitTest.Advance();
 
 			// todo: indices here might be off but the test crashes before this so we can't validate yet
 			EXPECT_EQ(ClusterMap.Num(), 2);
-			EXPECT_EQ(ClusterMap[ParticleHandles[7]].Num(), 3);
-			EXPECT_TRUE(ClusterMap[ParticleHandles[7]].Contains(ParticleHandles[3]));
-			EXPECT_TRUE(ClusterMap[ParticleHandles[7]].Contains(ParticleHandles[4]));
-			EXPECT_TRUE(ClusterMap[ParticleHandles[7]].Contains(ParticleHandles[5]));
-			EXPECT_EQ(ClusterMap[ParticleHandles[6]].Num(), 3);
-			EXPECT_TRUE(ClusterMap[ParticleHandles[6]].Contains(ParticleHandles[0]));
-			EXPECT_TRUE(ClusterMap[ParticleHandles[6]].Contains(ParticleHandles[1]));
-			EXPECT_TRUE(ClusterMap[ParticleHandles[6]].Contains(ParticleHandles[2]));
-
-			EXPECT_TRUE(ParticleHandles[0]->Disabled());
-			EXPECT_TRUE(ParticleHandles[1]->Disabled());
-			EXPECT_TRUE(ParticleHandles[2]->Disabled());
-			EXPECT_TRUE(ParticleHandles[3]->Disabled());
-			EXPECT_TRUE(ParticleHandles[4]->Disabled());
-			EXPECT_TRUE(ParticleHandles[5]->Disabled());
-			EXPECT_FALSE(ParticleHandles[6]->Disabled());
-			EXPECT_FALSE(ParticleHandles[7]->Disabled());
-			EXPECT_TRUE(ParticleHandles[8]->Disabled());
+			EXPECT_EQ(ClusterMap[ParticleHandles(7)].Num(), 3);
+			EXPECT_TRUE(ClusterMap[ParticleHandles(7)].Contains(ParticleHandles(3)));
+			EXPECT_TRUE(ClusterMap[ParticleHandles(7)].Contains(ParticleHandles(4)));
+			EXPECT_TRUE(ClusterMap[ParticleHandles(7)].Contains(ParticleHandles(5)));
+			EXPECT_EQ(ClusterMap[ParticleHandles(6)].Num(), 3);
+			EXPECT_TRUE(ClusterMap[ParticleHandles(6)].Contains(ParticleHandles(0)));
+			EXPECT_TRUE(ClusterMap[ParticleHandles(6)].Contains(ParticleHandles(1)));
+			EXPECT_TRUE(ClusterMap[ParticleHandles(6)].Contains(ParticleHandles(2)));
+
+			EXPECT_TRUE(ParticleHandles(0)->Disabled());
+			EXPECT_TRUE(ParticleHandles(1)->Disabled());
+			EXPECT_TRUE(ParticleHandles(2)->Disabled());
+			EXPECT_TRUE(ParticleHandles(3)->Disabled());
+			EXPECT_TRUE(ParticleHandles(4)->Disabled());
+			EXPECT_TRUE(ParticleHandles(5)->Disabled());
+			EXPECT_FALSE(ParticleHandles(6)->Disabled());
+			EXPECT_FALSE(ParticleHandles(7)->Disabled());
+			EXPECT_TRUE(ParticleHandles(8)->Disabled());
 		}
 
 		delete FalloffField;
