Chaos - Fix for character ground constraint to handle disabling of constraint particles

diff --git a/Engine/Source/Programs/HeadlessChaos/Private/Character/TestCharacterGroundConstraintProxy.cpp b/Engine/Source/Programs/HeadlessChaos/Private/Character/TestCharacterGroundConstraintProxy.cpp
--- a/Engine/Source/Programs/HeadlessChaos/Private/Character/TestCharacterGroundConstraintProxy.cpp
+++ b/Engine/Source/Programs/HeadlessChaos/Private/Character/TestCharacterGroundConstraintProxy.cpp
@@ -137,7 +137,6 @@ namespace ChaosTest
 		FCharacterGroundConstraint* Constraint_GT = new FCharacterGroundConstraint();
 		FSingleParticlePhysicsProxy* CharacterProxy = CreateParticle();
 		FSingleParticlePhysicsProxy* GroundProxy = CreateParticle();
-		FSingleParticlePhysicsProxy* SecondGroundProxy = CreateParticle();
 		Constraint_GT->Init(CharacterProxy);
 		Constraint_GT->SetGroundParticleProxy(GroundProxy);
 
@@ -171,32 +170,14 @@ namespace ChaosTest
 		ASSERT_EQ(Solver->GetParticles().GetDynamicParticles().ParticleConstraints(0).Num(), 1);
 		EXPECT_EQ(Solver->GetParticles().GetDynamicParticles().ParticleConstraints(0)[0], Constraint_PT);
 
-		ASSERT_EQ(Solver->GetParticles().GetDynamicParticles().ParticleConstraints(1).Num(), 1);
-		EXPECT_EQ(Solver->GetParticles().GetDynamicParticles().ParticleConstraints(1)[0], Constraint_PT);
-
-		ASSERT_EQ(Solver->GetParticles().GetDynamicParticles().ParticleConstraints(2).Num(), 0);
-
-		// Change the ground particle and check that the particle constraints are updated
-
-		Constraint_PT->SetGroundParticle(SecondGroundProxy->GetHandle_LowLevel());
-
-		ASSERT_EQ(Solver->GetParticles().GetDynamicParticles().ParticleConstraints(1).Num(), 0);
-
-		ASSERT_EQ(Solver->GetParticles().GetDynamicParticles().ParticleConstraints(2).Num(), 1);
-		EXPECT_EQ(Solver->GetParticles().GetDynamicParticles().ParticleConstraints(2)[0], Constraint_PT);
-
 		// Destroy the constraint on the physics thread. Ensure that the particle constraints are unregistered
 		Proxy->DestroyOnPhysicsThread(Solver);
 
 		ASSERT_EQ(Solver->GetParticles().GetDynamicParticles().ParticleConstraints(0).Num(), 0);
-		ASSERT_EQ(Solver->GetParticles().GetDynamicParticles().ParticleConstraints(1).Num(), 0);
-		ASSERT_EQ(Solver->GetParticles().GetDynamicParticles().ParticleConstraints(2).Num(), 0);
 
 		// Reinitialize on the physics thread
 		Proxy->InitializeOnPhysicsThread(Solver, Manager, DataIdx, RemoteData);
 		ASSERT_EQ(Solver->GetParticles().GetDynamicParticles().ParticleConstraints(0).Num(), 1);
-		ASSERT_EQ(Solver->GetParticles().GetDynamicParticles().ParticleConstraints(1).Num(), 1);
-		ASSERT_EQ(Solver->GetParticles().GetDynamicParticles().ParticleConstraints(2).Num(), 0);
 
 		// Disable the character particle - constraint should be disabled
 		Solver->GetEvolution()->DisableParticle(CharacterProxy->GetHandle_LowLevel());
@@ -205,4 +186,71 @@ namespace ChaosTest
 		Solver->GetEvolution()->EnableParticle(CharacterProxy->GetHandle_LowLevel());
 		ASSERT_TRUE(Constraint_PT->IsEnabled());
 	}
+
+	// Test disabling constraint particles on the physics thread
+	TEST_F(CharacterGroundConstraintsProxyTest, TestParticleDisable)
+	{
+		// Create a game thread constraint and register it with the solver
+		FCharacterGroundConstraint* Constraint_GT = new FCharacterGroundConstraint();
+		FSingleParticlePhysicsProxy* CharacterProxy = CreateParticle();
+		FSingleParticlePhysicsProxy* GroundProxy = CreateParticle();
+		Constraint_GT->Init(CharacterProxy);
+		Constraint_GT->SetGroundParticleProxy(GroundProxy);
+
+		Solver->RegisterObject(CharacterProxy);
+		Solver->RegisterObject(GroundProxy);
+		Solver->RegisterObject(Constraint_GT);
+
+		// Get the proxy from the game thread constraint
+		// Should have been set when registering the constraint with the solver
+		FCharacterGroundConstraintProxy* Proxy = Constraint_GT->GetProxy<FCharacterGroundConstraintProxy>();
+		ASSERT_TRUE(Proxy != nullptr);
+
+		const int32 DataIdx = 0;
+		FDirtyProxiesBucketInfo BucketInfo;
+		BucketInfo.Num[(uint32)EPhysicsProxyType::CharacterGroundConstraintType] = 1;
+		FDirtyPropertiesManager Manager;
+		Manager.PrepareBuckets(BucketInfo);
+
+		FDirtyChaosProperties RemoteData;
+
+		Proxy->PushStateOnGameThread(Manager, 0, RemoteData);
+
+		// Initialize the physics thread data on the proxy
+		Proxy->InitializeOnPhysicsThread(Solver, Manager, DataIdx, RemoteData);
+
+		FCharacterGroundConstraintHandle* Constraint_PT = Proxy->GetPhysicsThreadAPI();
+		ASSERT_TRUE(Constraint_PT != nullptr);
+
+		// First disable the character particle. This should disable the constraint
+		Solver->GetEvolution()->DisableParticle(Constraint_PT->GetCharacterParticle());
+		EXPECT_FALSE(Constraint_PT->IsEnabled());
+
+		// Now re-enable. The constraint should be re-enabled
+		Solver->GetEvolution()->EnableParticle(Constraint_PT->GetCharacterParticle());
+		EXPECT_TRUE(Constraint_PT->IsEnabled());
+
+		// Disable the ground particle. The constraint should remain enabled with the ground body set to null
+		FGeometryParticleHandle* GroundParticle = Constraint_PT->GetGroundParticle();
+		EXPECT_TRUE(GroundParticle->ParticleConstraints().Contains(Constraint_PT));
+		Solver->GetEvolution()->DisableParticle(GroundParticle);
+		EXPECT_TRUE(Constraint_PT->IsEnabled());
+		EXPECT_EQ(Constraint_PT->GetGroundParticle(), nullptr);
+		EXPECT_FALSE(GroundParticle->ParticleConstraints().Contains(Constraint_PT));
+
+		Chaos::FDirtyCharacterGroundConstraintData DirtyConstraintData;
+		Proxy->BufferPhysicsResults(DirtyConstraintData);
+		Proxy->PullFromPhysicsState(DirtyConstraintData, 0);
+
+		EXPECT_EQ(Constraint_GT->GetGroundParticleProxy(), nullptr);
+
+		// Remove the constraint and particles from the solver
+		Solver->UnregisterObject(Constraint_GT);
+		Solver->UnregisterObject(CharacterProxy);
+		Solver->UnregisterObject(GroundProxy);
+
+		EXPECT_EQ(Proxy->GetGameThreadAPI(), nullptr);
+
+		// Physics thread constraint is deleted in a callback set on the marshalling manager
+	}
 }
\ No newline at end of file
diff --git a/Engine/Source/Runtime/Experimental/Chaos/Private/Chaos/Character/CharacterGroundConstraintContainer.cpp b/Engine/Source/Runtime/Experimental/Chaos/Private/Chaos/Character/CharacterGroundConstraintContainer.cpp
--- a/Engine/Source/Runtime/Experimental/Chaos/Private/Chaos/Character/CharacterGroundConstraintContainer.cpp
+++ b/Engine/Source/Runtime/Experimental/Chaos/Private/Chaos/Character/CharacterGroundConstraintContainer.cpp
@@ -132,6 +132,7 @@ namespace Chaos
 	{
 		for (TGeometryParticleHandle<FReal, 3>*RemovedParticle : RemovedParticles)
 		{
+			TArray<FCharacterGroundConstraintHandle*> ConstraintsWithRemovedGroundBody;
 			for (FConstraintHandle* ConstraintHandle : RemovedParticle->ParticleConstraints())
 			{
 				if (FCharacterGroundConstraintHandle* Constraint = ConstraintHandle->As<FCharacterGroundConstraintHandle>())
@@ -147,7 +148,7 @@ namespace Chaos
 						}
 						else if (Constraint->GroundParticle == RemovedParticle)
 						{
-							Constraint->GroundParticle = nullptr;
+							ConstraintsWithRemovedGroundBody.Add(Constraint);
 						}
 					}
 					else
@@ -156,6 +157,54 @@ namespace Chaos
 					}
 				}
 			}
+
+			for (FCharacterGroundConstraintHandle* Constraint : ConstraintsWithRemovedGroundBody)
+			{
+				Constraint->SetGroundParticle(nullptr);
+			}
+		}
+	}
+
+	void FCharacterGroundConstraintContainer::OnDisableParticle(FGeometryParticleHandle* DisabledParticle)
+	{
+		// Only disable the constraint if it's the character particle being disabled.
+		// If the ground particle is disabled set it to null. This will remove the constraint
+		// from the disabled particle's constraint list, so if the ground particle is re-enabled
+		// it won't automatically be set as the ground particle in the constraint
+		for (FCharacterGroundConstraintHandle* Constraint : GetConstraints())
+		{
+			if (Constraint->IsValid())
+			{
+				if (Constraint->GetCharacterParticle() == DisabledParticle)
+				{
+					if (Constraint->IsEnabled())
+					{
+						Constraint->SetEnabled(false);
+					}
+				}
+				else if (Constraint->GetGroundParticle() == DisabledParticle)
+				{
+					Constraint->SetGroundParticle(nullptr);
+				}
+			}
+			else
+			{
+				Constraint->SetEnabled(false);
+			}
+		}
+	}
+
+	void FCharacterGroundConstraintContainer::OnEnableParticle(FGeometryParticleHandle* EnabledParticle)
+	{
+		for (FConstraintHandle* ConstraintHandle : EnabledParticle->ParticleConstraints())
+		{
+			if (FCharacterGroundConstraintHandle* Constraint = ConstraintHandle->As<FCharacterGroundConstraintHandle>())
+			{
+				if ((Constraint->GetCharacterParticle() == EnabledParticle) && !Constraint->IsEnabled())
+				{
+					Constraint->SetEnabled(true);
+				}
+			}
 		}
 	}
 
diff --git a/Engine/Source/Runtime/Experimental/Chaos/Private/Chaos/PBDConstraintContainer.cpp b/Engine/Source/Runtime/Experimental/Chaos/Private/Chaos/PBDConstraintContainer.cpp
--- a/Engine/Source/Runtime/Experimental/Chaos/Private/Chaos/PBDConstraintContainer.cpp
+++ b/Engine/Source/Runtime/Experimental/Chaos/Private/Chaos/PBDConstraintContainer.cpp
@@ -14,4 +14,26 @@ namespace Chaos
 	{
 	}
 
+	void FPBDConstraintContainer::OnDisableParticle(FGeometryParticleHandle* DisabledParticle)
+	{
+		for (FConstraintHandle* ConstraintHandle : DisabledParticle->ParticleConstraints())
+		{
+			if ((ConstraintHandle->GetContainerId() == ContainerId) && ConstraintHandle->IsEnabled())
+			{
+				ConstraintHandle->SetEnabled(false);
+			}
+		}
+	}
+
+	void FPBDConstraintContainer::OnEnableParticle(FGeometryParticleHandle* EnabledParticle)
+	{
+		for (FConstraintHandle* ConstraintHandle : EnabledParticle->ParticleConstraints())
+		{
+			if ((ConstraintHandle->GetContainerId() == ContainerId) && !ConstraintHandle->IsEnabled())
+			{
+				ConstraintHandle->SetEnabled(true);
+			}
+		}
+	}
+
 }
diff --git a/Engine/Source/Runtime/Experimental/Chaos/Private/PhysicsProxy/CharacterGroundConstraintProxy.cpp b/Engine/Source/Runtime/Experimental/Chaos/Private/PhysicsProxy/CharacterGroundConstraintProxy.cpp
--- a/Engine/Source/Runtime/Experimental/Chaos/Private/PhysicsProxy/CharacterGroundConstraintProxy.cpp
+++ b/Engine/Source/Runtime/Experimental/Chaos/Private/PhysicsProxy/CharacterGroundConstraintProxy.cpp
@@ -138,7 +138,7 @@ namespace Chaos
 	void FCharacterGroundConstraintProxy::BufferPhysicsResults(FDirtyCharacterGroundConstraintData& Buffer)
 	{
 		Buffer.SetProxy(*this);
-		if (Constraint_PT != nullptr && Constraint_PT->IsValid())
+		if (Constraint_PT != nullptr && Constraint_PT->IsValid() && Constraint_PT->IsEnabled())
 		{
 			Buffer.Force = Constraint_PT->GetSolverAppliedForce();
 			Buffer.Torque = Constraint_PT->GetSolverAppliedTorque();
@@ -164,7 +164,7 @@ namespace Chaos
 				Data.TargetDeltaFacing = Buffer.TargetDeltaFacing;
 				});
 
-			if (Buffer.GroundParticle)
+			if (Buffer.GroundParticle && (Buffer.GroundParticle->ParticleIdx != INDEX_NONE))
 			{
 				FSingleParticlePhysicsProxy* NewGroundProxy = (FSingleParticlePhysicsProxy*)(Buffer.GroundParticle->PhysicsProxy());
 				Constraint_GT->GroundProxy.Modify(false, Constraint_GT->DirtyFlags, Constraint_GT->Proxy, [NewGroundProxy](FParticleProxyProperty& Data)
diff --git a/Engine/Source/Runtime/Experimental/Chaos/Public/Chaos/Character/CharacterGroundConstraintContainer.h b/Engine/Source/Runtime/Experimental/Chaos/Public/Chaos/Character/CharacterGroundConstraintContainer.h
--- a/Engine/Source/Runtime/Experimental/Chaos/Public/Chaos/Character/CharacterGroundConstraintContainer.h
+++ b/Engine/Source/Runtime/Experimental/Chaos/Public/Chaos/Character/CharacterGroundConstraintContainer.h
@@ -146,6 +146,8 @@ namespace Chaos
 		CHAOS_API virtual void PrepareTick() override final;
 		CHAOS_API virtual void UnprepareTick() override final;
 		CHAOS_API virtual void DisconnectConstraints(const TSet<TGeometryParticleHandle<FReal, 3>*>& RemovedParticles) override final;
+		CHAOS_API virtual void OnDisableParticle(FGeometryParticleHandle* DisabledParticle) override final;
+		CHAOS_API virtual void OnEnableParticle(FGeometryParticleHandle* EnabledParticle) override final;
 
 		//////////////////////////////////////////////////////////////////////////
 		// Required API from FConstraintContainer
diff --git a/Engine/Source/Runtime/Experimental/Chaos/Public/Chaos/PBDConstraintContainer.h b/Engine/Source/Runtime/Experimental/Chaos/Public/Chaos/PBDConstraintContainer.h
--- a/Engine/Source/Runtime/Experimental/Chaos/Public/Chaos/PBDConstraintContainer.h
+++ b/Engine/Source/Runtime/Experimental/Chaos/Public/Chaos/PBDConstraintContainer.h
@@ -73,9 +73,16 @@ namespace Chaos
 		*/
 		virtual void UnprepareTick() = 0;
 
+		// This is called when a particle is destroyed. The constraint should be disabled and the particle in the constraint set to null
 		// @todo(chaos): remove the set
 		virtual void DisconnectConstraints(const TSet<TGeometryParticleHandle<FReal, 3>*>&) {}
 
+		// This is called to notify the constraint container that a particle has been disabled
+		virtual void OnDisableParticle(FGeometryParticleHandle* DisabledParticle);
+
+		// This is called to notify the constraint container that a particle has been disabled
+		virtual void OnEnableParticle(FGeometryParticleHandle* EnabledParticle);
+
 		/**
 		* Create a constraint solver for an Evolution without Graph support (RBAN evolution).
 		* There will only be one of these per scene (RBAN node) and it is used to solve all constraints
diff --git a/Engine/Source/Runtime/Experimental/Chaos/Public/Chaos/PBDRigidsEvolution.h b/Engine/Source/Runtime/Experimental/Chaos/Public/Chaos/PBDRigidsEvolution.h
--- a/Engine/Source/Runtime/Experimental/Chaos/Public/Chaos/PBDRigidsEvolution.h
+++ b/Engine/Source/Runtime/Experimental/Chaos/Public/Chaos/PBDRigidsEvolution.h
@@ -645,29 +645,23 @@ public:
 	*/
 	void DisableConstraints(FGeometryParticleHandle* ParticleHandle)
 	{
-		RemoveConstraintsFromConstraintGraph(ParticleHandle->ParticleConstraints());
-
-		for (FConstraintHandle* Constraint : ParticleHandle->ParticleConstraints())
+		for (FPBDConstraintContainer* Container : ConstraintContainers)
 		{
-			if (Constraint->IsEnabled())
-			{
-				Constraint->SetEnabled(false);
-			}
+			Container->OnDisableParticle(ParticleHandle);
 		}
+
+		RemoveConstraintsFromConstraintGraph(ParticleHandle->ParticleConstraints());
 	}
 
 	/** 
 	* Enable constraints (all types except collisions) from the enabled particles; constraints will only become enabled if their particle end points are valid.
-	* @note This only applies to persistent constraints (joints etc), not transient constraints (collisons)
+	* @note This only applies to persistent constraints (joints etc), not transient constraints (collisions)
 	*/
 	void EnableConstraints(FGeometryParticleHandle* ParticleHandle)
 	{
-		for (FConstraintHandle* Constraint : ParticleHandle->ParticleConstraints())
+		for (FPBDConstraintContainer* Container : ConstraintContainers)
 		{
-			if (!Constraint->IsEnabled())
-			{
-				Constraint->SetEnabled(true);
-			}
+			Container->OnEnableParticle(ParticleHandle);
 		}
 	}
 
