bugfix Chaos Overlaps: Transforms including reflection did not always function correctly
#UDN 00712852

diff --git a/Engine/Source/Programs/HeadlessChaos/Private/HeadlessChaosTestOverlap.cpp b/Engine/Source/Programs/HeadlessChaos/Private/HeadlessChaosTestOverlap.cpp
--- a/Engine/Source/Programs/HeadlessChaos/Private/HeadlessChaosTestOverlap.cpp
+++ b/Engine/Source/Programs/HeadlessChaos/Private/HeadlessChaosTestOverlap.cpp
@@ -1092,6 +1092,22 @@ namespace ChaosTest
 					EXPECT_EQ(bResult, bResultMTD);					
 				}
 
+				// Mirrored and Rotated box, mirrored trimesh // regression test
+				{
+					FBoxPtr BoxSafe2(new TBox<FReal, 3>(FVec3(-100.0, -100.0, -100.0), FVec3(100.0, 100.0, 100.0)));
+					FVec3 TriMeshScale2 = { -10.0f, 10.0f, 10.0f };
+					FVec3 InvTriMeshScale2 = 1.0 / TriMeshScale2;
+					FVec3 BoxScale2 = { 1, 1, 1 };
+					BoxScale2 = BoxScale2 * InvTriMeshScale2; // Box needs to be in mesh space // TODO make sure other tests are correct
+					TImplicitObjectScaled<TBox<FReal, 3>> ScaledBox2 = TImplicitObjectScaled<TBox<FReal, 3>>(BoxSafe2, BoxScale2);
+
+					FRigidTransform3 QueryTM(FVec3(0, 0.0, 105.0) * InvTriMeshScale2, FQuat{ 1, 0, 0, 0 }); // Pi rotation around x axis
+					bool bResult = TriangleMesh->OverlapGeom(ScaledBox2, QueryTM, 0.0, nullptr, TriMeshScale2);
+					EXPECT_EQ(bResult, true);
+					bool bResultMTD = TriangleMesh->OverlapGeom(ScaledBox2, QueryTM, 0.0, &MTDInfo, TriMeshScale2);
+					EXPECT_EQ(bResult, bResultMTD);
+				}
+
 				// Unrotated unscaled box
 				{
 					FRigidTransform3 QueryTM(FVec3(0, 0.0, 11.0), FQuat::Identity);
@@ -1661,7 +1677,7 @@ namespace ChaosTest
 				}
 			}
 		}
-{
+		{
 			FTriangleMeshImplicitObject::ParticlesType TrimeshParticles(
 				{
 					{0.0, 0.0, 0.0},
diff --git a/Engine/Source/Runtime/Experimental/Chaos/Private/Chaos/OBBVectorized.cpp b/Engine/Source/Runtime/Experimental/Chaos/Private/Chaos/OBBVectorized.cpp
--- a/Engine/Source/Runtime/Experimental/Chaos/Private/Chaos/OBBVectorized.cpp
+++ b/Engine/Source/Runtime/Experimental/Chaos/Private/Chaos/OBBVectorized.cpp
@@ -30,7 +30,7 @@ namespace Chaos
 
 		VectorRegister4Float FurthestPoint = VectorMultiplyAdd(VectorAbs(XAxis), XHalfExtent, VectorMultiplyAdd(VectorAbs(YAxis), YHalfExtent, VectorMultiply(VectorAbs(ZAxis), ZHalfExtent)));
 
-		FurthestPoint = VectorMultiply(FurthestPoint, InvScale);
+		FurthestPoint = VectorMultiply(FurthestPoint, VectorAbs(InvScale));
 		MaxObb = VectorAdd(Position, FurthestPoint);
 		MinObb = VectorSubtract(Position, FurthestPoint);
 
