Chaos - fix collisions being destroyed when particles are explicitly slept.

This bug prevented islands from being awoken if you explicitly set all of its particles to sleep, and then teleported one of the particles elsewhere.

This was caused by SleepParticle being called before BeginDetectCollisions. This put the particle to sleep but left the collision constraints as they were. When we prune out of date collisions they were destroyed because they were "awake" but were not updated this tick.

The fix is to update the sleep state of constraints whose islands will be put to sleep as a result of all particles being put to sleep.

diff --git a/Engine/Source/Programs/HeadlessChaos/Private/HeadlessChaosTestGraphEvolution.cpp b/Engine/Source/Programs/HeadlessChaos/Private/HeadlessChaosTestGraphEvolution.cpp
--- a/Engine/Source/Programs/HeadlessChaos/Private/HeadlessChaosTestGraphEvolution.cpp
+++ b/Engine/Source/Programs/HeadlessChaos/Private/HeadlessChaosTestGraphEvolution.cpp
@@ -379,8 +379,8 @@ namespace ChaosTest
 	}
 
 	// Start with an island containing 4 particle connected in a chain, then make the middle two kinematic.
-	// This makes the B-C constraint kinematic ("invalid" in the Constraint Graph terminolgy) which means
-	// it does not belong in any island and is kicked out of the graph (the edge is deleted).
+	// This makes the B-C constraint kinematic which means it does not belong in any island and is kicked 
+	// out of the graph (the edge is deleted).
 	// Check that the island manager handles kinematic-kinematic constraints
 	//		A-B-C-D 
 	// =>	A-B  C-D
@@ -697,80 +697,38 @@ namespace ChaosTest
 		EXPECT_FALSE(FConstGenericParticleHandle(Test.ParticleHandles[3])->Sleeping());
 	}
 
-	// Start with an island containing 4 sleeping particles connected in a chain. Then wake the island
-	// using the WakeIsland method.
-	//
-	//		(d=dynamic, s=sleeping, k=kinematic)
-	//		As - Bs - Cs - Ds  =>	Ad - Bd - Cd - Dd
-	//
-	// @todo(chaos): Explicit waking of islands is not currently supported
-	GTEST_TEST(GraphEvolutionTests, DISABLED_TestConstraintGraph_WakeIsland)
-	{
-		FGraphEvolutionTest Test(4);
-		Test.MakeChain();
-
-		Test.AdvanceUntilSleeping();
-
-		// All particles asleep
-		EXPECT_TRUE(FConstGenericParticleHandle(Test.ParticleHandles[0])->Sleeping());
-		EXPECT_TRUE(FConstGenericParticleHandle(Test.ParticleHandles[1])->Sleeping());
-		EXPECT_TRUE(FConstGenericParticleHandle(Test.ParticleHandles[2])->Sleeping());
-		EXPECT_TRUE(FConstGenericParticleHandle(Test.ParticleHandles[3])->Sleeping());
-		EXPECT_EQ(Test.IslandManager->GetNumIslands(), 1);
-
-		// Flag the island for wake up
-		if (Test.IslandManager->GetNumIslands() == 1)
-		{
-			EXPECT_TRUE(Test.IslandManager->GetIsland(0)->IsSleeping());
-			EXPECT_TRUE(false);	// Should be calling WakeIsland here
-			//Test.IslandManager->WakeIsland(Test.Particles, 0);
-		}
-
-		Test.Advance();
-
-		// Island and all particles should now be awake
-		EXPECT_EQ(Test.IslandManager->GetNumIslands(), 1);
-		if (Test.IslandManager->GetNumIslands() == 1)
-		{
-			EXPECT_FALSE(Test.IslandManager->GetIsland(0)->IsSleeping());
-		}
-		EXPECT_FALSE(FConstGenericParticleHandle(Test.ParticleHandles[0])->Sleeping());
-		EXPECT_FALSE(FConstGenericParticleHandle(Test.ParticleHandles[1])->Sleeping());
-		EXPECT_FALSE(FConstGenericParticleHandle(Test.ParticleHandles[2])->Sleeping());
-		EXPECT_FALSE(FConstGenericParticleHandle(Test.ParticleHandles[3])->Sleeping());
-	}
-
 	// Start with an island containing 4 awake particles connected in a chain. Then sleep the island
-	// using the SleepIsland method.
+	// by explicitly putting all particles to sleep
 	//
 	//		(d=dynamic, s=sleeping, k=kinematic)
 	//		Ad - Bd - Cd - Dd  =>	As - Bs - Cs - Ds
 	//
-	// @todo(chaos): this test fails. Explicit call to SleepIsland do not work atm
-	GTEST_TEST(GraphEvolutionTests, DISABLED_TestConstraintGraph_SleepIsland)
+	GTEST_TEST(GraphEvolutionTests, TestConstraintGraph_SleepIsland)
 	{
 		FGraphEvolutionTest Test(4);
 		Test.MakeChain();
 
 		Test.Advance();
 
-		// All particles awake
+		// All particles and costraints are awake
 		EXPECT_FALSE(FConstGenericParticleHandle(Test.ParticleHandles[0])->Sleeping());
 		EXPECT_FALSE(FConstGenericParticleHandle(Test.ParticleHandles[1])->Sleeping());
 		EXPECT_FALSE(FConstGenericParticleHandle(Test.ParticleHandles[2])->Sleeping());
 		EXPECT_FALSE(FConstGenericParticleHandle(Test.ParticleHandles[3])->Sleeping());
+		EXPECT_FALSE(Test.ConstraintHandles[0]->IsSleeping());
+		EXPECT_FALSE(Test.ConstraintHandles[1]->IsSleeping());
+		EXPECT_FALSE(Test.ConstraintHandles[2]->IsSleeping());
 		EXPECT_EQ(Test.IslandManager->GetNumIslands(), 1);
 
-		// Flag the island for sleep
-		if (Test.IslandManager->GetNumIslands() == 1)
-		{
-			EXPECT_FALSE(Test.IslandManager->GetIsland(0)->IsSleeping());
-			Test.IslandManager->SetParticleIslandIsSleeping(Test.ParticleHandles[0], true);
-		}
+		// Put all of the particles to sleep
+		Test.Evolution.SetParticleObjectState(Test.ParticleHandles[0], EObjectStateType::Sleeping);
+		Test.Evolution.SetParticleObjectState(Test.ParticleHandles[1], EObjectStateType::Sleeping);
+		Test.Evolution.SetParticleObjectState(Test.ParticleHandles[2], EObjectStateType::Sleeping);
+		Test.Evolution.SetParticleObjectState(Test.ParticleHandles[3], EObjectStateType::Sleeping);
 
 		Test.Advance();
 
-		// Island and all particles should now be asleep
+		// Island and all particles and constraints should now be asleep
 		EXPECT_EQ(Test.IslandManager->GetNumIslands(), 1);
 		if (Test.IslandManager->GetNumIslands() == 1)
 		{
@@ -780,6 +738,23 @@ namespace ChaosTest
 		EXPECT_TRUE(FConstGenericParticleHandle(Test.ParticleHandles[1])->Sleeping());
 		EXPECT_TRUE(FConstGenericParticleHandle(Test.ParticleHandles[2])->Sleeping());
 		EXPECT_TRUE(FConstGenericParticleHandle(Test.ParticleHandles[3])->Sleeping());
+
+		EXPECT_TRUE(Test.ConstraintHandles[0]->IsInConstraintGraph());
+		EXPECT_TRUE(Test.ConstraintHandles[1]->IsInConstraintGraph());
+		EXPECT_TRUE(Test.ConstraintHandles[2]->IsInConstraintGraph());
+
+		EXPECT_TRUE(Test.ConstraintHandles[0]->IsSleeping());
+		EXPECT_TRUE(Test.ConstraintHandles[1]->IsSleeping());
+		EXPECT_TRUE(Test.ConstraintHandles[2]->IsSleeping());
+	}
+
+	// @todo(chaos): Implement this. TestConstraintGraph_SleepIsland is intended to reproduce a bug 
+	// exp[osed by collision constraints where collisions were being destroyed on particles that were 
+	// explicitly put to sleep. However that bug was a result of how collision constraints are 
+	// destroyed (i.e., when they are not updated this tick) and the NullConstraints don't have that same 
+	// behaviour. We need a unit testing constraint that can reproduce that behaviour. 
+	GTEST_TEST(GraphEvolutionTests, DISABLED_TestConstraintGraph_SleepIsland_Collisions)
+	{
 	}
 
 	// Add a constraint between a sleeping and a kinematic body and tick.
diff --git a/Engine/Source/Runtime/Experimental/Chaos/Private/Chaos/Island/IslandManager.cpp b/Engine/Source/Runtime/Experimental/Chaos/Private/Chaos/Island/IslandManager.cpp
--- a/Engine/Source/Runtime/Experimental/Chaos/Private/Chaos/Island/IslandManager.cpp
+++ b/Engine/Source/Runtime/Experimental/Chaos/Private/Chaos/Island/IslandManager.cpp
@@ -373,6 +373,9 @@ namespace Chaos::Private
 		if (Constraint != nullptr)
 		{
 			Constraint->SetConstraintGraphEdge(this);
+
+			// Initialize edge state to match constraint state
+			Flags.bIsSleeping = Constraint->IsSleeping();
 		}
 	}
 
@@ -744,6 +747,9 @@ namespace Chaos::Private
 		// When we explicitly wake we reset sleep counters etc
 		if (FPBDIslandParticle* Node = GetGraphNode(Particle))
 		{
+			// Make sure the node sleep state matches the particle state
+			Node->Flags.bIsSleeping = FConstGenericParticleHandle(Particle)->IsSleeping();
+
 			// NOTE: We could check Flags.bIsDynamic here, but checking the Island pointer means we 
 			// are automatically handling the case where a kinematic with constraints on it 
 			// was made dynamic right before calling WakeParticleIsland
@@ -770,6 +776,23 @@ namespace Chaos::Private
 		}
 	}
 
+	void FPBDIslandManager::SleepParticle(FGeometryParticleHandle* Particle)
+	{
+		// When we explicitly sleep we must check the island for sleeping before we update constraints
+		// so that constraint sleep state matched their particles' sleep state.
+		const bool bIsSleepAllowed = true;
+		if (FPBDIslandParticle* Node = GetGraphNode(Particle))
+		{
+			Node->Flags.bIsSleeping = true;
+
+			if ((Node->Island != nullptr) && !Node->Island->Flags.bIsSleeping)
+			{
+				// Set the checksleep flag for processing in See UpdateParticlesCheckSleep
+				EnqueueIslandCheckSleep(Node->Island, bIsSleepAllowed);
+			}
+		}
+	}
+
 	void FPBDIslandManager::AddParticle(FGeometryParticleHandle* Particle)
 	{
 		// Particles get auto-added via AddConstraint (see CreateGraphEdge, GetOrCreateGraphNode)
@@ -933,19 +956,6 @@ namespace Chaos::Private
 		}
 	}
 
-	void FPBDIslandManager::SetParticleIslandIsSleeping(FGeometryParticleHandle* Particle, const bool bInIsSleeping)
-	{
-		if (FPBDIslandParticle* Node = GetGraphNode(Particle))
-		{
-			if ((Node->Island != nullptr) && (Node->Island->Flags.bIsSleeping != bInIsSleeping))
-			{
-				Node->Island->Flags.bIsSleeping = bInIsSleeping;
-				Node->Island->SleepCounter = 0;
-				PropagateIslandSleep(Node->Island);
-			}
-		}
-	}
-
 	void FPBDIslandManager::UpdateIslands()
 	{
 		ProcessIslands();
@@ -2067,6 +2077,52 @@ namespace Chaos::Private
 		Validate();
 	}
 
+	void FPBDIslandManager::UpdateExplicitSleep()
+	{
+		// UpdateExplicitSleep should be called after processing physics inputs that may change the sleep/wake status of particles. 
+		// It is required so that the Particle/Constraint/Island sleep states are in sync when we Integrate and Detection Collisions.
+		// E.g., If we do not do this and a particle was explicitly put to sleep, it will not be integrated and no collisions will 
+		// be detected. If that particle is in an island with other awake particles it will be woken immediately but will not have 
+		// moved or have any collisions.
+
+		const auto ShouldIslandSleep = [](FPBDIsland * Island) -> bool
+		{
+			for (FPBDIslandParticle* Node : Island->GetParticles())
+			{
+				const bool bIsDynamic = IsParticleDynamic(Node->GetParticle());
+				const bool bIsSleeping = IsParticleSleeping(Node->GetParticle());
+				if (bIsDynamic && !bIsSleeping)
+				{
+					return false;
+				}
+			}
+			return true;
+		};
+
+		// If we have explicitly made some particles go to sleep (as opposed to them naturally sleeping based on low movement)
+		// we must check to see if the whole island can go to sleep. This is primarily to address the issue that we do not
+		// update collisions for sleeping particles and destroy collisions in awake islands that are not updated this tick.
+		for (FPBDIsland* Island : Islands)
+		{
+			if (!!Island->Flags.bCheckSleep)
+			{
+				const bool bIslandShouldSleep = !!Island->Flags.bIsSleepAllowed && ShouldIslandSleep(Island);
+				if (bIslandShouldSleep != Island->Flags.bIsSleeping)
+				{
+					Island->Flags.bIsSleeping = bIslandShouldSleep;
+					Island->SleepCounter = 0;
+				}
+
+				// NOTE: we only get here if particle sleep state was changed. We need to ensure that
+				// constraint sleep state matches the particle sleep state. Ideally we would only
+				// do this if we know that they don't match, but that's hard to determine.
+				PropagateIslandSleep(Island);
+
+				Island->Flags.bCheckSleep = false;
+			}
+		}
+	}
+
 	void FPBDIslandManager::ProcessSleep(const FRealSingle Dt)
 	{
 		if (!CVars::bChaosSolverSleepEnabled)
@@ -2205,6 +2261,12 @@ namespace Chaos::Private
 	}
 
 	void FPBDIslandManager::PropagateIslandSleep(FPBDIsland* Island)
+	{
+		PropagateIslandSleepToParticles(Island);
+		PropagateIslandSleepToConstraints(Island);
+	}
+
+	void FPBDIslandManager::PropagateIslandSleepToParticles(FPBDIsland* Island)
 	{
 		bool bRebuildViews = false;
 		for (FPBDIslandParticle* IslandNode : Island->Nodes)
@@ -2247,7 +2309,10 @@ namespace Chaos::Private
 		{
 			Particles.RebuildViews();
 		}
+	}
 
+	void FPBDIslandManager::PropagateIslandSleepToConstraints(FPBDIsland* Island)
+	{
 		// Set the constraint sleep state to match
 		for (TArray<FPBDIslandConstraint*>& IslandEdges : Island->ContainerEdges)
 		{
diff --git a/Engine/Source/Runtime/Experimental/Chaos/Private/Chaos/PBDRigidsEvolution.cpp b/Engine/Source/Runtime/Experimental/Chaos/Private/Chaos/PBDRigidsEvolution.cpp
--- a/Engine/Source/Runtime/Experimental/Chaos/Private/Chaos/PBDRigidsEvolution.cpp
+++ b/Engine/Source/Runtime/Experimental/Chaos/Private/Chaos/PBDRigidsEvolution.cpp
@@ -1049,6 +1049,7 @@ namespace Chaos
 		const EObjectStateType InitialState = Particle->ObjectState();
 		const bool bWasDynamic = (InitialState == EObjectStateType::Dynamic) || (InitialState == EObjectStateType::Sleeping);
 		const bool bIsDynamic = (ObjectState == EObjectStateType::Dynamic) || (ObjectState == EObjectStateType::Sleeping);
+		const bool bIsSleeping = (ObjectState == EObjectStateType::Sleeping);
 
 		Particle->SetObjectStateLowLevel(ObjectState);
 
@@ -1079,6 +1080,15 @@ namespace Chaos
 				Particle->SetSleepCounter(0);
 			}
 
+			// If all particles in an island are explicitly put to sleep, the island and all its
+			// constraints need to sleep at the start of the next tick. This is because we don't
+			// detect collisions on sleeping particles, and we destroy collisions that are 
+			// awake and were not updated this tick.
+			if (bIsSleeping)
+			{
+				IslandManager.SleepParticle(Particle);
+			}
+
 			// If we are not kinematic, the MovingKinematic flag should be cleared
 			// If we are kinematic and are (or become) moving, the flag will be updated in ApplyKinematicTargets
 			// Either way, we can reset the flag here
@@ -1101,14 +1111,16 @@ namespace Chaos
 
 	void FPBDRigidsEvolutionBase::WakeParticle(FPBDRigidParticleHandle* Particle)
 	{
-		if (Particle->IsDynamic())
+		if (Particle->IsSleeping())
 		{
+			// Set to dynamic - this will also wake the particle's island
 			SetParticleObjectState(Particle, EObjectStateType::Dynamic);
-
-			// Reset the sleep counters etc in the sleep management system
-			// @todo(chaos): ideally we would just set reset the sleepiness of the particle
-			IslandManager.WakeParticleIslands(Particle);
 		}
+
+		// Explicitly waking a particle should reset any sleep state, even if we are already awake.
+		// E.g., this will reset the sleep counter to prevent the particle from immediately sleeping again.
+		// NOTE: This also allows us to "wake" a kinematic which will wake all islands that the kinematic is in.
+		IslandManager.WakeParticleIslands(Particle);
 	}
 
 	void FPBDRigidsEvolutionBase::SetParticleSleepType(FPBDRigidParticleHandle* Particle, ESleepType InSleepType)
diff --git a/Engine/Source/Runtime/Experimental/Chaos/Private/PBDRigidsSolver.cpp b/Engine/Source/Runtime/Experimental/Chaos/Private/PBDRigidsSolver.cpp
--- a/Engine/Source/Runtime/Experimental/Chaos/Private/PBDRigidsSolver.cpp
+++ b/Engine/Source/Runtime/Experimental/Chaos/Private/PBDRigidsSolver.cpp
@@ -454,6 +454,10 @@ namespace Chaos
 				// the next broadphase, we'll pass the epoch check but have an invalid index into a now-empty active constraints array.
 				MSolver->GetEvolution()->GetRigidClustering().UnionClusterGroups();
 
+				// Process any sleep/wake requests that came from the game thread
+				// NOTE: Must be before GetCollisionConstraints().BeginFrame() (because its behaviour depends on sleep state)
+				MSolver->GetEvolution()->GetIslandManager().UpdateExplicitSleep();
+
 				// clear out the collision constraints as they will be stale from last frame if AdvanceOneTimeStep never gets called due to TimeRemaining being less than MinDeltaTime 
 				// @todo(chaos): maybe we can pull data at a better time instead to avoid collision-specific code here for event dispatch
 				MSolver->GetEvolution()->GetCollisionConstraints().BeginFrame();
diff --git a/Engine/Source/Runtime/Experimental/Chaos/Public/Chaos/Island/IslandManager.h b/Engine/Source/Runtime/Experimental/Chaos/Public/Chaos/Island/IslandManager.h
--- a/Engine/Source/Runtime/Experimental/Chaos/Public/Chaos/Island/IslandManager.h
+++ b/Engine/Source/Runtime/Experimental/Chaos/Public/Chaos/Island/IslandManager.h
@@ -402,6 +402,7 @@ namespace Chaos::Private
 		CHAOS_API void UpdateParticleMaterial(FGeometryParticleHandle* Particle);
 		CHAOS_API int32 GetParticleLevel(FGeometryParticleHandle* Particle) const;
 		CHAOS_API void WakeParticleIslands(FGeometryParticleHandle* Particle);
+		CHAOS_API void SleepParticle(FGeometryParticleHandle* Particle);
 
 		// Constraint management
 		int32 GetNumConstraints() const { return Edges.Num(); }
@@ -429,6 +430,7 @@ namespace Chaos::Private
 		CHAOS_API void SetDisableCounterThreshold(const int32 InDisableCounterThreshold);
 		CHAOS_API void SetIsDeterministic(const bool bInIsDeterministic);
 		CHAOS_API void SetAssignLevels(const bool bInAssignLevels);
+		CHAOS_API void UpdateExplicitSleep();
 		CHAOS_API void UpdateParticles();
 		CHAOS_API void UpdateIslands();
 		CHAOS_API void UpdateSleep(const FReal Dt = 0);
@@ -467,7 +469,6 @@ namespace Chaos::Private
 		CHAOS_API TArray<const FPBDIsland*> FindParticleIslands(const FGeometryParticleHandle* Particle) const;
 		CHAOS_API TArray<const FGeometryParticleHandle*> FindParticlesInIslands(const TArray<const FPBDIsland*> Islands) const;
 		CHAOS_API TArray<const FConstraintHandle*> FindConstraintsInIslands(const TArray<const FPBDIsland*> Islands, int32 ContainerId) const;
-		CHAOS_API void SetParticleIslandIsSleeping(FGeometryParticleHandle* Particle, const bool bInIsSleeping);
 
 		// Debug draw
 #if CHAOS_DEBUG_DRAW
@@ -547,6 +548,8 @@ namespace Chaos::Private
 		CHAOS_API void ProcessParticlesSleep(const FRealSingle Dt);
 		CHAOS_API void ProcessIslandSleep(FPBDIsland* Island, const FRealSingle Dt);
 		CHAOS_API void PropagateIslandSleep(FPBDIsland* Island);
+		CHAOS_API void PropagateIslandSleepToParticles(FPBDIsland* Island);
+		CHAOS_API void PropagateIslandSleepToConstraints(FPBDIsland* Island);
 
 		// Disabling
 		CHAOS_API void ProcessDisable(TFunctionRef<void(FPBDRigidParticleHandle*)> ParticleDisableFunctor);
