Remove deprecated waning in Chaos Headless

diff --git a/Engine/Source/Programs/HeadlessChaos/Private/GeometryCollection/GeometryCollectionTestClustering.cpp b/Engine/Source/Programs/HeadlessChaos/Private/GeometryCollection/GeometryCollectionTestClustering.cpp
--- a/Engine/Source/Programs/HeadlessChaos/Private/GeometryCollection/GeometryCollectionTestClustering.cpp
+++ b/Engine/Source/Programs/HeadlessChaos/Private/GeometryCollection/GeometryCollectionTestClustering.cpp
@@ -106,8 +106,8 @@ namespace GeometryCollectionTest
 
 		auto& Clustering = UnitTest.Solver->GetEvolution()->GetRigidClustering();
 		const auto & ClusterMap = Clustering.GetChildrenMap();
-		EXPECT_TRUE(ClusterMapContains(ClusterMap, Collection->PhysObject->GetSolverClusterHandles()[0], 
-			{ Collection->PhysObject->GetSolverParticleHandles()[0],Collection->PhysObject->GetSolverParticleHandles()[1] }));
+		EXPECT_TRUE(ClusterMapContains(ClusterMap, Collection->PhysObject->GetSolverClusterHandle_Internal(0),
+			{ Collection->PhysObject->GetParticle_Internal(0),Collection->PhysObject->GetParticle_Internal(1) }));
 
 		FReal InitialZ = Collection->RestCollection->Transform[2].GetTranslation().Z;
 		for (int Frame = 1; Frame < 10; Frame++)
@@ -123,8 +123,8 @@ namespace GeometryCollectionTest
 			EXPECT_LT(Collection->DynamicCollection->GetTransform(2).GetTranslation().Z, InitialZ); // body should be falling and decreasing in Z			
 		}
 
-		EXPECT_TRUE(ClusterMapContains(ClusterMap, Collection->PhysObject->GetSolverClusterHandles()[0],
-			{ Collection->PhysObject->GetSolverParticleHandles()[0],Collection->PhysObject->GetSolverParticleHandles()[1] }));
+		EXPECT_TRUE(ClusterMapContains(ClusterMap, Collection->PhysObject->GetSolverClusterHandle_Internal(0),
+			{ Collection->PhysObject->GetParticle_Internal(0),Collection->PhysObject->GetParticle_Internal(1) }));
 
 	}
 
@@ -225,7 +225,17 @@ namespace GeometryCollectionTest
 
 		UnitTest.Advance();
 
-		TArray<Chaos::TPBDRigidClusteredParticleHandle<FReal, 3>*> &ParticleHandles = Collection->PhysObject->GetSolverParticleHandles();
+			TArray<Chaos::TPBDRigidClusteredParticleHandle<FReal, 3>*> ParticleHandles = {
+			Collection->PhysObject->GetParticle_Internal(0),
+			Collection->PhysObject->GetParticle_Internal(1),
+			Collection->PhysObject->GetParticle_Internal(2),
+			Collection->PhysObject->GetParticle_Internal(3),
+			Collection->PhysObject->GetParticle_Internal(4),
+			Collection->PhysObject->GetParticle_Internal(5),
+			Collection->PhysObject->GetParticle_Internal(6),
+			Collection->PhysObject->GetParticle_Internal(7),
+			Collection->PhysObject->GetParticle_Internal(8),
+		};
 
 		auto& Clustering = UnitTest.Solver->GetEvolution()->GetRigidClustering();
 		const auto & ClusterMap = Clustering.GetChildrenMap();
@@ -339,7 +349,17 @@ namespace GeometryCollectionTest
 
 		UnitTest.Advance();
 
-		TArray<Chaos::TPBDRigidClusteredParticleHandle<FReal, 3>*>& ParticleHandles = Collection->PhysObject->GetSolverParticleHandles();
+		TArray<Chaos::TPBDRigidClusteredParticleHandle<FReal, 3>*> ParticleHandles = {
+			Collection->PhysObject->GetParticle_Internal(0),
+			Collection->PhysObject->GetParticle_Internal(1),
+			Collection->PhysObject->GetParticle_Internal(2),
+			Collection->PhysObject->GetParticle_Internal(3),
+			Collection->PhysObject->GetParticle_Internal(4),
+			Collection->PhysObject->GetParticle_Internal(5),
+			Collection->PhysObject->GetParticle_Internal(6),
+			Collection->PhysObject->GetParticle_Internal(7),
+			Collection->PhysObject->GetParticle_Internal(8),
+		};
 
 		auto& Clustering = UnitTest.Solver->GetEvolution()->GetRigidClustering();
 		const auto& ClusterMap = Clustering.GetChildrenMap();
@@ -465,10 +485,8 @@ namespace GeometryCollectionTest
 		auto& Clustering = UnitTest.Solver->GetEvolution()->GetRigidClustering();
 		const auto& ClusterMap = Clustering.GetChildrenMap();
 
-		EXPECT_TRUE(ClusterMapContains(ClusterMap, Collection->PhysObject->GetSolverClusterHandles()[0],
-			{ Collection->PhysObject->GetSolverParticleHandles()[0],Collection->PhysObject->GetSolverParticleHandles()[1] }));
-
-		TArray<Chaos::TPBDRigidClusteredParticleHandle<FReal, 3>*>& ParticleHandles = Collection->PhysObject->GetSolverParticleHandles();
+		EXPECT_TRUE(ClusterMapContains(ClusterMap, Collection->PhysObject->GetSolverClusterHandle_Internal(0),
+			{ Collection->PhysObject->GetParticle_Internal(0),Collection->PhysObject->GetParticle_Internal(1) }));
 
 		// Particles array contains the following:		
 		// 0: Box1 (top)
@@ -482,21 +500,21 @@ namespace GeometryCollectionTest
 
 			CurrentRigidDistance = (Collection->DynamicCollection->GetTransform(1).GetTranslation() - Collection->DynamicCollection->GetTransform(0).GetTranslation()).Size();
 
-			if ((BrokenFrame == INDEX_NONE) && !ParticleHandles[2]->Disabled())
+			if ((BrokenFrame == INDEX_NONE) && !Collection->PhysObject->GetSolverClusterHandle_Internal(2)->Disabled())
 			{
 				// The two boxes are dropping to the ground as a cluster
-				EXPECT_TRUE(ParticleHandles[0]->Disabled());
-				EXPECT_TRUE(ParticleHandles[1]->Disabled());
+				EXPECT_TRUE(Collection->PhysObject->GetParticle_Internal(0)->Disabled());
+				EXPECT_TRUE(Collection->PhysObject->GetParticle_Internal(1)->Disabled());
 
 				// The boxes are still separated by StartingRigidDistance
 				EXPECT_LT(FMath::Abs(CurrentRigidDistance - StartingRigidDistance), 1e-4);
 			}
 
-			if ((BrokenFrame == INDEX_NONE) && ParticleHandles[2]->Disabled())
+			if ((BrokenFrame == INDEX_NONE) && Collection->PhysObject->GetParticle_Internal(2)->Disabled())
 			{
 				// The cluster has just hit the ground and should have broken.
-				EXPECT_FALSE(ParticleHandles[0]->Disabled());
-				EXPECT_FALSE(ParticleHandles[1]->Disabled());
+				EXPECT_FALSE(Collection->PhysObject->GetParticle_Internal(0)->Disabled());
+				EXPECT_FALSE(Collection->PhysObject->GetParticle_Internal(1)->Disabled());
 				EXPECT_EQ(ClusterMap.Num(), 0);
 				BrokenFrame = Frame;
 			}
@@ -510,9 +528,9 @@ namespace GeometryCollectionTest
 		}
 
 		// Make sure it actually broke
-		EXPECT_FALSE(ParticleHandles[0]->Disabled());
-		EXPECT_FALSE(ParticleHandles[1]->Disabled());
-		EXPECT_TRUE(ParticleHandles[2]->Disabled());
+		EXPECT_FALSE(Collection->PhysObject->GetParticle_Internal(0)->Disabled());
+		EXPECT_FALSE(Collection->PhysObject->GetParticle_Internal(1)->Disabled());
+		EXPECT_TRUE(Collection->PhysObject->GetParticle_Internal(2)->Disabled());
 		EXPECT_TRUE(BrokenFrame != INDEX_NONE);
 
 		EXPECT_GT(FMath::Abs(CurrentRigidDistance - StartingRigidDistance), 1e-4);
@@ -567,13 +585,12 @@ namespace GeometryCollectionTest
 		// [1]: GeometryCollection Sphere1 at 0,10,40
 		// [2]: GeometryCollection Cluster0 of Sphere0 and Sphere1 at 0,10,50 (root rotated 90deg about X)
 		// [3]: GeometryCollection Cluster1 of Cluster0 at 0,10,50
-		TArray<Chaos::TPBDRigidClusteredParticleHandle<FReal, 3>*>& ParticleHandles = Collection->PhysObject->GetSolverParticleHandles();
 
 		auto& Clustering = UnitTest.Solver->GetEvolution()->GetRigidClustering();
 		const auto& ClusterMap = Clustering.GetChildrenMap();
 
-		EXPECT_TRUE(ClusterMapContains(ClusterMap, ParticleHandles[2], { ParticleHandles[0],ParticleHandles[1] }));
-		EXPECT_TRUE(ClusterMapContains(ClusterMap, ParticleHandles[3], { ParticleHandles[2] }));
+		EXPECT_TRUE(ClusterMapContains(ClusterMap, Collection->PhysObject->GetParticle_Internal(2), { Collection->PhysObject->GetParticle_Internal(0),Collection->PhysObject->GetParticle_Internal(1) }));
+		EXPECT_TRUE(ClusterMapContains(ClusterMap, Collection->PhysObject->GetParticle_Internal(3), { Collection->PhysObject->GetParticle_Internal(2) }));
 		
 		TArray<bool> Conditions = {false,false,false};
 
@@ -586,10 +603,10 @@ namespace GeometryCollectionTest
 			if (Conditions[0]==false)
 			{
 				if (
-					ParticleHandles[0]->Disabled() == true &&
-					ParticleHandles[1]->Disabled() == true &&
-					ParticleHandles[2]->Disabled() == true &&
-					ParticleHandles[3]->Disabled() == false) 
+					Collection->PhysObject->GetParticle_Internal(0)->Disabled() == true &&
+					Collection->PhysObject->GetParticle_Internal(1)->Disabled() == true &&
+					Collection->PhysObject->GetParticle_Internal(2)->Disabled() == true &&
+					Collection->PhysObject->GetParticle_Internal(3)->Disabled() == false) 
 				{
 					// Only the outer Cluster1 is active. 
 					// This is the initial condition
@@ -599,26 +616,26 @@ namespace GeometryCollectionTest
 			else if (Conditions[0]==true && Conditions[1] == false)
 			{
 				if (
-					ParticleHandles[0]->Disabled() == true &&
-					ParticleHandles[1]->Disabled() == true &&
-					ParticleHandles[2]->Disabled() == false &&
-					ParticleHandles[3]->Disabled() == true)
+					Collection->PhysObject->GetParticle_Internal(0)->Disabled() == true &&
+					Collection->PhysObject->GetParticle_Internal(1)->Disabled() == true &&
+					Collection->PhysObject->GetParticle_Internal(2)->Disabled() == false &&
+					Collection->PhysObject->GetParticle_Internal(3)->Disabled() == true)
 				{
 					// Cluster1 is now disabled, and Cluster0 was activated.
 					// This happens when Cluster1 collides with the floor
 					Conditions[1] = true;
-					EXPECT_TRUE(ClusterMapContains(ClusterMap, ParticleHandles[2], { ParticleHandles[0],ParticleHandles[1] }));
+					EXPECT_TRUE(ClusterMapContains(ClusterMap, Collection->PhysObject->GetParticle_Internal(2), { Collection->PhysObject->GetParticle_Internal(0),Collection->PhysObject->GetParticle_Internal(1) }));
 					EXPECT_EQ(ClusterMap.Num(), 1);
-					EXPECT_TRUE(!ClusterMap.Contains(ParticleHandles[3]));
+					EXPECT_TRUE(!ClusterMap.Contains(Collection->PhysObject->GetParticle_Internal(3)));
 				}
 			}
 			else if (Conditions[1] == true && Conditions[2] == false)
 			{
 				if (
-					ParticleHandles[0]->Disabled() == false &&
-					ParticleHandles[1]->Disabled() == false &&
-					ParticleHandles[2]->Disabled() == true &&
-					ParticleHandles[3]->Disabled() == true)
+					Collection->PhysObject->GetParticle_Internal(0)->Disabled() == false &&
+					Collection->PhysObject->GetParticle_Internal(1)->Disabled() == false &&
+					Collection->PhysObject->GetParticle_Internal(2)->Disabled() == true &&
+					Collection->PhysObject->GetParticle_Internal(3)->Disabled() == true)
 				{
 					// Cluster0 is now disabled because it had a damage threshold of 0
 					// and the boxes should now be active.
@@ -667,15 +684,14 @@ namespace GeometryCollectionTest
 		auto& Clustering = UnitTest.Solver->GetEvolution()->GetRigidClustering();
 		const auto& ClusterMap = Clustering.GetChildrenMap();
 		TArray<FTransform> Collection1_InitialTM; GeometryCollectionAlgo::GlobalMatrices(Collection1->RestCollection->Transform, Collection1->RestCollection->Parent, Collection1_InitialTM);
-		TArray<Chaos::TPBDRigidClusteredParticleHandle<FReal, 3>*>& Collection1Handles = Collection1->PhysObject->GetSolverParticleHandles();
 		const auto& SovlerParticleHandles = UnitTest.Solver->GetParticles().GetParticleHandles();
 
 		UnitTest.Solver->RegisterSimOneShotCallback([&]()
 		{
 			EXPECT_EQ(SovlerParticleHandles.Size(),4);
 			EXPECT_EQ(ClusterMap.Num(),2);
-			EXPECT_TRUE(ClusterMapContains(ClusterMap,Collection1Handles[2],{Collection1Handles[1],Collection1Handles[0]}));
-			EXPECT_TRUE(ClusterMapContains(ClusterMap,Collection1Handles[3],{Collection1Handles[2]}));
+			EXPECT_TRUE(ClusterMapContains(ClusterMap, Collection1->PhysObject->GetParticle_Internal(2), { Collection1->PhysObject->GetParticle_Internal(1),Collection1->PhysObject->GetParticle_Internal(0) }));
+			EXPECT_TRUE(ClusterMapContains(ClusterMap, Collection1->PhysObject->GetParticle_Internal(3), { Collection1->PhysObject->GetParticle_Internal(2) }));
 		});
 		
 
@@ -683,27 +699,27 @@ namespace GeometryCollectionTest
 
 		EXPECT_EQ(SovlerParticleHandles.Size(), 4);
 		EXPECT_EQ(ClusterMap.Num(), 2);
-		EXPECT_TRUE(ClusterMapContains(ClusterMap, Collection1Handles[2], { Collection1Handles[1],Collection1Handles[0] }));
-		EXPECT_TRUE(ClusterMapContains(ClusterMap, Collection1Handles[3], { Collection1Handles[2] }));
+		EXPECT_TRUE(ClusterMapContains(ClusterMap, Collection1->PhysObject->GetParticle_Internal(2), { Collection1->PhysObject->GetParticle_Internal(1),Collection1->PhysObject->GetParticle_Internal(0) }));
+		EXPECT_TRUE(ClusterMapContains(ClusterMap, Collection1->PhysObject->GetParticle_Internal(3), { Collection1->PhysObject->GetParticle_Internal(2) }));
 		TArray<FTransform> Collection1_PreReleaseTM; 
 		GeometryCollectionAlgo::Private::GlobalMatrices(*DynamicCollection1.Get(), Collection1_PreReleaseTM);
 		for (int Idx = 0; Idx < Collection1_PreReleaseTM.Num(); Idx++) {
 			EXPECT_TRUE( (Collection1_PreReleaseTM[Idx].GetTranslation()-Collection1_InitialTM[Idx].GetTranslation()).Size()<KINDA_SMALL_NUMBER);
 		}
 
-		UnitTest.Solver->GetEvolution()->GetRigidClustering().DeactivateClusterParticle({ Collection1Handles[3] });
+		UnitTest.Solver->GetEvolution()->GetRigidClustering().DeactivateClusterParticle({ Collection1->PhysObject->GetParticle_Internal(3) });
 		UnitTest.Advance();
 
 		EXPECT_EQ(SovlerParticleHandles.Size(), 4);
 		EXPECT_EQ(ClusterMap.Num(), 1);
-		EXPECT_TRUE(ClusterMapContains(ClusterMap, Collection1Handles[2], { Collection1Handles[1],Collection1Handles[0] }));
+		EXPECT_TRUE(ClusterMapContains(ClusterMap, Collection1->PhysObject->GetParticle_Internal(2), { Collection1->PhysObject->GetParticle_Internal(1),Collection1->PhysObject->GetParticle_Internal(0) }));
 		TArray<FTransform> Collection1_PostReleaseTM; 
 		GeometryCollectionAlgo::Private::GlobalMatrices(*DynamicCollection1.Get(), Collection1_PostReleaseTM);
 		for (int Idx = 0; Idx < Collection1_PostReleaseTM.Num(); Idx++) {
 			EXPECT_TRUE((Collection1_PostReleaseTM[Idx].GetTranslation() - Collection1_InitialTM[Idx].GetTranslation()).Size() < KINDA_SMALL_NUMBER);
 		}
 
-		UnitTest.Solver->GetEvolution()->GetRigidClustering().DeactivateClusterParticle({ Collection1Handles[2] });
+		UnitTest.Solver->GetEvolution()->GetRigidClustering().DeactivateClusterParticle({ Collection1->PhysObject->GetParticle_Internal(2) });
 		UnitTest.Advance();
 
 		EXPECT_EQ(SovlerParticleHandles.Size(), 4);
@@ -757,7 +773,17 @@ namespace GeometryCollectionTest
 
 		UnitTest.Advance();
 
-		TArray<Chaos::TPBDRigidClusteredParticleHandle<FReal, 3>*>& ParticleHandles = Collection->PhysObject->GetSolverParticleHandles();
+		TArray<Chaos::TPBDRigidClusteredParticleHandle<FReal, 3>*> ParticleHandles = { 
+			Collection->PhysObject->GetParticle_Internal(0),
+			Collection->PhysObject->GetParticle_Internal(1),
+			Collection->PhysObject->GetParticle_Internal(2),
+			Collection->PhysObject->GetParticle_Internal(3),
+			Collection->PhysObject->GetParticle_Internal(4),
+			Collection->PhysObject->GetParticle_Internal(5),
+			Collection->PhysObject->GetParticle_Internal(6),
+			Collection->PhysObject->GetParticle_Internal(7),
+			Collection->PhysObject->GetParticle_Internal(8),
+		};
 
 		auto& Clustering = UnitTest.Solver->GetEvolution()->GetRigidClustering();
 		const auto& ClusterMap = Clustering.GetChildrenMap();
@@ -927,8 +953,28 @@ namespace GeometryCollectionTest
 		// Tick once to fush commands
 		UnitTest.Advance();
 
-		TArray<Chaos::TPBDRigidClusteredParticleHandle<FReal, 3>*>& ParticleHandles = Collection->PhysObject->GetSolverParticleHandles();
-		TArray<Chaos::TPBDRigidClusteredParticleHandle<FReal, 3>*>& ClusterHandles = Collection->PhysObject->GetSolverClusterHandles();
+		TArray<Chaos::TPBDRigidClusteredParticleHandle<FReal, 3>*> ParticleHandles = {
+			Collection->PhysObject->GetParticle_Internal(0),
+			Collection->PhysObject->GetParticle_Internal(1),
+			Collection->PhysObject->GetParticle_Internal(2),
+			Collection->PhysObject->GetParticle_Internal(3),
+			Collection->PhysObject->GetParticle_Internal(4),
+			Collection->PhysObject->GetParticle_Internal(5),
+			Collection->PhysObject->GetParticle_Internal(6),
+			Collection->PhysObject->GetParticle_Internal(7),
+			Collection->PhysObject->GetParticle_Internal(8),
+		};
+		TArray<Chaos::TPBDRigidClusteredParticleHandle<FReal, 3>*> ClusterHandles = {
+			Collection->PhysObject->GetSolverClusterHandle_Internal(0),
+			Collection->PhysObject->GetSolverClusterHandle_Internal(1),
+			Collection->PhysObject->GetSolverClusterHandle_Internal(2),
+			Collection->PhysObject->GetSolverClusterHandle_Internal(3),
+			Collection->PhysObject->GetSolverClusterHandle_Internal(4),
+			Collection->PhysObject->GetSolverClusterHandle_Internal(5),
+			Collection->PhysObject->GetSolverClusterHandle_Internal(6),
+			Collection->PhysObject->GetSolverClusterHandle_Internal(7),
+			Collection->PhysObject->GetSolverClusterHandle_Internal(8),
+		};
 
 		FRigidClustering& Clustering = UnitTest.Solver->GetEvolution()->GetRigidClustering();
 		const FRigidClustering::FClusterMap& ClusterMap = Clustering.GetChildrenMap();
@@ -1225,8 +1271,28 @@ namespace GeometryCollectionTest
 		// Tick once to fush commands
 		UnitTest.Advance();
 
-		TArray<Chaos::TPBDRigidClusteredParticleHandle<FReal, 3>*>& ParticleHandles = Collection->PhysObject->GetSolverParticleHandles();
-		TArray<Chaos::TPBDRigidClusteredParticleHandle<FReal, 3>*>& ClusterHandles = Collection->PhysObject->GetSolverClusterHandles();
+		TArray<Chaos::TPBDRigidClusteredParticleHandle<FReal, 3>*> ParticleHandles = {
+			Collection->PhysObject->GetParticle_Internal(0),
+			Collection->PhysObject->GetParticle_Internal(1),
+			Collection->PhysObject->GetParticle_Internal(2),
+			Collection->PhysObject->GetParticle_Internal(3),
+			Collection->PhysObject->GetParticle_Internal(4),
+			Collection->PhysObject->GetParticle_Internal(5),
+			Collection->PhysObject->GetParticle_Internal(6),
+			Collection->PhysObject->GetParticle_Internal(7),
+			Collection->PhysObject->GetParticle_Internal(8),
+		};
+		TArray<Chaos::TPBDRigidClusteredParticleHandle<FReal, 3>*> ClusterHandles = {
+			Collection->PhysObject->GetSolverClusterHandle_Internal(0),
+			Collection->PhysObject->GetSolverClusterHandle_Internal(1),
+			Collection->PhysObject->GetSolverClusterHandle_Internal(2),
+			Collection->PhysObject->GetSolverClusterHandle_Internal(3),
+			Collection->PhysObject->GetSolverClusterHandle_Internal(4),
+			Collection->PhysObject->GetSolverClusterHandle_Internal(5),
+			Collection->PhysObject->GetSolverClusterHandle_Internal(6),
+			Collection->PhysObject->GetSolverClusterHandle_Internal(7),
+			Collection->PhysObject->GetSolverClusterHandle_Internal(8),
+		};
 
 		FRigidClustering& Clustering = UnitTest.Solver->GetEvolution()->GetRigidClustering();
 		const FRigidClustering::FClusterMap& ClusterMap = Clustering.GetChildrenMap();
@@ -1711,12 +1777,22 @@ namespace GeometryCollectionTest
 		UnitTest.Solver->RegisterSimOneShotCallback([&]()
 		{
 			EXPECT_EQ(ClusterMap.Num(),2);
-			const auto& CollectionParticles = Collection->PhysObject->GetSolverParticleHandles();
+			TArray<Chaos::TPBDRigidClusteredParticleHandle<FReal, 3>*> CollectionParticles = {
+			Collection->PhysObject->GetParticle_Internal(0),
+			Collection->PhysObject->GetParticle_Internal(1),
+			Collection->PhysObject->GetParticle_Internal(2)
+			};
 			EXPECT_EQ(CollectionParticles.Num(),3);
+			EXPECT_EQ(Collection->PhysObject->GetNumTransforms(), 3);
 			EXPECT_TRUE(ClusterMapContains(ClusterMap,CollectionParticles[2],{CollectionParticles[1],CollectionParticles[0]}));
 
-			const auto& CollectionParticles2 = Collection2->PhysObject->GetSolverParticleHandles();
+			TArray<Chaos::TPBDRigidClusteredParticleHandle<FReal, 3>*> CollectionParticles2 = {
+			Collection2->PhysObject->GetParticle_Internal(0),
+			Collection2->PhysObject->GetParticle_Internal(1),
+			Collection2->PhysObject->GetParticle_Internal(2)
+			};
 			EXPECT_EQ(CollectionParticles2.Num(),3);
+			EXPECT_EQ(Collection2->PhysObject->GetNumTransforms(), 3);
 			EXPECT_TRUE(ClusterMapContains(ClusterMap,CollectionParticles2[2],{CollectionParticles2[1],CollectionParticles2[0]}));
 		});
 
@@ -1724,10 +1800,19 @@ namespace GeometryCollectionTest
 		{
 			UnitTest.Advance();
 
-			const auto& CollectionParticles = Collection->PhysObject->GetSolverParticleHandles();
+			TArray<Chaos::TPBDRigidClusteredParticleHandle<FReal, 3>*> CollectionParticles = {
+			Collection->PhysObject->GetParticle_Internal(0),
+			Collection->PhysObject->GetParticle_Internal(1),
+			Collection->PhysObject->GetParticle_Internal(2),
+			};
 			EXPECT_EQ(CollectionParticles.Num(),3);
+			EXPECT_EQ(Collection->PhysObject->GetNumTransforms(), 3);
 
-			const auto& CollectionParticles2 = Collection2->PhysObject->GetSolverParticleHandles();
+			TArray<Chaos::TPBDRigidClusteredParticleHandle<FReal, 3>*> CollectionParticles2 = {
+			Collection2->PhysObject->GetParticle_Internal(0),
+			Collection2->PhysObject->GetParticle_Internal(1),
+			Collection2->PhysObject->GetParticle_Internal(2),
+			};
 
 			const auto* Root = CollectionParticles[0]->ClusterIds().Id;
 
@@ -1832,13 +1917,36 @@ namespace GeometryCollectionTest
 		TArray<FTransform> Collection1_InitialTM; GeometryCollectionAlgo::Private::GlobalMatrices(*DynamicCollection1, Collection1_InitialTM);
 		TArray<FTransform> Collection2_InitialTM; GeometryCollectionAlgo::Private::GlobalMatrices(*DynamicCollection2, Collection2_InitialTM);
 
-		TArray<Chaos::TPBDRigidClusteredParticleHandle<FReal, 3>*>& Collection1Handles = Collection1->PhysObject->GetSolverParticleHandles();
-		TArray<Chaos::TPBDRigidClusteredParticleHandle<FReal, 3>*>& Collection2Handles = Collection2->PhysObject->GetSolverParticleHandles();
+		TArray<Chaos::TPBDRigidClusteredParticleHandle<FReal, 3>*> Collection1Handles = {
+			Collection1->PhysObject->GetParticle_Internal(0),
+			Collection1->PhysObject->GetParticle_Internal(1),
+			Collection1->PhysObject->GetParticle_Internal(2),
+		};
+		TArray<Chaos::TPBDRigidClusteredParticleHandle<FReal, 3>*> Collection2Handles = {
+			Collection2->PhysObject->GetParticle_Internal(0),
+			Collection2->PhysObject->GetParticle_Internal(1),
+			Collection2->PhysObject->GetParticle_Internal(2),
+		};
+		/*TArray<Chaos::TPBDRigidClusteredParticleHandle<FReal, 3>*>& Collection1Handles = Collection1->PhysObject->GetSolverParticleHandles();
+		TArray<Chaos::TPBDRigidClusteredParticleHandle<FReal, 3>*>& Collection2Handles = Collection2->PhysObject->GetSolverParticleHandles();*/
+
+		
+
 		const auto& SovlerParticleHandles = UnitTest.Solver->GetParticles().GetParticleHandles();
 		
 
 		UnitTest.Solver->RegisterSimOneShotCallback([&]()
 		{
+			Collection1Handles = {
+			Collection1->PhysObject->GetParticle_Internal(0),
+			Collection1->PhysObject->GetParticle_Internal(1),
+			Collection1->PhysObject->GetParticle_Internal(2),
+				};
+			Collection2Handles = {
+				Collection2->PhysObject->GetParticle_Internal(0),
+				Collection2->PhysObject->GetParticle_Internal(1),
+				Collection2->PhysObject->GetParticle_Internal(2),
+			};
 			EXPECT_EQ(SovlerParticleHandles.Size(), 6);
 			EXPECT_EQ(ClusterMap.Num(),2);
 			EXPECT_TRUE(ClusterMapContains(ClusterMap,Collection1Handles[2],{Collection1Handles[1],Collection1Handles[0]}));
@@ -1948,12 +2056,30 @@ namespace GeometryCollectionTest
 		TArray<FTransform> Collection1_InitialTM; GeometryCollectionAlgo::Private::GlobalMatrices(*DynamicCollection1, Collection1_InitialTM);
 		TArray<FTransform> Collection2_InitialTM; GeometryCollectionAlgo::Private::GlobalMatrices(*DynamicCollection2, Collection2_InitialTM);
 
-		TArray<Chaos::TPBDRigidClusteredParticleHandle<FReal, 3>*>& Collection1Handles = Collection1->PhysObject->GetSolverParticleHandles();
-		TArray<Chaos::TPBDRigidClusteredParticleHandle<FReal, 3>*>& Collection2Handles = Collection2->PhysObject->GetSolverParticleHandles();
+		TArray<Chaos::TPBDRigidClusteredParticleHandle<FReal, 3>*> Collection1Handles = {
+			Collection1->PhysObject->GetParticle_Internal(0),
+			Collection1->PhysObject->GetParticle_Internal(1),
+			Collection1->PhysObject->GetParticle_Internal(2),
+		};
+		TArray<Chaos::TPBDRigidClusteredParticleHandle<FReal, 3>*> Collection2Handles = {
+			Collection2->PhysObject->GetParticle_Internal(0),
+			Collection2->PhysObject->GetParticle_Internal(1),
+			Collection2->PhysObject->GetParticle_Internal(2),
+		};
 		const auto& SovlerParticleHandles = UnitTest.Solver->GetParticles().GetParticleHandles();
 
 		UnitTest.Solver->RegisterSimOneShotCallback([&]()
 		{
+			Collection1Handles = {
+			Collection1->PhysObject->GetParticle_Internal(0),
+			Collection1->PhysObject->GetParticle_Internal(1),
+			Collection1->PhysObject->GetParticle_Internal(2),
+				};
+			Collection2Handles = {
+				Collection2->PhysObject->GetParticle_Internal(0),
+				Collection2->PhysObject->GetParticle_Internal(1),
+				Collection2->PhysObject->GetParticle_Internal(2),
+			};
 			EXPECT_EQ(SovlerParticleHandles.Size(),6);
 			EXPECT_EQ(ClusterMap.Num(),2);
 			EXPECT_TRUE(ClusterMapContains(ClusterMap,Collection1Handles[2],{Collection1Handles[1],Collection1Handles[0]}));
@@ -2057,11 +2183,20 @@ namespace GeometryCollectionTest
 		auto& Clustering = UnitTest.Solver->GetEvolution()->GetRigidClustering();
 		const auto& ClusterMap = Clustering.GetChildrenMap();
 		TArray<FTransform> Collection1_InitialTM; GeometryCollectionAlgo::Private::GlobalMatrices(*DynamicCollection1, Collection1_InitialTM);
-		TArray<Chaos::TPBDRigidClusteredParticleHandle<FReal, 3>*>& Collection1Handles = Collection1->PhysObject->GetSolverParticleHandles();
+		TArray<Chaos::TPBDRigidClusteredParticleHandle<FReal, 3>*> Collection1Handles = {
+			Collection1->PhysObject->GetParticle_Internal(0),
+			Collection1->PhysObject->GetParticle_Internal(1),
+			Collection1->PhysObject->GetParticle_Internal(2),
+		};
 		const auto& SovlerParticleHandles = UnitTest.Solver->GetParticles().GetParticleHandles();
 
 		UnitTest.Solver->RegisterSimOneShotCallback([&]()
 		{
+			Collection1Handles = {
+			Collection1->PhysObject->GetParticle_Internal(0),
+			Collection1->PhysObject->GetParticle_Internal(1),
+			Collection1->PhysObject->GetParticle_Internal(2),
+				};
 			EXPECT_EQ(SovlerParticleHandles.Size(),4);
 			EXPECT_EQ(ClusterMap.Num(),2);
 			EXPECT_TRUE(ClusterMapContains(ClusterMap,Collection1Handles[2],{Collection1Handles[1],Collection1Handles[0]}));
@@ -2269,7 +2404,19 @@ namespace GeometryCollectionTest
 		FalloffField->Position = FVector(0.0, 0.0, 0.0);
 		FalloffField->Falloff = EFieldFalloffType::Field_FallOff_None;
 
-		TArray<Chaos::TPBDRigidClusteredParticleHandle<FReal, 3>*>& ParticleHandles = Collection->PhysObject->GetSolverParticleHandles();
+		TArray<Chaos::TPBDRigidClusteredParticleHandle<FReal, 3>*> ParticleHandles = {
+			Collection->PhysObject->GetParticle_Internal(0),
+			Collection->PhysObject->GetParticle_Internal(1),
+			Collection->PhysObject->GetParticle_Internal(2),
+			Collection->PhysObject->GetParticle_Internal(3),
+			Collection->PhysObject->GetParticle_Internal(4),
+			Collection->PhysObject->GetParticle_Internal(5),
+			Collection->PhysObject->GetParticle_Internal(6),
+			Collection->PhysObject->GetParticle_Internal(7),
+			Collection->PhysObject->GetParticle_Internal(8),
+			Collection->PhysObject->GetParticle_Internal(9),
+			Collection->PhysObject->GetParticle_Internal(10),
+		};
 		Chaos::FRigidClustering& Clustering = UnitTest.Solver->GetEvolution()->GetRigidClustering();
 		const auto& ClusterMap = Clustering.GetChildrenMap();
 		auto& x = UnitTest.Solver->GetEvolution()->GetRigidClustering().GetStrainArray();
@@ -2383,9 +2530,21 @@ namespace GeometryCollectionTest
 		UnitTest.Advance();		// this call triggers the array resize based on the fraction
 
 		// Test non-clustered bodies
-		TArray<Chaos::TPBDRigidClusteredParticleHandle<FReal, 3>*>& ParticleHandles = Collection->PhysObject->GetSolverParticleHandles();
+		TArray<Chaos::TPBDRigidClusteredParticleHandle<FReal, 3>*> ParticleHandles = {
+			Collection->PhysObject->GetParticle_Internal(0),
+			Collection->PhysObject->GetParticle_Internal(1),
+			Collection->PhysObject->GetParticle_Internal(2),
+			Collection->PhysObject->GetParticle_Internal(3),
+			Collection->PhysObject->GetParticle_Internal(4),
+			Collection->PhysObject->GetParticle_Internal(5),
+			Collection->PhysObject->GetParticle_Internal(6),
+			Collection->PhysObject->GetParticle_Internal(7),
+			Collection->PhysObject->GetParticle_Internal(8),
+			Collection->PhysObject->GetParticle_Internal(9),
+			Collection->PhysObject->GetParticle_Internal(10),
+		};
 		int32 NumCollisionParticles, ExpectedNumCollisionParticles;
-		for (int i = 0; i < ParticleHandles.Num(); i++)
+		for (int i = 0; i < Collection->PhysObject->GetNumTransforms(); i++)
 		{
 			if (Collection->RestCollection->SimulationType[i] == FGeometryCollection::ESimulationTypes::FST_Rigid)
 			{
