Remove deprecated warning in headless chaos

diff --git a/Engine/Source/Programs/HeadlessChaos/Private/GeometryCollection/GeometryCollectionTestClustering.cpp b/Engine/Source/Programs/HeadlessChaos/Private/GeometryCollection/GeometryCollectionTestClustering.cpp
--- a/Engine/Source/Programs/HeadlessChaos/Private/GeometryCollection/GeometryCollectionTestClustering.cpp
+++ b/Engine/Source/Programs/HeadlessChaos/Private/GeometryCollection/GeometryCollectionTestClustering.cpp
@@ -1927,10 +1927,6 @@ namespace GeometryCollectionTest
 			Collection2->PhysObject->GetParticle_Internal(1),
 			Collection2->PhysObject->GetParticle_Internal(2),
 		};
-		/*TArray<Chaos::TPBDRigidClusteredParticleHandle<FReal, 3>*>& Collection1Handles = Collection1->PhysObject->GetSolverParticleHandles();
-		TArray<Chaos::TPBDRigidClusteredParticleHandle<FReal, 3>*>& Collection2Handles = Collection2->PhysObject->GetSolverParticleHandles();*/
-
-		
 
 		const auto& SovlerParticleHandles = UnitTest.Solver->GetParticles().GetParticleHandles();
 		
diff --git a/Engine/Source/Programs/HeadlessChaos/Private/GeometryCollection/GeometryCollectionTestSimulationField.cpp b/Engine/Source/Programs/HeadlessChaos/Private/GeometryCollection/GeometryCollectionTestSimulationField.cpp
--- a/Engine/Source/Programs/HeadlessChaos/Private/GeometryCollection/GeometryCollectionTestSimulationField.cpp
+++ b/Engine/Source/Programs/HeadlessChaos/Private/GeometryCollection/GeometryCollectionTestSimulationField.cpp
@@ -793,7 +793,15 @@ namespace GeometryCollectionTest
 		auto& Clustering = UnitTest.Solver->GetEvolution()->GetRigidClustering();
 		const auto& ClusterMap = Clustering.GetChildrenMap();
 
-		TArray<Chaos::TPBDRigidClusteredParticleHandle<FReal, 3>*>& ParticleHandles = Collection->PhysObject->GetSolverParticleHandles();
+		TArray<Chaos::TPBDRigidClusteredParticleHandle<FReal, 3>*> ParticleHandles = {
+			Collection->PhysObject->GetParticle_Internal(0),
+			Collection->PhysObject->GetParticle_Internal(1),
+			Collection->PhysObject->GetParticle_Internal(2),
+			Collection->PhysObject->GetParticle_Internal(3),
+			Collection->PhysObject->GetParticle_Internal(4),
+			Collection->PhysObject->GetParticle_Internal(5),
+			Collection->PhysObject->GetParticle_Internal(6),
+		};
 		{
 			FName TargetName = GetFieldPhysicsName(EFieldPhysicsType::Field_ExternalClusterStrain);
 			UnitTest.Solver->GetPerSolverField().AddTransientCommand({ TargetName, FalloffField->NewCopy() });
@@ -853,12 +861,56 @@ namespace GeometryCollectionTest
 		FalloffField->Falloff = EFieldFalloffType::Field_FallOff_None;
 
 		UnitTest.Initialize();
-		auto& ParticleHandles = Collection->PhysObject->GetSolverParticleHandles();
-		auto& ClusteredParticleHandles = Collection->PhysObject->GetSolverClusterHandles();
+
+		TArray<Chaos::TPBDRigidClusteredParticleHandle<FReal, 3>*> ParticleHandles = {
+			Collection->PhysObject->GetParticle_Internal(0),
+			Collection->PhysObject->GetParticle_Internal(1),
+			Collection->PhysObject->GetParticle_Internal(2),
+			Collection->PhysObject->GetParticle_Internal(3),
+			Collection->PhysObject->GetParticle_Internal(4),
+			Collection->PhysObject->GetParticle_Internal(5),
+			Collection->PhysObject->GetParticle_Internal(6),
+			Collection->PhysObject->GetParticle_Internal(7),
+			Collection->PhysObject->GetParticle_Internal(8),
+		};
+		TArray<Chaos::TPBDRigidClusteredParticleHandle<FReal, 3>*> ClusteredParticleHandles = {
+			Collection->PhysObject->GetSolverClusterHandle_Internal(0),
+			Collection->PhysObject->GetSolverClusterHandle_Internal(1),
+			Collection->PhysObject->GetSolverClusterHandle_Internal(2),
+			Collection->PhysObject->GetSolverClusterHandle_Internal(3),
+			Collection->PhysObject->GetSolverClusterHandle_Internal(4),
+			Collection->PhysObject->GetSolverClusterHandle_Internal(5),
+			Collection->PhysObject->GetSolverClusterHandle_Internal(6),
+			Collection->PhysObject->GetSolverClusterHandle_Internal(7),
+			Collection->PhysObject->GetSolverClusterHandle_Internal(8),
+		};
+
 		auto& Clustering = UnitTest.Solver->GetEvolution()->GetRigidClustering();
 		const auto& ClusterMap = Clustering.GetChildrenMap();
 		UnitTest.Solver->RegisterSimOneShotCallback([&]()
 		{
+			ParticleHandles = {
+			Collection->PhysObject->GetParticle_Internal(0),
+			Collection->PhysObject->GetParticle_Internal(1),
+			Collection->PhysObject->GetParticle_Internal(2),
+			Collection->PhysObject->GetParticle_Internal(3),
+			Collection->PhysObject->GetParticle_Internal(4),
+			Collection->PhysObject->GetParticle_Internal(5),
+			Collection->PhysObject->GetParticle_Internal(6),
+			Collection->PhysObject->GetParticle_Internal(7),
+			Collection->PhysObject->GetParticle_Internal(8),
+			};
+			ClusteredParticleHandles = {
+				Collection->PhysObject->GetSolverClusterHandle_Internal(0),
+				Collection->PhysObject->GetSolverClusterHandle_Internal(1),
+				Collection->PhysObject->GetSolverClusterHandle_Internal(2),
+				Collection->PhysObject->GetSolverClusterHandle_Internal(3),
+				Collection->PhysObject->GetSolverClusterHandle_Internal(4),
+				Collection->PhysObject->GetSolverClusterHandle_Internal(5),
+				Collection->PhysObject->GetSolverClusterHandle_Internal(6),
+				Collection->PhysObject->GetSolverClusterHandle_Internal(7),
+				Collection->PhysObject->GetSolverClusterHandle_Internal(8),
+			};
 			EXPECT_EQ(ClusterMap.Num(), 3);
 			EXPECT_EQ(ClusteredParticleHandles[1]->GetInternalStrains(), 101);
 		});
@@ -1018,7 +1070,9 @@ namespace GeometryCollectionTest
 		UnitTest.Initialize();
 		UnitTest.Advance();
 
-		TArray<Chaos::TPBDRigidClusteredParticleHandle<FReal, 3>*>& ParticleHandles = Collection->PhysObject->GetSolverParticleHandles();
+		TArray<Chaos::TPBDRigidClusteredParticleHandle<FReal, 3>*> ParticleHandles = {
+			Collection->PhysObject->GetParticle_Internal(0),
+		};
 		Chaos::TVector<float, 3> CurrV = ParticleHandles[0]->V();
 		
 
@@ -1129,7 +1183,13 @@ namespace GeometryCollectionTest
 		UnitTest.Initialize();
 		UnitTest.Advance();
 
-		TArray<Chaos::TPBDRigidClusteredParticleHandle<FReal, 3>*>& ParticleHandles = Collection->PhysObject->GetSolverParticleHandles();
+		TArray<Chaos::TPBDRigidClusteredParticleHandle<FReal, 3>*> ParticleHandles = {
+			Collection->PhysObject->GetParticle_Internal(0),
+			Collection->PhysObject->GetParticle_Internal(1),
+			Collection->PhysObject->GetParticle_Internal(2),
+			Collection->PhysObject->GetParticle_Internal(3),
+			Collection->PhysObject->GetParticle_Internal(4),
+		};
 		auto& Clustering = UnitTest.Solver->GetEvolution()->GetRigidClustering();
 		const auto& ClusterMap = Clustering.GetChildrenMap();
 
@@ -1198,7 +1258,13 @@ namespace GeometryCollectionTest
 		UnitTest.Initialize();
 		UnitTest.Advance();
 
-		TArray<Chaos::TPBDRigidClusteredParticleHandle<FReal, 3>*>& ParticleHandles = Collection->PhysObject->GetSolverParticleHandles();
+		TArray<Chaos::TPBDRigidClusteredParticleHandle<FReal, 3>*> ParticleHandles = {
+			Collection->PhysObject->GetParticle_Internal(0),
+			Collection->PhysObject->GetParticle_Internal(1),
+			Collection->PhysObject->GetParticle_Internal(2),
+			Collection->PhysObject->GetParticle_Internal(3),
+			Collection->PhysObject->GetParticle_Internal(4),
+		};
 		auto& Clustering = UnitTest.Solver->GetEvolution()->GetRigidClustering();
 		const auto& ClusterMap = Clustering.GetChildrenMap();
 
diff --git a/Engine/Source/Programs/HeadlessChaos/Private/GeometryCollection/GeometryCollectionTestSimulationSolver.cpp b/Engine/Source/Programs/HeadlessChaos/Private/GeometryCollection/GeometryCollectionTestSimulationSolver.cpp
--- a/Engine/Source/Programs/HeadlessChaos/Private/GeometryCollection/GeometryCollectionTestSimulationSolver.cpp
+++ b/Engine/Source/Programs/HeadlessChaos/Private/GeometryCollection/GeometryCollectionTestSimulationSolver.cpp
@@ -113,7 +113,7 @@ using namespace ChaosTest;
 			EXPECT_EQ(Collection->DynamicCollection->GetParent(2),-1);
 
 			// Set the one cluster to disabled
-			UnitTest.Solver->GetEvolution()->DisableParticle(Collection->PhysObject->GetSolverClusterHandles()[0]);
+			UnitTest.Solver->GetEvolution()->DisableParticle(Collection->PhysObject->GetSolverClusterHandle_Internal(0));
 
 
 			StartingRigidDistance = (Collection->DynamicCollection->GetTransform(1).GetTranslation() - Collection->DynamicCollection->GetTransform(0).GetTranslation()).Size();
@@ -377,7 +377,12 @@ using namespace ChaosTest;
 			UnitTest.Solver->SetGenerateBreakingData(true);
 			UnitTest.Solver->SetBreakingFilterSettings(BreakingFilterSettings);
 
-			ParticleHandles = Collection->PhysObject->GetSolverParticleHandles();
+			ParticleHandles = {
+			Collection->PhysObject->GetParticle_Internal(0),
+			Collection->PhysObject->GetParticle_Internal(1),
+			Collection->PhysObject->GetParticle_Internal(2),
+			Collection->PhysObject->GetParticle_Internal(3),
+			};
 
 			ParticleHandles[0]->SetM(TestMass + 1.0f);
 			ParticleHandles[1]->SetM(TestMass - 1.0f);
diff --git a/Engine/Source/Programs/HeadlessChaos/Private/GeometryCollection/GeometryCollectionTestSimulationStreaming.cpp b/Engine/Source/Programs/HeadlessChaos/Private/GeometryCollection/GeometryCollectionTestSimulationStreaming.cpp
--- a/Engine/Source/Programs/HeadlessChaos/Private/GeometryCollection/GeometryCollectionTestSimulationStreaming.cpp
+++ b/Engine/Source/Programs/HeadlessChaos/Private/GeometryCollection/GeometryCollectionTestSimulationStreaming.cpp
@@ -61,9 +61,11 @@ namespace GeometryCollectionTest
 		{
 			FGeometryCollectionWrapper* GCW = UnitTest.PhysicsObjects[i]->As<FGeometryCollectionWrapper>();
 			FGeometryCollectionWrapper* GCW2 = UnitTest.PhysicsObjects[i + 1]->As<FGeometryCollectionWrapper>();
-			EXPECT_EQ(GCW->PhysObject->GetSolverParticleHandles().Num(), 1);
-			EXPECT_EQ(GCW2->PhysObject->GetSolverParticleHandles().Num(), 1);
-			EXPECT_LT(GCW->PhysObject->GetSolverParticleHandles()[0]->X().Z, GCW2->PhysObject->GetSolverParticleHandles()[0]->X().Z);
+			EXPECT_EQ(GCW->PhysObject->GetUnorderedParticles_Internal().Num(), 1);
+			EXPECT_EQ(GCW2->PhysObject->GetUnorderedParticles_Internal().Num(), 1);
+			EXPECT_EQ(GCW->PhysObject->GetNumTransforms(), 1);
+			EXPECT_EQ(GCW2->PhysObject->GetNumTransforms(), 1);
+			EXPECT_LT(GCW->PhysObject->GetParticle_Internal(0)->X().Z, GCW2->PhysObject->GetParticle_Internal(0)->X().Z);
 		}
 	}
 
@@ -118,9 +120,11 @@ namespace GeometryCollectionTest
 		{
 			FGeometryCollectionWrapper* GCW = UnitTest.PhysicsObjects[i]->As<FGeometryCollectionWrapper>();
 			FGeometryCollectionWrapper* GCW2 = UnitTest.PhysicsObjects[i + 1]->As<FGeometryCollectionWrapper>();
-			EXPECT_EQ(GCW->PhysObject->GetSolverParticleHandles().Num(), 1);
-			EXPECT_EQ(GCW2->PhysObject->GetSolverParticleHandles().Num(), 1);
-			EXPECT_LT(FMath::Abs(GCW->PhysObject->GetSolverParticleHandles()[0]->X().Z - GCW2->PhysObject->GetSolverParticleHandles()[0]->X().Z), KINDA_SMALL_NUMBER);
+			EXPECT_EQ(GCW->PhysObject->GetUnorderedParticles_Internal().Num(), 1);
+			EXPECT_EQ(GCW2->PhysObject->GetUnorderedParticles_Internal().Num(), 1);
+			EXPECT_EQ(GCW->PhysObject->GetNumTransforms(), 1);
+			EXPECT_EQ(GCW2->PhysObject->GetNumTransforms(), 1);
+			EXPECT_LT(FMath::Abs(GCW->PhysObject->GetParticle_Internal(0)->X().Z - GCW2->PhysObject->GetParticle_Internal(0)->X().Z), KINDA_SMALL_NUMBER);
 		}
 	}
 
@@ -168,8 +172,9 @@ namespace GeometryCollectionTest
 		for (int32 i = 0; i < UnitTest.PhysicsObjects.Num(); i++)
 		{
 			FGeometryCollectionWrapper* GCW = UnitTest.PhysicsObjects[i]->As<FGeometryCollectionWrapper>();
-			EXPECT_EQ(GCW->PhysObject->GetSolverParticleHandles().Num(), 1);
-			EXPECT_EQ(GCW->PhysObject->GetSolverParticleHandles()[0]->Disabled(), true);
+			EXPECT_EQ(GCW->PhysObject->GetUnorderedParticles_Internal().Num(), 1);
+			EXPECT_EQ(GCW->PhysObject->GetNumTransforms(), 1);
+			EXPECT_EQ(GCW->PhysObject->GetParticle_Internal(0)->Disabled(), true);
 		}
 
 		/*
@@ -183,8 +188,9 @@ namespace GeometryCollectionTest
 		for (int32 i = 0; i < UnitTest.PhysicsObjects.Num(); i++)
 		{
 			FGeometryCollectionWrapper* GCW = UnitTest.PhysicsObjects[i]->As<FGeometryCollectionWrapper>();
-			EXPECT_EQ(GCW->PhysObject->GetSolverParticleHandles().Num(), 1);
-			EXPECT_EQ(GCW->PhysObject->GetSolverParticleHandles()[0]->Disabled(), false);
+			EXPECT_EQ(GCW->PhysObject->GetUnorderedParticles_Internal().Num(), 1);
+			EXPECT_EQ(GCW->PhysObject->GetNumTransforms(), 1);
+			EXPECT_EQ(GCW->PhysObject->GetParticle_Internal(0)->Disabled(), false);
 		}
 
 		for (int32 Frame = 1; Frame < 100; Frame++)
@@ -199,13 +205,14 @@ namespace GeometryCollectionTest
 		for (int32 i = 0; i < UnitTest.PhysicsObjects.Num() - 1; i++)
 		{
 			FGeometryCollectionWrapper* GCW = UnitTest.PhysicsObjects[i]->As<FGeometryCollectionWrapper>();
-			EXPECT_EQ(GCW->PhysObject->GetSolverParticleHandles().Num(), 1);
-			EXPECT_EQ(GCW->PhysObject->GetSolverParticleHandles()[0]->Disabled(), true);
+			EXPECT_EQ(GCW->PhysObject->GetUnorderedParticles_Internal().Num(), 1);
+			EXPECT_EQ(GCW->PhysObject->GetNumTransforms(), 1);
+			EXPECT_EQ(GCW->PhysObject->GetParticle_Internal(0)->Disabled(), true);
 		}
 		FGeometryCollectionWrapper* GCW = UnitTest.PhysicsObjects[UnitTest.PhysicsObjects.Num() - 1]->As<FGeometryCollectionWrapper>();
-		EXPECT_EQ(GCW->PhysObject->GetSolverParticleHandles()[0]->Disabled(), false);
+		EXPECT_EQ(GCW->PhysObject->GetParticle_Internal(0)->Disabled(), false);
 
-		EXPECT_LT(GCW->PhysObject->GetSolverParticleHandles()[0]->X().Z, -1.f);					
+		EXPECT_LT(GCW->PhysObject->GetParticle_Internal(0)->X().Z, -1.f);
 	}
 
 
