Fix Warning in Headless in Chaos Vehcile and Cloth Tests

diff --git a/Engine/Source/Programs/HeadlessChaos/Private/HeadlessChaosTestCloth.cpp b/Engine/Source/Programs/HeadlessChaos/Private/HeadlessChaosTestCloth.cpp
--- a/Engine/Source/Programs/HeadlessChaos/Private/HeadlessChaosTestCloth.cpp
+++ b/Engine/Source/Programs/HeadlessChaos/Private/HeadlessChaosTestCloth.cpp
@@ -29,7 +29,7 @@ namespace ChaosTest {
 		const Softs::FSolverReal Damping=0.04)
 	{
 		Chaos::Softs::FSolverParticles Particles;
-		Chaos::Softs::FSolverRigidParticles RigidParticles;
+		Chaos::Softs::FSolverCollisionParticles RigidParticles;
 		TUniquePtr<Softs::FPBDEvolution> Evolution(
 			new Softs::FPBDEvolution(
 				MoveTemp(Particles),
diff --git a/Engine/Source/Programs/HeadlessChaos/Private/HeadlessChaosTestVehicles.cpp b/Engine/Source/Programs/HeadlessChaos/Private/HeadlessChaosTestVehicles.cpp
--- a/Engine/Source/Programs/HeadlessChaos/Private/HeadlessChaosTestVehicles.cpp
+++ b/Engine/Source/Programs/HeadlessChaos/Private/HeadlessChaosTestVehicles.cpp
@@ -1442,10 +1442,12 @@ namespace ChaosTest
 
 	FVector WorldVelocityAtPoint(TPBDRigidParticleHandle<FReal, 3>* Rigid, const FVector& InPoint)
 	{
-		const Chaos::FVec3 COM = Rigid ? Chaos::FParticleUtilitiesGT::GetCoMWorldPosition(Rigid) : Chaos::FParticleUtilitiesGT::GetActorWorldTransform(Rigid).GetTranslation();
+		check(Rigid);
+		const Chaos::FVec3 COM = Chaos::FParticleUtilitiesGT::GetCoMWorldPosition(Rigid);
 		const Chaos::FVec3 Diff = InPoint - COM;
-		return Rigid->GetV() - Chaos::FVec3::CrossProduct(Diff, Rigid->GetW());
-
+		const Chaos::FVec3 V = Rigid->GetV();
+		const Chaos::FVec3 W = Rigid->GetW();
+		return V - Chaos::FVec3::CrossProduct(Diff, W);
 	}
 
 	// #todo: break out vehicle simulation setup so it can be used across number of tests
