Fix shutdown crash in HeadlessChaos unit test runner.

#changelist validated
#virtualized

diff --git a/Engine/Source/Programs/HeadlessChaos/Private/HeadlessChaosApp.cpp b/Engine/Source/Programs/HeadlessChaos/Private/HeadlessChaosApp.cpp
--- a/Engine/Source/Programs/HeadlessChaos/Private/HeadlessChaosApp.cpp
+++ b/Engine/Source/Programs/HeadlessChaos/Private/HeadlessChaosApp.cpp
@@ -547,10 +547,10 @@ INT32_MAIN_INT32_ARGC_TCHAR_ARGV()
 
 	ensure(RUN_ALL_TESTS() == 0);
 
+	RequestEngineExit(TEXT("Exiting"));
 	FEngineLoop::AppPreExit();
 	FModuleManager::Get().UnloadModulesAtShutdown();
 	FEngineLoop::AppExit();
-	FPlatformMisc::RequestExit(false);
 
 	return 0;
 }
