Chaos - fix AABB support vertex index and match Box indices to AABB indices
- unit test to verify box support vertices for all cardinal and diagonal directions
- unit test to verify that AABB and Box vertex indices match

There was a mismatch between the vertex indices on TBox versus TAABB, but also the TAABB's vertex index calculation in the support functions was incorrect. Luckily this didnt matter, but some upcoming changes now require that the Box support functions return the correct vertex index, and that requires the above two fixes.

diff --git a/Engine/Source/Programs/HeadlessChaos/Private/HeadlessChaosTestGJK.cpp b/Engine/Source/Programs/HeadlessChaos/Private/HeadlessChaosTestGJK.cpp
--- a/Engine/Source/Programs/HeadlessChaos/Private/HeadlessChaosTestGJK.cpp
+++ b/Engine/Source/Programs/HeadlessChaos/Private/HeadlessChaosTestGJK.cpp
@@ -1458,9 +1458,12 @@ namespace ChaosTest
 			const FRigidTransform3 BTM(FVec3(2461.92749, -205.484283, 106.071632), FRotation3::FromElements(0,0,0,1));
 			const FRigidTransform3 BToATM(FVec3(102.903252, 218.050415, 102.071655), FRotation3::FromElements(5.07916162e-08, 3.39378659e-08, -0.555569768, 0.831469893));
 
-			FReal Penetration;
-			FVec3 ClosestA,ClosestB,Normal;
-			int32 ClosestVertexIndexA, ClosestVertexIndexB;
+			FReal Penetration = 0;
+			FVec3 ClosestA = FVec3(0);
+			FVec3 ClosestB = FVec3(0);
+			FVec3 Normal = FVec3(0);
+			int32 ClosestVertexIndexA = INDEX_NONE;
+			int32 ClosestVertexIndexB = INDEX_NONE;
 			const FVec3 Offset ={162.072754,-178.514679,-102.071632};
 			EXPECT_TRUE((GJKPenetration<false, FReal>(A,B,BToATM,Penetration,ClosestA,ClosestB,Normal,ClosestVertexIndexA,ClosestVertexIndexB,0,0,Offset)));
 
@@ -2099,24 +2102,4 @@ namespace ChaosTest
 		EXPECT_TRUE(GJKIntersection(BigBoxScaled, SmallBox, BToATM, FReal(0), Chaos::TVector<FReal, 3>(-16000, -16000, 500)));		
 
 	}
-
-
-	// Test capsule support functions when scaled
-	GTEST_TEST(GJKTests, GJK_TestSupportFunctions)
-	{
-		FVec3 PointA{ -1, 0, -100 };
-		FVec3 PointB{ 1, 0, 100 };
-		FVec3 ScaleX{ 100, 1, 1 };
-		FCapsulePtr Capsule( new FCapsule(PointA, PointB, 10));
-		TImplicitObjectScaled<Chaos::FCapsule> CapsuleScaled(Capsule, ScaleX);
-
-		FVec3 SupportDir{0.1f, 0.0f, -1.0f}; // Pointing down and slightly to the right
-
-		int32 Vertex = INDEX_NONE;
-		FVec3 Support = Capsule->SupportCore(SupportDir,0, nullptr, Vertex);
-		EXPECT_EQ(Support, PointA); // Expect bottom point
-
-		Support = CapsuleScaled.SupportCore(SupportDir, 0, nullptr, Vertex);
-		EXPECT_EQ(Support, ScaleX * PointA); // Still expect bottom point (But scaled)
-	}
 }
\ No newline at end of file
diff --git a/Engine/Source/Programs/HeadlessChaos/Private/HeadlessChaosTestSupport.cpp b/Engine/Source/Programs/HeadlessChaos/Private/HeadlessChaosTestSupport.cpp
new file mode 100644
--- /dev/null
+++ b/Engine/Source/Programs/HeadlessChaos/Private/HeadlessChaosTestSupport.cpp
@@ -0,0 +1,132 @@
+// Copyright Epic Games, Inc. All Rights Reserved.
+
+#include "HeadlessChaosTestGJK.h"
+
+#include "HeadlessChaos.h"
+#include "HeadlessChaosTestUtility.h"
+#include "Chaos/GJK.h"
+#include "Chaos/Capsule.h"
+#include "Chaos/Convex.h"
+#include "Chaos/GJK.h"
+#include "Chaos/ImplicitObjectScaled.h"
+#include "Chaos/Collision/PBDCollisionConstraint.h"
+#include "Chaos/Triangle.h"
+#include "Chaos/TriangleRegister.h"
+
+namespace ChaosTest
+{
+	using namespace Chaos;
+
+	// InSupportDirection need not be normalized;
+	// InExpectedSupport is any point on the support plane - we only check the distance along the support direction
+	template<typename BoxType>
+	void CheckBoxSupportDir(const BoxType& Box, const FVec3& InSupportDirection, const FVec3& InExpectedSupport)
+	{
+		const FReal DistanceTolerance = UE_KINDA_SMALL_NUMBER;
+
+		const FVec3 SupportDirection = InSupportDirection.GetSafeNormal();
+		const FReal ExpectedSupportDistance = FVec3::DotProduct(InExpectedSupport, SupportDirection);
+
+		// NOTE: Support functions should return a vertex position that is furthest along the support direct.
+		// When there are multiple options, any of the options should be ok
+		// Note: using non-normalized direction
+		const FReal Margin = 0;
+		FReal SupportDelta = 0;
+		int32 SupportVertexIndex = INDEX_NONE;
+		FVec3 Support = Box.SupportCore(InSupportDirection, Margin, &SupportDelta, SupportVertexIndex);
+
+		// The support point has the expected distance along the support direction
+		const FReal SupportDistance = FVec3::DotProduct(Support, SupportDirection);
+		EXPECT_NEAR(SupportDistance, ExpectedSupportDistance, DistanceTolerance);
+
+		// We have a vertex index
+		EXPECT_NE(SupportVertexIndex, INDEX_NONE);
+		if (SupportVertexIndex != INDEX_NONE)
+		{
+			// SupportVertexIndex is for a point with the correct support distance
+			const FVec3 SupportVertex = Box.GetVertex(SupportVertexIndex);
+			const FReal SupportVertexDistance = FVec3::DotProduct(SupportVertex, SupportDirection);
+			EXPECT_NEAR(SupportVertexDistance, ExpectedSupportDistance, DistanceTolerance);
+
+			// The internal Aabb vertex has the correct support distance
+			const FAABB3 Aabb = Box.BoundingBox();
+			const FVec3 AabbVertex = Aabb.GetVertex(SupportVertexIndex);
+			const FReal AabbVertexDistance = FVec3::DotProduct(AabbVertex, SupportDirection);
+			EXPECT_NEAR(AabbVertexDistance, ExpectedSupportDistance, DistanceTolerance);
+
+			// The Aabb and Box vertices are the same
+			EXPECT_NEAR(AabbVertex.X, SupportVertex.X, DistanceTolerance);
+			EXPECT_NEAR(AabbVertex.Y, SupportVertex.Y, DistanceTolerance);
+			EXPECT_NEAR(AabbVertex.Z, SupportVertex.Z, DistanceTolerance);
+		}
+	}
+
+	template<typename BoxType>
+	void CheckBoxSupportVertex(const BoxType& Box, int32 X, int32 Y, int32 Z)
+	{
+		const FAABB3 Aabb = Box.BoundingBox();
+		const FVec3 UnitCubePoint = FVec3(FReal(X), FReal(Y), FReal(Z));
+		const FVec3 SupportDir = UnitCubePoint;	// No need to normalize
+		const FVec3 SupportVertex = Aabb.GetCenter() + FReal(0.5) * Aabb.Extents() * UnitCubePoint;
+		CheckBoxSupportDir(Box, SupportDir, SupportVertex);
+	}
+
+	template<typename BoxType>
+	void CheckBoxSupportAxes(const BoxType& Box)
+	{
+		// Check all cardinal and diagonal directions
+		for (int32 X = -1; X <= 1; ++X)
+		{
+			for (int32 Y = -1; Y <= 1; ++Y)
+			{
+				for (int32 Z = -1; Z <= 1; ++Z)
+				{
+					CheckBoxSupportVertex(Box, X, Y, Z);
+				}
+			}
+		}
+	}
+
+	// Test SupportCore for all cardinal and diagonal axes
+	GTEST_TEST(SupportTests, TestBoxSupportCore)
+	{
+		const FVec3 Center = FVec3(3.5, -94.1, 15.6);
+		const FVec3 HalfExtent = FVec3(100, 200, 300);
+		const FReal Margin = FReal(0);
+
+		const FImplicitBox3 Box = FImplicitBox3(Center - HalfExtent, Center + HalfExtent, Margin);
+		CheckBoxSupportAxes(Box);
+	}
+
+	GTEST_TEST(SupportTests, TestScaledBoxSupportCore)
+	{
+		const FVec3 Center = FVec3(3.5, -94.1, 15.6);
+		const FVec3 HalfExtent = FVec3(100, 200, 300);
+		const FReal Margin = FReal(0);
+		const FVec3 Scale = FVec3(8.7, 3.2, 1.1);
+
+		TRefCountPtr<FImplicitBox3> Box = TRefCountPtr<FImplicitBox3>(new FImplicitBox3(Center - HalfExtent, Center + HalfExtent, Margin));
+		TRefCountPtr<TImplicitObjectScaled<FImplicitBox3>> BoxScaled = TRefCountPtr<TImplicitObjectScaled<FImplicitBox3>>(new TImplicitObjectScaled<FImplicitBox3>(Box, Scale));
+		CheckBoxSupportAxes(*BoxScaled.GetReference());
+	}
+
+	// Test capsule support functions when scaled
+	GTEST_TEST(SupportTests, TestScaledCapsuleSupportCore)
+	{
+		FVec3 PointA{ -1, 0, -100 };
+		FVec3 PointB{ 1, 0, 100 };
+		FVec3 ScaleX{ 100, 1, 1 };
+		FCapsulePtr Capsule(new FImplicitCapsule3(PointA, PointB, 10));
+		TImplicitObjectScaled<FImplicitCapsule3> CapsuleScaled(Capsule, ScaleX);
+
+		FVec3 SupportDir{ 0.1f, 0.0f, -1.0f }; // Pointing down and slightly to the right
+
+		int32 Vertex = INDEX_NONE;
+		FVec3 Support = Capsule->SupportCore(SupportDir, 0, nullptr, Vertex);
+		EXPECT_EQ(Support, PointA); // Expect bottom point
+
+		Support = CapsuleScaled.SupportCore(SupportDir, 0, nullptr, Vertex);
+		EXPECT_EQ(Support, ScaleX * PointA); // Still expect bottom point (But scaled)
+	}
+
+}
diff --git a/Engine/Source/Runtime/Experimental/Chaos/Private/Chaos/Box.cpp b/Engine/Source/Runtime/Experimental/Chaos/Private/Chaos/Box.cpp
--- a/Engine/Source/Runtime/Experimental/Chaos/Private/Chaos/Box.cpp
+++ b/Engine/Source/Runtime/Experimental/Chaos/Private/Chaos/Box.cpp
@@ -5,16 +5,39 @@
 
 namespace Chaos
 {
+	
+	// Indices below must match TAABB::GetVertex (ascii art below was copied from there)
+	// because TBox uses TAABB::SupportCore which returns a vertex index
+	// NOTE: This shows a right-handed coordinate system, but we are left-handed...
+	//
+	//	   6---------7
+	//	  /|        /|
+	//	 / |       / |
+	//	4---------5  |
+	//	|  |      |  |
+	//	|  2------|--3
+	//	| /       | /
+	//	|/        |/
+	//	0---------1
+	//
+	//  Z
+	//  |  Y
+	//  | /
+	//  |/
+	//  o-----X
+	// 
+	//
+
 	FConvexHalfEdgeStructureDataS16 MakeBoxStructureData()
 	{
 		TArray<TArray<int32>> PlaneVertices
 		{
-			{7, 5, 4, 6},	// X
-			{2, 3, 7, 6},	// Y
-			{5, 7, 3, 1},	// Z,
-			{2, 0, 1, 3},	// -X
-			{0, 4, 5, 1},	// -Y
-			{0, 2, 6, 4},	// -Z,
+			{ 0, 4, 6, 2 },	// -X,
+			{ 0, 1, 5, 4 },	// -Y
+			{ 0, 2, 3, 1 },	// -Z
+			{ 1, 3, 7, 5 },	//  X
+			{ 2, 6, 7, 3 },	//  Y
+			{ 4, 5, 7, 6 },	//  Z
 		};
 
 		return FConvexHalfEdgeStructureDataS16::MakePlaneVertices(PlaneVertices, 8);
@@ -22,24 +45,24 @@ namespace Chaos
 
 	template<typename T, int D> TArray<FVec3> TBox<T, D>::SNormals =
 	{
-		FVec3(1,0,0),	 // X
-		FVec3(0,1,0),	 // Y
-		FVec3(0,0,1),	 // Z
-		FVec3(-1,0,0), // -X
-		FVec3(0,-1,0), // -Y
-		FVec3(0,0,-1), // -Z
+		FVec3(-1,0,0),	// -X
+		FVec3(0,-1,0),	// -Y
+		FVec3(0,0,-1),	// -Z
+		FVec3(1,0,0),	//  X
+		FVec3(0,1,0),	//  Y
+		FVec3(0,0,1),	//  Z
 	};
 	
 	template<typename T, int D> TArray<FVec3> TBox<T, D>::SVertices =
 	{
-		FVec3(-1,-1,-1),
-		FVec3(-1,-1,1),
-		FVec3(-1,1,-1),
-		FVec3(-1,1,1),
-		FVec3(1,-1,-1),
-		FVec3(1,-1,1),
-		FVec3(1,1,-1),
-		FVec3(1,1,1),
+		FVec3(-1,-1,-1),	// 0 
+		FVec3( 1,-1,-1),	// 1 
+		FVec3(-1, 1,-1),	// 2 
+		FVec3( 1, 1,-1),	// 3 
+		FVec3(-1,-1, 1),	// 4 
+		FVec3( 1,-1, 1),	// 5 
+		FVec3(-1, 1, 1),	// 6 
+		FVec3( 1, 1, 1),	// 7 
 	};
 
 	template<typename T, int D> FConvexHalfEdgeStructureDataS16 TBox<T, D>::SStructureData = MakeBoxStructureData();
diff --git a/Engine/Source/Runtime/Experimental/Chaos/Public/Chaos/AABB.h b/Engine/Source/Runtime/Experimental/Chaos/Public/Chaos/AABB.h
--- a/Engine/Source/Runtime/Experimental/Chaos/Public/Chaos/AABB.h
+++ b/Engine/Source/Runtime/Experimental/Chaos/Public/Chaos/AABB.h
@@ -311,7 +311,7 @@ namespace Chaos
 			return BestNormal;
 		}
 		
-		FORCEINLINE_DEBUGGABLE TVector<T, d> Support(const TVector<T, d>& Direction, const T Thickness, int32& VertexIndex) const
+		FORCEINLINE_DEBUGGABLE TVector<T, d> Support(const TVector<T, d>& Direction, const T Thickness, int32& OutVertexIndex) const
 		{
 			TVector<T, d> ChosenPt;
 			FIntVector ChosenAxis;
@@ -328,7 +328,8 @@ namespace Chaos
 					ChosenAxis[Axis] = 1;
 				}
 			}
-			VertexIndex =  ChosenAxis[0] * 4 + ChosenAxis[1] * 2 + ChosenAxis[2];
+
+			OutVertexIndex = GetIndex(ChosenAxis);
 
 			if (Thickness != (T)0)
 			{
@@ -349,7 +350,7 @@ namespace Chaos
 		}
 
 		// Support vertex in the specified direction, assuming each face has been moved inwards by InMargin
-		FORCEINLINE_DEBUGGABLE FVec3 SupportCore(const FVec3& Direction, const FReal InMargin, FReal* OutSupportDelta, int32& VertexIndex) const
+		FORCEINLINE_DEBUGGABLE FVec3 SupportCore(const FVec3& Direction, const FReal InMargin, FReal* OutSupportDelta, int32& OutVertexIndex) const
 		{
 			FVec3 ChosenPt;
 			FIntVector ChosenAxis;
@@ -366,7 +367,9 @@ namespace Chaos
 					ChosenAxis[Axis] = 1;
 				}
 			}
-			VertexIndex = ChosenAxis[0] * 4 + ChosenAxis[1] * 2 + ChosenAxis[2];
+
+			OutVertexIndex = GetIndex(ChosenAxis);
+
 			// Maximum distance between the Core+Margin position and the original outer vertex
 			constexpr FReal RootThreeMinusOne = FReal(1.7320508075688772935274463415059 - 1.0);
 			if (OutSupportDelta != nullptr)
@@ -386,7 +389,7 @@ namespace Chaos
 			return MakeVectorRegisterFloatFromDouble(MakeVectorRegister(SupportVert.X, SupportVert.Y, SupportVert.Z, 0.0));
 		}
 
-		FORCEINLINE_DEBUGGABLE TVector<T, d> SupportCoreScaled(const TVector<T, d>& Direction, const T InMargin, const TVector<T, d>& Scale, T* OutSupportDelta, int32& VertexIndex) const
+		FORCEINLINE_DEBUGGABLE TVector<T, d> SupportCoreScaled(const TVector<T, d>& Direction, const T InMargin, const TVector<T, d>& Scale, T* OutSupportDelta, int32& OutVertexIndex) const
 		{
 			const TVector<T, d> ScaledDirection = Direction * Scale; // Compensate for Negative scales, only the signs really matter here
 
@@ -405,8 +408,9 @@ namespace Chaos
 					ChosenAxis[Axis] = 1;
 				}
 			}
-			VertexIndex = ChosenAxis[0] * 4 + ChosenAxis[1] * 2 + ChosenAxis[2];
-			
+
+			OutVertexIndex = GetIndex(ChosenAxis);
+
 			constexpr T RootThreeMinusOne = T(1.7320508075688772935274463415059 - 1.0);
 			if (OutSupportDelta != nullptr)
 			{
@@ -495,12 +499,26 @@ namespace Chaos
 		FORCEINLINE TVector<T, d> GetVertex(const int32 Index) const
 		{
 			check(0 <= Index && Index < 8);
+
+			// See GetIndex() for reverse logic
 			return TVector<T, d>(
 				(Index & (1 << 0)) == 0 ? MMin.X : MMax.X,
 				(Index & (1 << 1)) == 0 ? MMin.Y : MMax.Y,
 				(Index & (1 << 2)) == 0 ? MMin.Z : MMax.Z);
 		}
 
+		/**
+		* Given a point on a unit cube expressed as an IntVector of 0s and 1s, get the vertex index of that point.
+		* 0 represents a negative axis and 1 represents a positive axis. All other values are invalid, but we do 
+		* not check this as it would add overhead to support functions. This is used by the support functions and 
+		* must invert the index logic in GetVertex().
+		*/
+		FORCEINLINE int32 GetIndex(const FIntVector& AxisSelector) const
+		{
+			// See GetVertex() for reverse logic
+			return AxisSelector[0] + AxisSelector[1] * 2 + AxisSelector[2] * 4;
+		}
+
 		/**
 		*	Get an array of two indices into the vertex list for one of the
 		*	twelve edges of the box. Edges are ordered by increasing vertex
