Chaos - failing EPA test case (disabled)

diff --git a/Engine/Source/Programs/HeadlessChaos/Private/HeadlessChaosTestEPA.cpp b/Engine/Source/Programs/HeadlessChaos/Private/HeadlessChaosTestEPA.cpp
--- a/Engine/Source/Programs/HeadlessChaos/Private/HeadlessChaosTestEPA.cpp
+++ b/Engine/Source/Programs/HeadlessChaos/Private/HeadlessChaosTestEPA.cpp
@@ -1111,10 +1111,62 @@ namespace ChaosTest
 
 	}
 
-	// Currently broken EPA edge cases. As they are fixed move them to EPARealFailures_Fixed above so that we can ensure they don't break again.
-	GTEST_TEST(EPATests, EPARealFailures_Broken)
+	// Currently broken EPA edge cases
+	// A box above a triangle, almost exactly parallel and touching.
+	// EPA fails due to numerical error and returns an very bad contact.
+	// We hit this condition in EPA: if (UpperBound <= UpperBoundTolerance)
+	// but have previously rejected all of the actual closest faces
+	// because of numerical error.
+	GTEST_TEST(EPATests, DISABLED_EPARealFailures_TouchingBoxTriangle)
 	{
-		
+		{
+			FImplicitBox3 Box = FImplicitBox3(
+				FVec3(-50.000000000000000, -50.000000000000000, -15.000000000000000),
+				FVec3(50.000000000000000, 50.000000000000000, 15.000000000000000)
+			);
+
+			FTriangle Triangle = FTriangle(
+				FVec3(94.478362706670822, -120.65494588586357, -14.999999386949069),
+				FVec3(89.056288683336533, 179.29605196669289, -15.000000556768336),
+				FVec3(-210.89470916921991, 173.87397794335860, -15.000000537575422)
+			);
+
+			const TGJKShape<FImplicitBox3> GJKConvex(Box);
+			const TGJKShape<FTriangle> GJKTriangle(Triangle);
+
+			const FReal GJKEpsilon = 1.e-6;
+			const FReal EPAEpsilon = 1.e-6;
+			FReal UnusedMaxMarginDelta = FReal(0);
+			int32 ConvexVertexIndex = INDEX_NONE;
+			int32 TriangleVertexIndex = INDEX_NONE;
+			FReal Penetration;
+			FVec3 ConvexClosest, TriangleClosest, ConvexNormal;
+			FVec3 InitialGJKDir = FVec3(1, 0, 0);
+
+			const bool bHaveContact = GJKPenetrationSameSpace(
+				GJKConvex,
+				GJKTriangle,
+				Penetration,
+				ConvexClosest,
+				TriangleClosest,
+				ConvexNormal,
+				ConvexVertexIndex,
+				TriangleVertexIndex,
+				UnusedMaxMarginDelta,
+				InitialGJKDir,
+				GJKEpsilon, EPAEpsilon);
+
+			EXPECT_TRUE(bHaveContact);
+
+			// Should be touching
+			EXPECT_NEAR(Penetration, 0, UE_KINDA_SMALL_NUMBER);
+
+			// Normal should point directly down
+			EXPECT_NEAR(ConvexNormal.Z, -1, UE_KINDA_SMALL_NUMBER);
+
+			// Contact should be on bottom of box
+			EXPECT_NEAR(ConvexClosest.Z, -15.0, UE_KINDA_SMALL_NUMBER);
+		}
 	}
 
 	//
