Make consistent API by using GetX, SetX, GetP and SetP for physics particles

diff --git a/Engine/Plugins/Experimental/ChaosCaching/Source/ChaosCaching/Private/Chaos/Adapters/GeometryCollectionComponentCacheAdapter.cpp b/Engine/Plugins/Experimental/ChaosCaching/Source/ChaosCaching/Private/Chaos/Adapters/GeometryCollectionComponentCacheAdapter.cpp
--- a/Engine/Plugins/Experimental/ChaosCaching/Source/ChaosCaching/Private/Chaos/Adapters/GeometryCollectionComponentCacheAdapter.cpp
+++ b/Engine/Plugins/Experimental/ChaosCaching/Source/ChaosCaching/Private/Chaos/Adapters/GeometryCollectionComponentCacheAdapter.cpp
@@ -153,7 +153,7 @@ namespace Chaos
 					Pending.ParticleIndex = TransformIndex;
 					
 					// All recorded transforms are in actor space, ie relative to the Cache Manager making the recording.
-					FTransform LocalTransform = MassToLocal[TransformIndex].Inverse() * FTransform(Handle->GetR(), Handle->X());
+					FTransform LocalTransform = MassToLocal[TransformIndex].Inverse() * FTransform(Handle->GetR(), Handle->GetX());
 					FTransform ActorSpaceTransform = LocalTransform * WorldToActor;
 					Pending.PendingTransform = ActorSpaceTransform;
 				}
@@ -176,7 +176,7 @@ namespace Chaos
 					OutFrame.PendingParticleData.AddDefaulted();
 					FPendingParticleWrite& Pending = OutFrame.PendingParticleData.Last();
 					Pending.ParticleIndex = TransformIndex;
-					FTransform LocalTransform = MassToLocal[TransformIndex].Inverse() * FTransform(Handle->GetR(), Handle->X());
+					FTransform LocalTransform = MassToLocal[TransformIndex].Inverse() * FTransform(Handle->GetR(), Handle->GetX());
 					FTransform ActorSpaceTransform = LocalTransform * WorldToActor;
 					Pending.PendingTransform = ActorSpaceTransform;
 					Pending.bPendingDeactivate = true;
@@ -490,7 +490,7 @@ namespace Chaos
 
 				Handle->SetP(WorldTransform.GetTranslation());
 				Handle->SetQ(WorldTransform.GetRotation());
-				Handle->SetX(Handle->P());
+				Handle->SetX(Handle->GetP());
 				Handle->SetRf(Handle->GetQf());
 				
 				Handle->UpdateWorldSpaceState(WorldTransform, FVec3(0));
diff --git a/Engine/Plugins/Experimental/ChaosCaching/Source/ChaosCaching/Private/Chaos/Adapters/StaticMeshComponentCacheAdapter.cpp b/Engine/Plugins/Experimental/ChaosCaching/Source/ChaosCaching/Private/Chaos/Adapters/StaticMeshComponentCacheAdapter.cpp
--- a/Engine/Plugins/Experimental/ChaosCaching/Source/ChaosCaching/Private/Chaos/Adapters/StaticMeshComponentCacheAdapter.cpp
+++ b/Engine/Plugins/Experimental/ChaosCaching/Source/ChaosCaching/Private/Chaos/Adapters/StaticMeshComponentCacheAdapter.cpp
@@ -53,7 +53,7 @@ namespace Chaos
 				FPendingParticleWrite NewData;
 
 				NewData.ParticleIndex = 0; // Only one particle for static caches
-				NewData.PendingTransform = FTransform(AsRigid->GetR(), AsRigid->X()).GetRelativeTransform(InRootTransform);
+				NewData.PendingTransform = FTransform(AsRigid->GetR(), AsRigid->GetX()).GetRelativeTransform(InRootTransform);
 
 				OutFrame.PendingParticleData.Add(MoveTemp(NewData));
 			}
diff --git a/Engine/Plugins/Experimental/ChaosCloth/Source/Chaos/Private/ChaosCloth/ChaosClothingSimulationSolver.cpp b/Engine/Plugins/Experimental/ChaosCloth/Source/Chaos/Private/ChaosCloth/ChaosClothingSimulationSolver.cpp
--- a/Engine/Plugins/Experimental/ChaosCloth/Source/Chaos/Private/ChaosCloth/ChaosClothingSimulationSolver.cpp
+++ b/Engine/Plugins/Experimental/ChaosCloth/Source/Chaos/Private/ChaosCloth/ChaosClothingSimulationSolver.cpp
@@ -258,7 +258,8 @@ FClothingSimulationSolver::FClothingSimulationSolver(bool bForceBasedSolver, FCl
 					{
 						const int32 SubBoneIndexLocal = SubBoneIndices[SubBoneIdx] - ParticlesInput.GetOffset();
 						checkSlow(SubBoneIndexLocal < Index);
-						SubBoneTransforms[SubBoneIdx] = TRigidTransform<FReal, 3>(ParticlesInput.X(SubBoneIndexLocal), ParticlesInput.R(SubBoneIndexLocal)) * RootTransformInv;
+						const Softs::FSolverCollisionParticlesRange& ConstParticlesInput = ParticlesInput;
+						SubBoneTransforms[SubBoneIdx] = TRigidTransform<FReal, 3>(ParticlesInput.X(SubBoneIndexLocal), ConstParticlesInput.R(SubBoneIndexLocal)) * RootTransformInv;
 					}
 					SkinnedLevelSet->DeformPoints(SubBoneTransforms);
 					SkinnedLevelSet->UpdateSpatialHierarchy();
@@ -312,13 +313,13 @@ FClothingSimulationSolver::FClothingSimulationSolver(bool bForceBasedSolver, FCl
 		PBDEvolution->SetCollisionKinematicUpdateFunction(
 			[this](Softs::FSolverCollisionParticles& ParticlesInput, const Softs::FSolverReal Dt, const Softs::FSolverReal LocalTime, const int32 Index)
 		{
-			LastSubframeCollisionTransformsCCD[Index] = Softs::FSolverRigidTransform3(ParticlesInput.X(Index), ParticlesInput.GetR(Index));
+			LastSubframeCollisionTransformsCCD[Index] = Softs::FSolverRigidTransform3(ParticlesInput.GetX(Index), ParticlesInput.GetR(Index));
 
 			checkSlow(Dt > SMALL_NUMBER && DeltaTime > SMALL_NUMBER);
 			const Softs::FSolverReal Alpha = (LocalTime - Time) / DeltaTime;
 			const Softs::FSolverVec3 NewX =
 				Alpha * CollisionTransforms[Index].GetTranslation() + ((Softs::FSolverReal)1. - Alpha) * OldCollisionTransforms[Index].GetTranslation();
-			ParticlesInput.V(Index) = (NewX - ParticlesInput.X(Index)) / Dt;
+			ParticlesInput.V(Index) = (NewX - ParticlesInput.GetX(Index)) / Dt;
 			ParticlesInput.X(Index) = NewX;
 			const Softs::FSolverRotation3 NewR = Softs::FSolverRotation3::Slerp(OldCollisionTransforms[Index].GetRotation(), CollisionTransforms[Index].GetRotation(), Alpha);
 			const Softs::FSolverRotation3 Delta = NewR * ParticlesInput.GetR(Index).Inverse();
@@ -331,13 +332,13 @@ FClothingSimulationSolver::FClothingSimulationSolver(bool bForceBasedSolver, FCl
 				const_cast<FImplicitObject*>(ParticlesInput.GetGeometry(Index).GetReference())->GetObject<TWeightedLatticeImplicitObject<FLevelSet>>())
 			{
 				const TArray<int32>& SubBoneIndices = SkinnedLevelSet->GetSolverBoneIndices();
-				const FTransform RootTransformInv = TRigidTransform<FReal, 3>(ParticlesInput.X(Index), ParticlesInput.GetR(Index)).Inverse();
+				const FTransform RootTransformInv = TRigidTransform<FReal, 3>(ParticlesInput.GetX(Index), ParticlesInput.GetR(Index)).Inverse();
 				TArray<FTransform> SubBoneTransforms;
 				SubBoneTransforms.SetNum(SubBoneIndices.Num());
 				for (int32 SubBoneIdx = 0; SubBoneIdx < SubBoneIndices.Num(); ++SubBoneIdx)
 				{
 					checkSlow(SubBoneIndices[SubBoneIdx] < Index);
-					SubBoneTransforms[SubBoneIdx] = TRigidTransform<FReal, 3>(ParticlesInput.X(SubBoneIndices[SubBoneIdx]), ParticlesInput.GetR(SubBoneIndices[SubBoneIdx])) * RootTransformInv;
+					SubBoneTransforms[SubBoneIdx] = TRigidTransform<FReal, 3>(ParticlesInput.GetX(SubBoneIndices[SubBoneIdx]), ParticlesInput.GetR(SubBoneIndices[SubBoneIdx])) * RootTransformInv;
 				}
 				SkinnedLevelSet->DeformPoints(SubBoneTransforms);
 				SkinnedLevelSet->UpdateSpatialHierarchy();
diff --git a/Engine/Plugins/Experimental/ChaosModularVehicle/Source/ChaosModularVehicleEngine/Private/ChaosModularVehicle/ModularVehicleSimulationCU.cpp b/Engine/Plugins/Experimental/ChaosModularVehicle/Source/ChaosModularVehicleEngine/Private/ChaosModularVehicle/ModularVehicleSimulationCU.cpp
--- a/Engine/Plugins/Experimental/ChaosModularVehicle/Source/ChaosModularVehicleEngine/Private/ChaosModularVehicle/ModularVehicleSimulationCU.cpp
+++ b/Engine/Plugins/Experimental/ChaosModularVehicle/Source/ChaosModularVehicleEngine/Private/ChaosModularVehicle/ModularVehicleSimulationCU.cpp
@@ -157,9 +157,9 @@ void FModularVehicleSimulationCU::PerformAdditionalSimWork(UWorld* InWorld, cons
 			if (FPBDRigidClusteredParticleHandle* ClusterHandle = ClusterUnion->InternalCluster)
 			{
 				TArray<FPBDRigidParticleHandle*> Particles = ClusterUnion->ChildParticles;
-				FRigidTransform3 ClusterWorldTM = FRigidTransform3(ClusterHandle->X(), ClusterHandle->GetR());
+				FRigidTransform3 ClusterWorldTM = FRigidTransform3(ClusterHandle->GetX(), ClusterHandle->GetR());
 
-				//Chaos::FDebugDrawQueue::GetInstance().DrawDebugCoordinateSystem(ClusterHandle->X(), FRotator(ClusterHandle->R()), 200, false, -1.f, 1.f, 3.0f);
+				//Chaos::FDebugDrawQueue::GetInstance().DrawDebugCoordinateSystem(ClusterHandle->GetX(), FRotator(ClusterHandle->R()), 200, false, -1.f, 1.f, 3.0f);
 
 				const TArray<Chaos::FSimModuleTree::FSimModuleNode>& ModuleArray = SimModuleTree->GetSimulationModuleTree();
 
@@ -183,7 +183,7 @@ void FModularVehicleSimulationCU::PerformAdditionalSimWork(UWorld* InWorld, cons
 						{
 							Frame = ClusterChild->ChildToParent();
 
-							const FRigidTransform3 ChildWorldTM(Child->X(), Child->GetR());
+							const FRigidTransform3 ChildWorldTM(Child->GetX(), Child->GetR());
 							Frame = ChildWorldTM.GetRelativeTransform(ClusterWorldTM);
 						}
 
diff --git a/Engine/Plugins/Experimental/ChaosVD/Source/ChaosVD/Private/ChaosVDTriMeshGenerator.cpp b/Engine/Plugins/Experimental/ChaosVD/Source/ChaosVD/Private/ChaosVDTriMeshGenerator.cpp
--- a/Engine/Plugins/Experimental/ChaosVD/Source/ChaosVD/Private/ChaosVDTriMeshGenerator.cpp
+++ b/Engine/Plugins/Experimental/ChaosVD/Source/ChaosVD/Private/ChaosVDTriMeshGenerator.cpp
@@ -15,7 +15,7 @@ void FChaosVDTriMeshGenerator::GenerateFromTriMesh(const Chaos::FTriangleMeshImp
 	// Fill the vertex buffer with the transformed vertices of the TriMesh Shape
 	for (int32 i = 0; i < NumVertices; i++)
 	{
-		Vertices[i] = FVector3d(UE::Math::TVector<double>(InTriMesh.Particles().X(i)));
+		Vertices[i] = FVector3d(UE::Math::TVector<double>(InTriMesh.Particles().GetX(i)));
 	}
 
 	const Chaos::FTrimeshIndexBuffer& IdxBuffer = InTriMesh.Elements();
diff --git a/Engine/Plugins/Experimental/Mover/Source/Mover/Private/Backends/MoverNetworkPhysicsLiaison.cpp b/Engine/Plugins/Experimental/Mover/Source/Mover/Private/Backends/MoverNetworkPhysicsLiaison.cpp
--- a/Engine/Plugins/Experimental/Mover/Source/Mover/Private/Backends/MoverNetworkPhysicsLiaison.cpp
+++ b/Engine/Plugins/Experimental/Mover/Source/Mover/Private/Backends/MoverNetworkPhysicsLiaison.cpp
@@ -800,12 +800,12 @@ void UMoverNetworkPhysicsLiaisonComponent::OnPreSimulate_Internal(const FPhysics
 			FFloorCheckResult LastFloorResult;
 			if (Blackboard->TryGet(KinematicBlackboard::LastFloorResult, LastFloorResult))
 			{
-				LocalGroundVelocity = UPhysicsMovementUtils::ComputeGroundVelocityFromHitResult(CharacterParticle->X(), LastFloorResult.HitResult, TickParams.DeltaTimeSeconds);
+				LocalGroundVelocity = UPhysicsMovementUtils::ComputeGroundVelocityFromHitResult(CharacterParticle->GetX(), LastFloorResult.HitResult, TickParams.DeltaTimeSeconds);
 				LocalGroundVelocity -= LocalGroundVelocity.ProjectOnToNormal(LastFloorResult.HitResult.ImpactNormal);
 			}
 		}
 	}
-	SyncState.SetTransforms_WorldSpace(CharacterParticle->X(), FRotator(CharacterParticle->GetR()), CharacterParticle->GetV() - LocalGroundVelocity);
+	SyncState.SetTransforms_WorldSpace(CharacterParticle->GetX(), FRotator(CharacterParticle->GetR()), CharacterParticle->GetV() - LocalGroundVelocity);
 
 	///////////////////////////////////////////////////////////////////////////////////////////////////////////////////////////
 	// Update the simulation
@@ -865,7 +865,7 @@ void UMoverNetworkPhysicsLiaisonComponent::OnPreSimulate_Internal(const FPhysics
 		return;
 	}
 
-	FVector TargetDeltaPos = OutputSyncState->GetLocation_WorldSpace() - CharacterParticle->X();
+	FVector TargetDeltaPos = OutputSyncState->GetLocation_WorldSpace() - CharacterParticle->GetX();
 
 	if (TargetDeltaPos.SizeSquared2D() > GPhysicsDrivenMotionDebugParams.TeleportThreshold * GPhysicsDrivenMotionDebugParams.TeleportThreshold)
 	{
@@ -880,7 +880,7 @@ void UMoverNetworkPhysicsLiaisonComponent::OnPreSimulate_Internal(const FPhysics
 	{
 		if (const UPhysicsDrivenWalkingMode* WalkingMode = Cast<UPhysicsDrivenWalkingMode>(MoverComp->FindMovementMode(UPhysicsDrivenWalkingMode::StaticClass())))
 		{
-			LocalGroundVelocity = UPhysicsMovementUtils::ComputeGroundVelocityFromHitResult(CharacterParticle->X(), Output.FloorResult.HitResult, TickParams.DeltaTimeSeconds);
+			LocalGroundVelocity = UPhysicsMovementUtils::ComputeGroundVelocityFromHitResult(CharacterParticle->GetX(), Output.FloorResult.HitResult, TickParams.DeltaTimeSeconds);
 			LocalGroundVelocity -= LocalGroundVelocity.ProjectOnToNormal(Output.FloorResult.HitResult.ImpactNormal);
 			TargetVelocity += WalkingMode->FractionalVelocityToTarget * LocalGroundVelocity;
 		}
@@ -1032,7 +1032,7 @@ void UMoverNetworkPhysicsLiaisonComponent::OnContactModification_Internal(const
 		float PawnRadius;
 		MoverComp->UpdatedComponent->CalcBoundingCylinder(PawnRadius, PawnHalfHeight);
 
-		const float CharacterHeight = CharacterParticle->X().Z;
+		const float CharacterHeight = CharacterParticle->GetX().Z;
 		const float EndCapHeight = CharacterHeight - PawnHalfHeight + PawnRadius;
 
 		const float CosThetaMax = 0.97f;
diff --git a/Engine/Plugins/Experimental/Mover/Source/Mover/Private/PhysicsMover/PhysicsMovementUtils.cpp b/Engine/Plugins/Experimental/Mover/Source/Mover/Private/PhysicsMover/PhysicsMovementUtils.cpp
--- a/Engine/Plugins/Experimental/Mover/Source/Mover/Private/PhysicsMover/PhysicsMovementUtils.cpp
+++ b/Engine/Plugins/Experimental/Mover/Source/Mover/Private/PhysicsMover/PhysicsMovementUtils.cpp
@@ -246,12 +246,12 @@ FVector UPhysicsMovementUtils::ComputeGroundVelocityFromHitResult(const FVector&
 	FVector GroundVelocity = FVector::ZeroVector;
 	if (const Chaos::FPBDRigidParticleHandle* Rigid = GetRigidParticelHandleFromHitResult(FloorHit))
 	{
-		FVector Offset = CharacterPosition - Rigid->X();
+		FVector Offset = CharacterPosition - Rigid->GetX();
 		Offset -= Offset.ProjectOnToNormal(FloorHit.ImpactNormal);
 
 		if (Rigid->KinematicTarget().IsSet())
 		{
-			const FVector LinearDisplacement = Rigid->KinematicTarget().GetTargetPosition() - Rigid->X();
+			const FVector LinearDisplacement = Rigid->KinematicTarget().GetTargetPosition() - Rigid->GetX();
 			const FQuat RelativeQuat = Rigid->GetR().Inverse() * Rigid->KinematicTarget().GetTargetRotation();
 			const FVector AngularDisplacement = RelativeQuat.ToRotationVector();
 			GroundVelocity = (LinearDisplacement + AngularDisplacement.Cross(Offset)) / DeltaSeconds;
diff --git a/Engine/Plugins/Runtime/NetworkPrediction/Source/NetworkPrediction/Private/NetworkPredictionPhysics.cpp b/Engine/Plugins/Runtime/NetworkPrediction/Source/NetworkPrediction/Private/NetworkPredictionPhysics.cpp
--- a/Engine/Plugins/Runtime/NetworkPrediction/Source/NetworkPrediction/Private/NetworkPredictionPhysics.cpp
+++ b/Engine/Plugins/Runtime/NetworkPrediction/Source/NetworkPrediction/Private/NetworkPredictionPhysics.cpp
@@ -102,7 +102,7 @@ bool FNetworkPredictionPhysicsState::ShouldReconcile(int32 PhysicsFrame, Chaos::
 
 	const Chaos::FGeometryParticleState LocalState = RewindData->GetPastStateAtFrame(*Handle->GetHandle_LowLevel(), PhysicsFrame);
 
-	if (CompareVector(LocalState.X(), RecvState->Location, NetworkPredictionPhysicsCvars::ToleranceX(), "X:"))
+	if (CompareVector(LocalState.GetX(), RecvState->Location, NetworkPredictionPhysicsCvars::ToleranceX(), "X:"))
 	{
 		return true;
 	}
@@ -150,7 +150,7 @@ void FNetworkPredictionPhysicsState::ToString(int32 PhysicsFrame, Chaos::FRewind
 	FPhysicsActorHandle& Handle = BodyInstance->GetPhysicsActorHandle();
 
 	const Chaos::FGeometryParticleState LocalState = RewindData->GetPastStateAtFrame(*Handle->GetHandle_LowLevel(), PhysicsFrame);
-	ToStringInternal(LocalState.X(), LocalState.GetR(), LocalState.GetV(), LocalState.GetW(), Builder);
+	ToStringInternal(LocalState.GetX(), LocalState.GetR(), LocalState.GetV(), LocalState.GetW(), Builder);
 }
 
 void FNetworkPredictionPhysicsState::ToString(FBodyInstance* BodyInstance, FAnsiStringBuilderBase& Builder)
diff --git a/Engine/Source/Programs/HeadlessChaos/Private/GeometryCollection/GeometryCollectionTestClustering.cpp b/Engine/Source/Programs/HeadlessChaos/Private/GeometryCollection/GeometryCollectionTestClustering.cpp
--- a/Engine/Source/Programs/HeadlessChaos/Private/GeometryCollection/GeometryCollectionTestClustering.cpp
+++ b/Engine/Source/Programs/HeadlessChaos/Private/GeometryCollection/GeometryCollectionTestClustering.cpp
@@ -1043,9 +1043,9 @@ namespace GeometryCollectionTest
 					ParticleHandles[8]->Disabled() == false)
 				{
 					bValidInitialState = true;
-					Ref0 = ParticleHandles[0]->X();
-					Ref1 = ParticleHandles[1]->X();
-					Ref6 = ParticleHandles[6]->X();
+					Ref0 = ParticleHandles[0]->GetX();
+					Ref1 = ParticleHandles[1]->GetX();
+					Ref6 = ParticleHandles[6]->GetX();
 
 					// Test kinematic particles have valid (0.0) inverse mass and have the kinematic object state set
 					EXPECT_EQ(ParticleHandles[7]->InvM(), 0.f); // kinematic cluster
@@ -1067,9 +1067,9 @@ namespace GeometryCollectionTest
 					ParticleHandles[8]->Disabled() == true)
 				{
 					bParticle8SucessfulRelease = true;
-					FVector X0 = ParticleHandles[0]->X();
-					FVector X1 = ParticleHandles[1]->X();
-					FVector X6 = ParticleHandles[6]->X();
+					FVector X0 = ParticleHandles[0]->GetX();
+					FVector X1 = ParticleHandles[1]->GetX();
+					FVector X6 = ParticleHandles[6]->GetX();
 
 					FVector X00 = Ref0;
 					FVector X11 = Ref1;
@@ -1126,9 +1126,9 @@ namespace GeometryCollectionTest
 					ParticleHandles[8]->Disabled() == true)
 				{
 					bParticle7SucessfulRelease = true;
-					FVector X0 = ParticleHandles[0]->X();
-					FVector X1 = ParticleHandles[1]->X();
-					FVector X6 = ParticleHandles[6]->X();
+					FVector X0 = ParticleHandles[0]->GetX();
+					FVector X1 = ParticleHandles[1]->GetX();
+					FVector X6 = ParticleHandles[6]->GetX();
 
 					// 0 is a dynamic unclustered body (was owned by cluster 8), check that it's moved since declustering
 					EXPECT_GT(FMath::Abs(X0.Size() - Ref0.Size()), KINDA_SMALL_NUMBER);
@@ -1172,9 +1172,9 @@ namespace GeometryCollectionTest
 					ParticleHandles[8]->Disabled() == true)
 				{
 					bValidFinalActiveState = true;
-					FVector X0 = ParticleHandles[0]->X();
-					FVector X1 = ParticleHandles[1]->X();
-					FVector X6 = ParticleHandles[6]->X();
+					FVector X0 = ParticleHandles[0]->GetX();
+					FVector X1 = ParticleHandles[1]->GetX();
+					FVector X6 = ParticleHandles[6]->GetX();
 
 					// 0 is a dynamic unclustered body (was owned by cluster 8), check that it's moved since declustering
 					EXPECT_GT(FMath::Abs(X0.Size() - Ref0.Size()), KINDA_SMALL_NUMBER);
@@ -1359,9 +1359,9 @@ namespace GeometryCollectionTest
 					ParticleHandles[8]->Disabled() == false)
 				{
 					bValidInitialState = true;
-					Ref0 = ParticleHandles[0]->X();
-					Ref1 = ParticleHandles[1]->X();
-					Ref6 = ParticleHandles[6]->X();
+					Ref0 = ParticleHandles[0]->GetX();
+					Ref1 = ParticleHandles[1]->GetX();
+					Ref6 = ParticleHandles[6]->GetX();
 
 					// Test static particles have valid (0.0) inverse mass and have the static object state set
 					EXPECT_EQ(ParticleHandles[7]->InvM(), 0.f); // kinematic cluster
@@ -1383,9 +1383,9 @@ namespace GeometryCollectionTest
 					ParticleHandles[8]->Disabled() == true)
 				{
 					bParticle8SucessfulRelease = true;
-					FVector X0 = ParticleHandles[0]->X();
-					FVector X1 = ParticleHandles[1]->X();
-					FVector X6 = ParticleHandles[6]->X();
+					FVector X0 = ParticleHandles[0]->GetX();
+					FVector X1 = ParticleHandles[1]->GetX();
+					FVector X6 = ParticleHandles[6]->GetX();
 
 					EXPECT_NEAR(FMath::Abs(X0.Size() - Ref0.Size()), 0, KINDA_SMALL_NUMBER);
 					EXPECT_NEAR(FMath::Abs(X1.Size() - Ref1.Size()), 0, KINDA_SMALL_NUMBER);
@@ -1429,9 +1429,9 @@ namespace GeometryCollectionTest
 					ParticleHandles[8]->Disabled() == true)
 				{
 					bParticle7SucessfulRelease = true;
-					FVector X0 = ParticleHandles[0]->X();
-					FVector X1 = ParticleHandles[1]->X();
-					FVector X6 = ParticleHandles[6]->X();
+					FVector X0 = ParticleHandles[0]->GetX();
+					FVector X1 = ParticleHandles[1]->GetX();
+					FVector X6 = ParticleHandles[6]->GetX();
 
 					// 0 is a dynamic unclustered body (was owned by cluster 8), check that it's moved since declustering
 					EXPECT_GT(FMath::Abs(X0.Size() - Ref0.Size()), KINDA_SMALL_NUMBER);
@@ -1476,9 +1476,9 @@ namespace GeometryCollectionTest
 					ParticleHandles[8]->Disabled() == true)
 				{
 					bValidFinalActiveState = true;
-					FVector X0 = ParticleHandles[0]->X();
-					FVector X1 = ParticleHandles[1]->X();
-					FVector X6 = ParticleHandles[6]->X();
+					FVector X0 = ParticleHandles[0]->GetX();
+					FVector X1 = ParticleHandles[1]->GetX();
+					FVector X6 = ParticleHandles[6]->GetX();
 
 					// 0 is a dynamic unclustered body (was owned by cluster 8), check that it's moved since declustering
 					EXPECT_GT(FMath::Abs(X0.Size() - Ref0.Size()), KINDA_SMALL_NUMBER);
@@ -1827,15 +1827,15 @@ namespace GeometryCollectionTest
 			//
 			if (Frame == 0)
 			{
-				InitialRootPosition = Root->X();
-				RelativeChildOffsets[0] = CollectionParticles[0]->X() - Root->X();
-				RelativeChildOffsets[1] = CollectionParticles[1]->X() - Root->X();
-				RelativeChildOffsets[2] = CollectionParticles2[0]->X() - Root->X();
-				RelativeChildOffsets[3] = CollectionParticles2[1]->X() - Root->X();
+				InitialRootPosition = Root->GetX();
+				RelativeChildOffsets[0] = CollectionParticles[0]->GetX() - Root->GetX();
+				RelativeChildOffsets[1] = CollectionParticles[1]->GetX() - Root->GetX();
+				RelativeChildOffsets[2] = CollectionParticles2[0]->GetX() - Root->GetX();
+				RelativeChildOffsets[3] = CollectionParticles2[1]->GetX() - Root->GetX();
 			}
 			else
 			{
-				FTransform RootTransform(Root->GetR(), Root->X());
+				FTransform RootTransform(Root->GetR(), Root->GetX());
 
 				TArray<FTransform> GlobalTransform1;
 				GeometryCollectionAlgo::Private::GlobalMatrices(*DynamicCollection, GlobalTransform1);
@@ -1843,7 +1843,7 @@ namespace GeometryCollectionTest
 				TArray<FTransform> GlobalTransform2;
 				GeometryCollectionAlgo::Private::GlobalMatrices(*DynamicCollection2, GlobalTransform2);
 
-				EXPECT_TRUE(!InitialRootPosition.Equals(Root->X())); // root moves
+				EXPECT_TRUE(!InitialRootPosition.Equals(Root->GetX())); // root moves
 
 				EXPECT_TRUE(RelativeChildOffsets[0].Equals(GlobalTransform1[0].GetRelativeTransform(RootTransform).GetTranslation()));
 				EXPECT_TRUE(RelativeChildOffsets[1].Equals(GlobalTransform1[1].GetRelativeTransform(RootTransform).GetTranslation()));
diff --git a/Engine/Source/Programs/HeadlessChaos/Private/GeometryCollection/GeometryCollectionTestDecimation.cpp b/Engine/Source/Programs/HeadlessChaos/Private/GeometryCollection/GeometryCollectionTestDecimation.cpp
--- a/Engine/Source/Programs/HeadlessChaos/Private/GeometryCollection/GeometryCollectionTestDecimation.cpp
+++ b/Engine/Source/Programs/HeadlessChaos/Private/GeometryCollection/GeometryCollectionTestDecimation.cpp
@@ -41,7 +41,7 @@ namespace GeometryCollectionTest
 		Chaos::FParticles particles;
 		particles.AddParticles(numParticles);
 		for (int i = 0; i < numParticles; i++)
-			particles.X(i).Set(Vertex[i][0], Vertex[i][1], Vertex[i][2]);
+			particles.SetX(i, FVec3(Vertex[i][0], Vertex[i][1], Vertex[i][2]));
 		return particles;
 	}
 
diff --git a/Engine/Source/Programs/HeadlessChaos/Private/GeometryCollection/GeometryCollectionTestInitilization.cpp b/Engine/Source/Programs/HeadlessChaos/Private/GeometryCollection/GeometryCollectionTestInitilization.cpp
--- a/Engine/Source/Programs/HeadlessChaos/Private/GeometryCollection/GeometryCollectionTestInitilization.cpp
+++ b/Engine/Source/Programs/HeadlessChaos/Private/GeometryCollection/GeometryCollectionTestInitilization.cpp
@@ -25,7 +25,7 @@ GTEST_TEST(AllTraits,GeometryCollection_Initilization_TransformedGeometryCollect
 
 	{ // test results
 		EXPECT_EQ(UnitTest.Solver->GetParticles().GetGeometryCollectionParticles().Size(),1);
-		FVector X = UnitTest.Solver->GetParticles().GetGeometryCollectionParticles().X(0);
+		FVector X = UnitTest.Solver->GetParticles().GetGeometryCollectionParticles().GetX(0);
 		FQuat R = UnitTest.Solver->GetParticles().GetGeometryCollectionParticles().GetR(0);
 		EXPECT_TRUE((R * GlobalRotation.Inverse()).IsIdentity(KINDA_SMALL_NUMBER));
 		EXPECT_NEAR(X.X - GlobalTranslation[0],0.0f,KINDA_SMALL_NUMBER);
@@ -56,7 +56,7 @@ GTEST_TEST(AllTraits,GeometryCollection_Initilization_TransformedGeometryCollect
 
 	{ // test results
 		EXPECT_EQ(UnitTest.Solver->GetParticles().GetGeometryCollectionParticles().Size(),1);
-		FVector X = UnitTest.Solver->GetParticles().GetGeometryCollectionParticles().X(0);
+		FVector X = UnitTest.Solver->GetParticles().GetGeometryCollectionParticles().GetX(0);
 		FQuat R = UnitTest.Solver->GetParticles().GetGeometryCollectionParticles().GetR(0);
 		EXPECT_TRUE((R * GlobalRotation.Inverse()).IsIdentity(KINDA_SMALL_NUMBER));
 		EXPECT_NEAR(X.X - GlobalTranslation[0],0.0f,KINDA_SMALL_NUMBER);
@@ -87,7 +87,7 @@ GTEST_TEST(AllTraits,GeometryCollection_Initilization_TransformedGeometryCollect
 
 	{ // test results
 		EXPECT_EQ(UnitTest.Solver->GetParticles().GetGeometryCollectionParticles().Size(),1);
-		FVector X = UnitTest.Solver->GetParticles().GetGeometryCollectionParticles().X(0);
+		FVector X = UnitTest.Solver->GetParticles().GetGeometryCollectionParticles().GetX(0);
 		FQuat R = UnitTest.Solver->GetParticles().GetGeometryCollectionParticles().GetR(0);
 		EXPECT_TRUE((R * GlobalRotation.Inverse()).IsIdentity(KINDA_SMALL_NUMBER));
 		EXPECT_NEAR(X.X - GlobalTranslation[0],0.0f,KINDA_SMALL_NUMBER);
@@ -142,7 +142,7 @@ GTEST_TEST(AllTraits,GeometryCollection_Initilization_TransformedGeometryCollect
 
 
 		EXPECT_EQ(UnitTest.Solver->GetParticles().GetGeometryCollectionParticles().Size(), 1);
-		FVector ParticlePos = UnitTest.Solver->GetParticles().GetGeometryCollectionParticles().X(0);
+		FVector ParticlePos = UnitTest.Solver->GetParticles().GetGeometryCollectionParticles().GetX(0);
 		FQuat ParticleR = UnitTest.Solver->GetParticles().GetGeometryCollectionParticles().GetR(0);
 
 		TArray<FTransform> Transform;
diff --git a/Engine/Source/Programs/HeadlessChaos/Private/GeometryCollection/GeometryCollectionTestMassProperties.cpp b/Engine/Source/Programs/HeadlessChaos/Private/GeometryCollection/GeometryCollectionTestMassProperties.cpp
--- a/Engine/Source/Programs/HeadlessChaos/Private/GeometryCollection/GeometryCollectionTestMassProperties.cpp
+++ b/Engine/Source/Programs/HeadlessChaos/Private/GeometryCollection/GeometryCollectionTestMassProperties.cpp
@@ -23,19 +23,19 @@ namespace GeometryCollectionTest
 
 		Chaos::FParticles Vertices;
 		Vertices.AddParticles(8);
-		Vertices.X(0) = FVec3(-1, 1, -1);
-		Vertices.X(1) = FVec3(1, 1, -1);
-		Vertices.X(2) = FVec3(1, -1, -1);
-		Vertices.X(3) = FVec3(-1, -1, -1);
-		Vertices.X(4) = FVec3(-1, 1, 1);
-		Vertices.X(5) = FVec3(1, 1, 1);
-		Vertices.X(6) = FVec3(1, -1, 1);
-		Vertices.X(7) = FVec3(-1, -1, 1);
+		Vertices.SetX(0, FVec3(-1, 1, -1));
+		Vertices.SetX(1, FVec3(1, 1, -1));
+		Vertices.SetX(2, FVec3(1, -1, -1));
+		Vertices.SetX(3, FVec3(-1, -1, -1));
+		Vertices.SetX(4, FVec3(-1, 1, 1));
+		Vertices.SetX(5, FVec3(1, 1, 1));
+		Vertices.SetX(6, FVec3(1, -1, 1));
+		Vertices.SetX(7, FVec3(-1, -1, 1));
 
 		// @todo(chaos):  breaking : this trips an ensure in the test, why?
 		for (int i = 0; i < 8; i++) {
-			Vertices.X(i) *= FVector(1, 2, 3);
-			Vertices.X(i) += FVector(1, 2, 3);
+			Vertices.SetX(i, Vertices.GetX(i) * FVector(1, 2, 3));
+			Vertices.SetX(i, Vertices.GetX(i) + FVector(1, 2, 3));
 		}
 
 		TArray<Chaos::TVec3<int32>> Faces;
@@ -63,7 +63,7 @@ namespace GeometryCollectionTest
 
 			for (int32 Idx = 0; Idx < 8; ++Idx)
 			{
-				Vertices.X(Idx) -= MassProperties.CenterOfMass;
+				Vertices.SetX(Idx, Vertices.GetX(Idx) - MassProperties.CenterOfMass);
 			}
 
 			check(MassProperties.Mass > 0);
@@ -134,7 +134,7 @@ namespace GeometryCollectionTest
 		MassSpaceParticles.AddParticles(Vertex.Num());
 		for (int32 Idx = 0; Idx < Vertex.Num(); ++Idx)
 		{
-			MassSpaceParticles.X(Idx) = Transform[BoneMap[Idx]].TransformPosition(FVector(Vertex[Idx]));
+			MassSpaceParticles.SetX(Idx, Transform[BoneMap[Idx]].TransformPosition(FVector(Vertex[Idx])));
 		}
 
 		CalculateVolumeAndCenterOfMass(MassSpaceParticles, TriMesh->GetElements(), MassProperties.Volume, MassProperties.CenterOfMass);
@@ -146,7 +146,7 @@ namespace GeometryCollectionTest
 
 		for (int32 Idx = 0; Idx < Vertex.Num(); ++Idx)
 		{
-			MassSpaceParticles.X(Idx) -= MassProperties.CenterOfMass;
+			MassSpaceParticles.SetX(Idx, MassSpaceParticles.GetX(Idx) - MassProperties.CenterOfMass);
 		}
 
 		FReal Density = 1.0;
@@ -207,7 +207,7 @@ namespace GeometryCollectionTest
 		MassSpaceParticles.AddParticles(Vertex.Num());
 		for (int32 Idx = 0; Idx < Vertex.Num(); ++Idx)
 		{
-			MassSpaceParticles.X(Idx) = Transform[BoneMap[Idx]].TransformPosition(FVector(Vertex[Idx]));
+			MassSpaceParticles.SetX(Idx, Transform[BoneMap[Idx]].TransformPosition(FVector(Vertex[Idx])));
 		}
 
 		CalculateVolumeAndCenterOfMass(MassSpaceParticles, TriMesh->GetElements(), MassProperties.Volume, MassProperties.CenterOfMass);
@@ -226,7 +226,7 @@ namespace GeometryCollectionTest
 
 		for (int32 Idx = 0; Idx < Vertex.Num(); ++Idx)
 		{
-			MassSpaceParticles.X(Idx) -= MassProperties.CenterOfMass;
+			MassSpaceParticles.SetX(Idx, MassSpaceParticles.GetX(Idx) - MassProperties.CenterOfMass);
 		}
 
 		FReal Density = 0.01;
@@ -289,7 +289,7 @@ namespace GeometryCollectionTest
 		MassSpaceParticles.AddParticles(Vertex.Num());
 		for (int32 Idx = 0; Idx < Vertex.Num(); ++Idx)
 		{
-			MassSpaceParticles.X(Idx) = Transform[BoneMap[Idx]].TransformPosition(FVector(Vertex[Idx]));
+			MassSpaceParticles.SetX(Idx, Transform[BoneMap[Idx]].TransformPosition(FVector(Vertex[Idx])));
 		}
 
 		CalculateVolumeAndCenterOfMass(MassSpaceParticles, TriMesh->GetElements(), MassProperties.Volume, MassProperties.CenterOfMass);
@@ -301,7 +301,7 @@ namespace GeometryCollectionTest
 
 		for (int32 Idx = 0; Idx < Vertex.Num(); ++Idx)
 		{
-			MassSpaceParticles.X(Idx) -= MassProperties.CenterOfMass;
+			MassSpaceParticles.SetX(Idx, MassSpaceParticles.GetX(Idx) - MassProperties.CenterOfMass);
 		}
 
 		FReal Density = 0.01;
@@ -371,8 +371,8 @@ namespace GeometryCollectionTest
 		for (int32 Idx = 0; Idx < Vertex.Num(); ++Idx)
 		{
 			FVector VertexPoint = FVector(Vertex[Idx]);
-			MassSpaceParticles.X(Idx) = Transform[BoneMap[Idx]].TransformPosition(FVector(Vertex[Idx]));
-			FVector MassSpacePoint = FVector(MassSpaceParticles.X(Idx)[0], MassSpaceParticles.X(Idx)[1], MassSpaceParticles.X(Idx)[2]);
+			MassSpaceParticles.SetX(Idx, Transform[BoneMap[Idx]].TransformPosition(FVector(Vertex[Idx])));
+			FVector MassSpacePoint = FVector(MassSpaceParticles.GetX(Idx)[0], MassSpaceParticles.GetX(Idx)[1], MassSpaceParticles.GetX(Idx)[2]);
 			SomeVec.Add(MassSpacePoint);
 		}
 
@@ -385,7 +385,7 @@ namespace GeometryCollectionTest
 
 		for (int32 Idx = 0; Idx < Vertex.Num(); ++Idx)
 		{
-			MassSpaceParticles.X(Idx) -= MassProperties.CenterOfMass;
+			MassSpaceParticles.SetX(Idx, MassSpaceParticles.GetX(Idx) - MassProperties.CenterOfMass);
 		}
 
 		FReal Density = 0.01;
diff --git a/Engine/Source/Programs/HeadlessChaos/Private/GeometryCollection/GeometryCollectionTestSimulation.cpp b/Engine/Source/Programs/HeadlessChaos/Private/GeometryCollection/GeometryCollectionTestSimulation.cpp
--- a/Engine/Source/Programs/HeadlessChaos/Private/GeometryCollection/GeometryCollectionTestSimulation.cpp
+++ b/Engine/Source/Programs/HeadlessChaos/Private/GeometryCollection/GeometryCollectionTestSimulation.cpp
@@ -378,11 +378,11 @@ namespace GeometryCollectionTest
 		FReal MaxZ = -TNumericLimits<FReal>::Max();
 		for (uint32 Idx = 0; Idx < Simplicial.Size(); ++Idx)
 		{
-			const FReal phi = Implicit.SignedDistance(Simplicial.X(Idx));
+			const FReal phi = Implicit.SignedDistance(Simplicial.GetX(Idx));
 			EXPECT_LT(FMath::Abs(phi), DxSize);
 			//EXPECT_LT(FMath::Abs(phi), 0.01f);
 
-			const auto& Pos = Simplicial.X(Idx);
+			const auto& Pos = Simplicial.GetX(Idx);
 			MinX = MinX < Pos[0] ? MinX : Pos[0];
 			MinY = MinY < Pos[1] ? MinY : Pos[1];
 			MinZ = MinZ < Pos[2] ? MinZ : Pos[2];
diff --git a/Engine/Source/Programs/HeadlessChaos/Private/GeometryCollection/GeometryCollectionTestSimulationField.cpp b/Engine/Source/Programs/HeadlessChaos/Private/GeometryCollection/GeometryCollectionTestSimulationField.cpp
--- a/Engine/Source/Programs/HeadlessChaos/Private/GeometryCollection/GeometryCollectionTestSimulationField.cpp
+++ b/Engine/Source/Programs/HeadlessChaos/Private/GeometryCollection/GeometryCollectionTestSimulationField.cpp
@@ -1133,9 +1133,9 @@ namespace GeometryCollectionTest
 			Chaos::TPBDGeometryCollectionParticles<Chaos::FReal, 3>& Particles = UnitTest.Solver->GetParticles().GetGeometryCollectionParticles();
 			EXPECT_NE(FMath::Abs(Collection->DynamicCollection->GetTransform(0).GetRotation().Euler().Y), SMALL_THRESHOLD); // not rotating in Y?
 			EXPECT_GT(Particles.GetW(0).X, PreviousX); // rotating in X?
-			EXPECT_LT(Particles.X(0).Z, PreviousHeight); // still falling?
+			EXPECT_LT(Particles.GetX(0).Z, PreviousHeight); // still falling?
 
-			PreviousHeight = Particles.X(0).Z;
+			PreviousHeight = Particles.GetX(0).Z;
 			PreviousX = Particles.GetW(0).X;
 		}
 
diff --git a/Engine/Source/Programs/HeadlessChaos/Private/GeometryCollection/GeometryCollectionTestSimulationStreaming.cpp b/Engine/Source/Programs/HeadlessChaos/Private/GeometryCollection/GeometryCollectionTestSimulationStreaming.cpp
--- a/Engine/Source/Programs/HeadlessChaos/Private/GeometryCollection/GeometryCollectionTestSimulationStreaming.cpp
+++ b/Engine/Source/Programs/HeadlessChaos/Private/GeometryCollection/GeometryCollectionTestSimulationStreaming.cpp
@@ -65,7 +65,7 @@ namespace GeometryCollectionTest
 			EXPECT_EQ(GCW2->PhysObject->GetUnorderedParticles_Internal().Num(), 1);
 			EXPECT_EQ(GCW->PhysObject->GetNumTransforms(), 1);
 			EXPECT_EQ(GCW2->PhysObject->GetNumTransforms(), 1);
-			EXPECT_LT(GCW->PhysObject->GetParticle_Internal(0)->X().Z, GCW2->PhysObject->GetParticle_Internal(0)->X().Z);
+			EXPECT_LT(GCW->PhysObject->GetParticle_Internal(0)->GetX().Z, GCW2->PhysObject->GetParticle_Internal(0)->GetX().Z);
 		}
 	}
 
@@ -124,7 +124,7 @@ namespace GeometryCollectionTest
 			EXPECT_EQ(GCW2->PhysObject->GetUnorderedParticles_Internal().Num(), 1);
 			EXPECT_EQ(GCW->PhysObject->GetNumTransforms(), 1);
 			EXPECT_EQ(GCW2->PhysObject->GetNumTransforms(), 1);
-			EXPECT_LT(FMath::Abs(GCW->PhysObject->GetParticle_Internal(0)->X().Z - GCW2->PhysObject->GetParticle_Internal(0)->X().Z), KINDA_SMALL_NUMBER);
+			EXPECT_LT(FMath::Abs(GCW->PhysObject->GetParticle_Internal(0)->GetX().Z - GCW2->PhysObject->GetParticle_Internal(0)->GetX().Z), KINDA_SMALL_NUMBER);
 		}
 	}
 
@@ -212,7 +212,7 @@ namespace GeometryCollectionTest
 		FGeometryCollectionWrapper* GCW = UnitTest.PhysicsObjects[UnitTest.PhysicsObjects.Num() - 1]->As<FGeometryCollectionWrapper>();
 		EXPECT_EQ(GCW->PhysObject->GetParticle_Internal(0)->Disabled(), false);
 
-		EXPECT_LT(GCW->PhysObject->GetParticle_Internal(0)->X().Z, -1.f);
+		EXPECT_LT(GCW->PhysObject->GetParticle_Internal(0)->GetX().Z, -1.f);
 	}
 
 
diff --git a/Engine/Source/Programs/HeadlessChaos/Private/HeadlessChaosTestBroadphase.cpp b/Engine/Source/Programs/HeadlessChaos/Private/HeadlessChaosTestBroadphase.cpp
--- a/Engine/Source/Programs/HeadlessChaos/Private/HeadlessChaosTestBroadphase.cpp
+++ b/Engine/Source/Programs/HeadlessChaos/Private/HeadlessChaosTestBroadphase.cpp
@@ -58,7 +58,7 @@ namespace ChaosTest
 		template <SQType>
 		bool Visit(int32 Idx, FQueryFastData& CurData)
 		{
-			const FRigidTransform3 BoxTM(Boxes.X(Idx), Boxes.GetR(Idx));
+			const FRigidTransform3 BoxTM(Boxes.GetX(Idx), Boxes.GetR(Idx));
 			FAABB3 Box = static_cast<const TBox<FReal, 3>*>(Boxes.GetGeometry(Idx).GetReference())->BoundingBox().TransformedAABB(BoxTM);
 			FAABB3 ThicknedBox(Box.Min() - HalfExtents, Box.Max() + HalfExtents);
 
@@ -133,7 +133,7 @@ namespace ChaosTest
 		bool VisitOverlap(TSpatialVisitorData<int32> Instance)
 		{
 			const int32 Idx = Instance.Payload;
-			const FRigidTransform3 BoxTM(Boxes.X(Idx), Boxes.GetR(Idx));
+			const FRigidTransform3 BoxTM(Boxes.GetX(Idx), Boxes.GetR(Idx));
 			FAABB3 Box = static_cast<const TBox<FReal, 3>*>(Boxes.GetGeometry(Idx).GetReference())->BoundingBox().TransformedAABB(BoxTM);
 			
 			if (Box.Intersects(Bounds))
@@ -240,7 +240,7 @@ namespace ChaosTest
 			{
 				for (int32 Col = 0; Col < NumCols; ++Col)
 				{
-					Boxes->X(Idx) = FVec3(Col * 100, Row * 100, Height * 100) + Offset;
+					Boxes->SetX(Idx, FVec3(Col * 100, Row * 100, Height * 100) + Offset);
 					Boxes->SetR(Idx, FRotation3::Identity);
 					Boxes->SetGeometry(Idx, Box);
 					++Idx;
@@ -286,8 +286,8 @@ namespace ChaosTest
 			//move instance away
 			{
 				const int32 MoveIdx = Visitor2.Instances[0];
-				Boxes->X(MoveIdx) += FVec3(1000, 0, 0);
-				FAABB3 NewBounds = Boxes->GetGeometry(MoveIdx)->template GetObject<TBox<FReal, 3>>()->BoundingBox().TransformedAABB(FRigidTransform3(Boxes->X(MoveIdx), Boxes->GetR(MoveIdx)));
+				Boxes->SetX(MoveIdx, Boxes->GetX(MoveIdx) + FVec3(1000, 0, 0));
+				FAABB3 NewBounds = Boxes->GetGeometry(MoveIdx)->template GetObject<TBox<FReal, 3>>()->BoundingBox().TransformedAABB(FRigidTransform3(Boxes->GetX(MoveIdx), Boxes->GetR(MoveIdx)));
 				Spatial2->UpdateElementIn(MoveIdx, NewBounds, true, SpatialIdx);
 
 				FVisitor Visitor3(FVec3(10, 0, 0), FVec3(0, 1, 0), 0, *Boxes);
@@ -295,8 +295,8 @@ namespace ChaosTest
 				EXPECT_EQ(Visitor3.Instances.Num(), 8);
 
 				//move instance back
-				Boxes->X(MoveIdx) -= FVec3(1000, 0, 0);
-				NewBounds = Boxes->GetGeometry(MoveIdx)->template GetObject<TBox<FReal, 3>>()->BoundingBox().TransformedAABB(FRigidTransform3(Boxes->X(MoveIdx), Boxes->GetR(MoveIdx)));
+				Boxes->SetX(MoveIdx, Boxes->GetX(MoveIdx) - FVec3(1000, 0, 0));
+				NewBounds = Boxes->GetGeometry(MoveIdx)->template GetObject<TBox<FReal, 3>>()->BoundingBox().TransformedAABB(FRigidTransform3(Boxes->GetX(MoveIdx), Boxes->GetR(MoveIdx)));
 				Spatial2->UpdateElementIn(MoveIdx, NewBounds, true, SpatialIdx);
 
 				FVisitor Visitor4(FVec3(10, 0, 0), FVec3(0, 1, 0), 0, *Boxes);
@@ -307,9 +307,9 @@ namespace ChaosTest
 			//move other instance into view
 			{
 				const int32 MoveIdx = 5 * 5 * 5;
-				const FVec3 OldPos = Boxes->X(MoveIdx);
-				Boxes->X(MoveIdx) = FVec3(0, 0, 0);
-				FAABB3 NewBounds = Boxes->GetGeometry(MoveIdx)->template GetObject<TBox<FReal, 3>>()->BoundingBox().TransformedAABB(FRigidTransform3(Boxes->X(MoveIdx), Boxes->GetR(MoveIdx)));
+				const FVec3 OldPos = Boxes->GetX(MoveIdx);
+				Boxes->SetX(MoveIdx, FVec3(0, 0, 0));
+				FAABB3 NewBounds = Boxes->GetGeometry(MoveIdx)->template GetObject<TBox<FReal, 3>>()->BoundingBox().TransformedAABB(FRigidTransform3(Boxes->GetX(MoveIdx), Boxes->GetR(MoveIdx)));
 				Spatial2->UpdateElementIn(MoveIdx, NewBounds, true, SpatialIdx);
 
 				FVisitor Visitor3(FVec3(10, 0, 0), FVec3(0, 1, 0), 0, *Boxes);
@@ -317,17 +317,17 @@ namespace ChaosTest
 				EXPECT_EQ(Visitor3.Instances.Num(), 10);
 
 				//move instance back
-				Boxes->X(MoveIdx) = OldPos;
-				NewBounds = Boxes->GetGeometry(MoveIdx)->template GetObject<TBox<FReal, 3>>()->BoundingBox().TransformedAABB(FRigidTransform3(Boxes->X(MoveIdx), Boxes->GetR(MoveIdx)));
+				Boxes->SetX(MoveIdx, OldPos);
+				NewBounds = Boxes->GetGeometry(MoveIdx)->template GetObject<TBox<FReal, 3>>()->BoundingBox().TransformedAABB(FRigidTransform3(Boxes->GetX(MoveIdx), Boxes->GetR(MoveIdx)));
 				Spatial2->UpdateElementIn(MoveIdx, NewBounds, true, SpatialIdx);
 			}
 
 			//move instance outside of grid bounds
 			{
 				const int32 MoveIdx = 5 * 5 * 5;
-				const FVec3 OldPos = Boxes->X(MoveIdx);
-				Boxes->X(MoveIdx) = FVec3(-50, 0, 0);
-				FAABB3 NewBounds = Boxes->GetGeometry(MoveIdx)->template GetObject<TBox<FReal, 3>>()->BoundingBox().TransformedAABB(FRigidTransform3(Boxes->X(MoveIdx), Boxes->GetR(MoveIdx)));
+				const FVec3 OldPos = Boxes->GetX(MoveIdx);
+				Boxes->SetX(MoveIdx, FVec3(-50, 0, 0));
+				FAABB3 NewBounds = Boxes->GetGeometry(MoveIdx)->template GetObject<TBox<FReal, 3>>()->BoundingBox().TransformedAABB(FRigidTransform3(Boxes->GetX(MoveIdx), Boxes->GetR(MoveIdx)));
 				Spatial2->UpdateElementIn(MoveIdx, NewBounds, true, SpatialIdx);
 
 				FVisitor Visitor3(FVec3(10, 0, 0), FVec3(0, 1, 0), 0, *Boxes);
@@ -346,15 +346,15 @@ namespace ChaosTest
 				EXPECT_EQ(Visitor5.Instances.Num(), 0);
 
 				//move instance back
-				Boxes->X(MoveIdx) = OldPos;
+				Boxes->SetX(MoveIdx, OldPos);
 
 				//create a new box
 				const int32 NewIdx = Boxes->Size();
 				Boxes->AddParticles(1);
-				Boxes->X(NewIdx) = FVec3(-20, 0, 0);
+				Boxes->SetX(NewIdx, FVec3(-20, 0, 0));
 				Boxes->SetR(NewIdx, FRotation3::Identity);
 				Boxes->SetGeometry(NewIdx, Box);
-				NewBounds = Boxes->GetGeometry(NewIdx)->template GetObject<TBox<FReal, 3>>()->BoundingBox().TransformedAABB(FRigidTransform3(Boxes->X(NewIdx), Boxes->GetR(NewIdx)));
+				NewBounds = Boxes->GetGeometry(NewIdx)->template GetObject<TBox<FReal, 3>>()->BoundingBox().TransformedAABB(FRigidTransform3(Boxes->GetX(NewIdx), Boxes->GetR(NewIdx)));
 				Spatial2->UpdateElementIn(NewIdx, NewBounds, true, SpatialIdx);
 				FVisitor Visitor6(FVec3(-20, 0, 0), FVec3(0, 1, 0), 0, *Boxes);
 				Spatial2->Raycast(Visitor6.Start, Visitor6.Dir, 1000, Visitor6);
@@ -495,7 +495,7 @@ namespace ChaosTest
 			{
 				for (int32 Col = 0; Col < NumCols; ++Col)
 				{
-					Boxes.X(Idx) = FVec3(Col * 100, Row * 100, Height * 100);
+					Boxes.SetX(Idx, FVec3(Col * 100, Row * 100, Height * 100));
 					Boxes.SetR(Idx, FRotation3::Identity);
 					Boxes.SetGeometry(Idx, Box);
 					++Idx;
@@ -1057,7 +1057,7 @@ namespace ChaosTest
 
 		// Construct a particle and set HasBounds to false.
 		int32 Idx = 0;
-		Boxes->X(Idx) = FVec3(0);
+		Boxes->SetX(Idx, FVec3(0));
 		Boxes->SetR(Idx, FRotation3::Identity);
 		Boxes->SetGeometry(Idx, Box);
 
diff --git a/Engine/Source/Programs/HeadlessChaos/Private/HeadlessChaosTestCCD.cpp b/Engine/Source/Programs/HeadlessChaos/Private/HeadlessChaosTestCCD.cpp
--- a/Engine/Source/Programs/HeadlessChaos/Private/HeadlessChaosTestCCD.cpp
+++ b/Engine/Source/Programs/HeadlessChaos/Private/HeadlessChaosTestCCD.cpp
@@ -54,12 +54,12 @@ namespace ChaosTest
 		Dynamic->InvI() = TVec3<FRealSingle>(1.0f / Mass, 1.0f / Mass, 1.0f / Mass);
 
 		// Positions and velocities
-		Static->X() = FVec3(0, 0, 0);
-		Dynamic->X() = FVec3(0, 0, InitialPosition); // Start 30cm above the static box
+		Static->SetX(FVec3(0, 0, 0));
+		Dynamic->SetX( FVec3(0, 0, InitialPosition)); // Start 30cm above the static box
 		Dynamic->SetV(FVec3(0, 0, -InitialSpeed));
 
 		// The position of the static has changed and statics don't automatically update bounds, so update explicitly
-		Static->UpdateWorldSpaceState(TRigidTransform<FReal, 3>(Static->X(), Static->GetR()), FVec3(0));
+		Static->UpdateWorldSpaceState(TRigidTransform<FReal, 3>(Static->GetX(), Static->GetR()), FVec3(0));
 
 		::ChaosTest::SetParticleSimDataToCollide({ Static,Dynamic });
 
@@ -81,7 +81,7 @@ namespace ChaosTest
 		// Large error margin, we are testing CCD and not solver accuracy
 		// The Box should pass right through the other one without interacting
 		const FReal LargeErrorMargin = 10.0f;
-		EXPECT_NEAR(Dynamic->X()[2], InitialPosition - InitialSpeed * Dt, LargeErrorMargin);
+		EXPECT_NEAR(Dynamic->GetX()[2], InitialPosition - InitialSpeed * Dt, LargeErrorMargin);
 	}
 
 	GTEST_TEST(CCDTests, ConvexConvex)
@@ -118,12 +118,12 @@ namespace ChaosTest
 		Dynamic->InvI() = TVec3<FRealSingle>(1.0f / Mass, 1.0f / Mass, 1.0f / Mass);
 
 		// Positions and velocities
-		Static->X() = FVec3(0, 0, 0);
-		Dynamic->X() = FVec3(0, 0, BoxHalfSize * 2 + 30); // Start 30cm above the static box
+		Static->SetX(FVec3(0, 0, 0));
+		Dynamic->SetX(FVec3(0, 0, BoxHalfSize * 2 + 30)); // Start 30cm above the static box
 		Dynamic->SetV(FVec3(0, 0, -InitialSpeed));
 
 		// The position of the static has changed and statics don't automatically update bounds, so update explicitly
-		Static->UpdateWorldSpaceState(TRigidTransform<FReal, 3>(Static->X(), Static->GetR()), FVec3(0));
+		Static->UpdateWorldSpaceState(TRigidTransform<FReal, 3>(Static->GetX(), Static->GetR()), FVec3(0));
 
 		::ChaosTest::SetParticleSimDataToCollide({ Static,Dynamic });
 
@@ -144,7 +144,7 @@ namespace ChaosTest
 
 		// Large error margin, we are testing CCD and not solver accuracy
 		const FReal LargeErrorMargin = 10.0f;
-		EXPECT_GE(Dynamic->X()[2], BoxHalfSize * 2 - LargeErrorMargin);
+		EXPECT_GE(Dynamic->GetX()[2], BoxHalfSize * 2 - LargeErrorMargin);
 	}
 
 	// CCD not implemented for sphere sphere
@@ -184,13 +184,13 @@ namespace ChaosTest
 		Dynamic->InvI() = TVec3<FRealSingle>(1.0f / Mass);
 
 		// Positions and velocities
-		Static->X() = FVec3(0, 0, 0);
+		Static->SetX(FVec3(0, 0, 0));
 
-		Dynamic->X() = FVec3(0, 0, SphereRadius * 2 + 10);
+		Dynamic->SetX(FVec3(0, 0, SphereRadius * 2 + 10));
 		Dynamic->SetV(FVec3(0, 0, -InitialSpeed));
 		
 		// The position of the static has changed and statics don't automatically update bounds, so update explicitly
-		Static->UpdateWorldSpaceState(TRigidTransform<FReal, 3>(Static->X(), Static->GetR()), FVec3(0));
+		Static->UpdateWorldSpaceState(TRigidTransform<FReal, 3>(Static->GetX(), Static->GetR()), FVec3(0));
 
 		::ChaosTest::SetParticleSimDataToCollide({ Static,Dynamic });
 
@@ -211,7 +211,7 @@ namespace ChaosTest
 
 		// Large error margin, we are testing CCD and not solver accuracy
 		const FReal LargeErrorMargin = 10.0f;
-		EXPECT_GE(Dynamic->X()[2], SphereRadius * 2 - LargeErrorMargin);
+		EXPECT_GE(Dynamic->GetX()[2], SphereRadius * 2 - LargeErrorMargin);
 	}
 
 	
@@ -269,22 +269,22 @@ namespace ChaosTest
 		Dynamic->InvI() = TVec3<FRealSingle>(1.0f / Mass);
 
 		// Positions and velocities
-		ContainerFaces[0]->X() = FVec3(ContainerBoxHalfSize, 0, 0);
-		ContainerFaces[1]->X() = FVec3(-ContainerBoxHalfSize, 0, 0);
-		ContainerFaces[2]->X() = FVec3(0, ContainerBoxHalfSize, 0);
-		ContainerFaces[3]->X() = FVec3(0, -ContainerBoxHalfSize, 0);
-		ContainerFaces[4]->X() = FVec3(0, 0, ContainerBoxHalfSize);
-		ContainerFaces[5]->X() = FVec3(0, 0, -ContainerBoxHalfSize);
+		ContainerFaces[0]->SetX(FVec3(ContainerBoxHalfSize, 0, 0));
+		ContainerFaces[1]->SetX(FVec3(-ContainerBoxHalfSize, 0, 0));
+		ContainerFaces[2]->SetX(FVec3(0, ContainerBoxHalfSize, 0));
+		ContainerFaces[3]->SetX(FVec3(0, -ContainerBoxHalfSize, 0));
+		ContainerFaces[4]->SetX(FVec3(0, 0, ContainerBoxHalfSize));
+		ContainerFaces[5]->SetX(FVec3(0, 0, -ContainerBoxHalfSize));
 
-		Dynamic->X() = FVec3(0, 0, 0);
+		Dynamic->SetX(FVec3(0, 0, 0));
 
 		// The position of the static has changed and statics don't automatically update bounds, so update explicitly
-		ContainerFaces[0]->UpdateWorldSpaceState(TRigidTransform<FReal, 3>(ContainerFaces[0]->X(), ContainerFaces[0]->GetR()), FVec3(0));
-		ContainerFaces[1]->UpdateWorldSpaceState(TRigidTransform<FReal, 3>(ContainerFaces[1]->X(), ContainerFaces[1]->GetR()), FVec3(0));
-		ContainerFaces[2]->UpdateWorldSpaceState(TRigidTransform<FReal, 3>(ContainerFaces[2]->X(), ContainerFaces[2]->GetR()), FVec3(0));
-		ContainerFaces[3]->UpdateWorldSpaceState(TRigidTransform<FReal, 3>(ContainerFaces[3]->X(), ContainerFaces[3]->GetR()), FVec3(0));
-		ContainerFaces[4]->UpdateWorldSpaceState(TRigidTransform<FReal, 3>(ContainerFaces[4]->X(), ContainerFaces[4]->GetR()), FVec3(0));
-		ContainerFaces[5]->UpdateWorldSpaceState(TRigidTransform<FReal, 3>(ContainerFaces[5]->X(), ContainerFaces[5]->GetR()), FVec3(0));
+		ContainerFaces[0]->UpdateWorldSpaceState(TRigidTransform<FReal, 3>(ContainerFaces[0]->GetX(), ContainerFaces[0]->GetR()), FVec3(0));
+		ContainerFaces[1]->UpdateWorldSpaceState(TRigidTransform<FReal, 3>(ContainerFaces[1]->GetX(), ContainerFaces[1]->GetR()), FVec3(0));
+		ContainerFaces[2]->UpdateWorldSpaceState(TRigidTransform<FReal, 3>(ContainerFaces[2]->GetX(), ContainerFaces[2]->GetR()), FVec3(0));
+		ContainerFaces[3]->UpdateWorldSpaceState(TRigidTransform<FReal, 3>(ContainerFaces[3]->GetX(), ContainerFaces[3]->GetR()), FVec3(0));
+		ContainerFaces[4]->UpdateWorldSpaceState(TRigidTransform<FReal, 3>(ContainerFaces[4]->GetX(), ContainerFaces[4]->GetR()), FVec3(0));
+		ContainerFaces[5]->UpdateWorldSpaceState(TRigidTransform<FReal, 3>(ContainerFaces[5]->GetX(), ContainerFaces[5]->GetR()), FVec3(0));
 
 		::ChaosTest::SetParticleSimDataToCollide({ Dynamic });
 		::ChaosTest::SetParticleSimDataToCollide({ ContainerFaces });
@@ -317,14 +317,14 @@ namespace ChaosTest
 		{
 			// If this failed, the dynamic cube escaped the air tight static container
 			const  FReal MaxCoordinates = ContainerBoxHalfSize - ContainerWallThickness / 2 - SmallBoxHalfSize;
-			EXPECT_LT(FMath::Abs(Dynamic->X()[axis]), MaxCoordinates + LargeErrorMargin);
+			EXPECT_LT(FMath::Abs(Dynamic->GetX()[axis]), MaxCoordinates + LargeErrorMargin);
 		}
 		/////////////////////////////////////////////
 		// Test2: Now launch to cube to a corner
 		Dynamic->SetVf(FVec3f(-ContainerBoxHalfSize * Fps * 10));
 		Dynamic->SetWf(FVec3f(0));
-		Dynamic->X() = FVec3(0);
-		Dynamic->P() = FVec3(0);
+		Dynamic->SetX(FVec3(0));
+		Dynamic->SetP(FVec3(0));
 		Dynamic->SetRf(TRotation<FRealSingle, 3>::FromIdentity());
 		Dynamic->SetQf(TRotation<FRealSingle, 3>::FromIdentity());
 
@@ -339,7 +339,7 @@ namespace ChaosTest
 		{
 			// If this failed, the dynamic cube escaped the air tight static container
 			const  FReal MaxCoordinates = ContainerBoxHalfSize - ContainerWallThickness / 2 - SmallBoxHalfSize;
-			EXPECT_LT(FMath::Abs(Dynamic->X()[axis]), MaxCoordinates + LargeErrorMargin);
+			EXPECT_LT(FMath::Abs(Dynamic->GetX()[axis]), MaxCoordinates + LargeErrorMargin);
 		}
 
 		/////////////////////////////////////////////////////////////////////
@@ -347,8 +347,8 @@ namespace ChaosTest
 		// Make sure it still stays inside the box (albeit with a very reduced velocity) 
 		Dynamic->SetV(InitialVelocity);
 		Dynamic->SetW(FVec3(0));
-		Dynamic->X() = FVec3(0);
-		Dynamic->P() = FVec3(0);
+		Dynamic->SetX(FVec3(0));
+		Dynamic->SetP(FVec3(0));
 		Dynamic->SetR(TRotation<FReal, 3>::FromIdentity());
 		Dynamic->SetQ(TRotation<FReal, 3>::FromIdentity());
 		PhysicsMaterial->Restitution = 0.9f;
@@ -364,7 +364,7 @@ namespace ChaosTest
 		{
 			// If this failed, the dynamic cube escaped the air tight static container
 			const  FReal MaxCoordinates = ContainerBoxHalfSize - ContainerWallThickness / 2 - SmallBoxHalfSize;
-			EXPECT_LT(FMath::Abs(Dynamic->X()[axis]), MaxCoordinates + LargeErrorMargin);
+			EXPECT_LT(FMath::Abs(Dynamic->GetX()[axis]), MaxCoordinates + LargeErrorMargin);
 		}
 	}
 }
diff --git a/Engine/Source/Programs/HeadlessChaos/Private/HeadlessChaosTestCloth.cpp b/Engine/Source/Programs/HeadlessChaos/Private/HeadlessChaosTestCloth.cpp
--- a/Engine/Source/Programs/HeadlessChaos/Private/HeadlessChaosTestCloth.cpp
+++ b/Engine/Source/Programs/HeadlessChaos/Private/HeadlessChaosTestCloth.cpp
@@ -53,7 +53,7 @@ namespace ChaosTest {
 		auto& Particles = Evolution->Particles();
 		const uint32 Idx = Particles.Size();
 		Particles.AddParticles(1);
-		Particles.X(Idx) = Position;
+		Particles.SetX(Idx, Position);
 		Particles.V(Idx) = Velocity;
 		Particles.M(Idx) = Mass;
 		Particles.InvM(Idx) = 1.0 / Mass;
@@ -73,15 +73,15 @@ namespace ChaosTest {
 		// Initialize particles.  Use 1/3 area of connected triangles for particle mass.
 		for (uint32 i = InitialNumParticles; i < Particles.Size(); i++)
 		{
-			Particles.X(i) += XOffset;
+			Particles.SetX(i, Particles.GetX(i) + XOffset);
 			Particles.V(i) = Chaos::Softs::FSolverVec3(0);
 			Particles.M(i) = 0;
 		}
 		for (const Chaos::TVec3<int32>& Tri : TriMesh.GetElements())
 		{
 			const Softs::FSolverReal TriArea = 0.5 * Chaos::Softs::FSolverVec3::CrossProduct(
-				Particles.X(Tri[1]) - Particles.X(Tri[0]),
-				Particles.X(Tri[2]) - Particles.X(Tri[0])).Size();
+				Particles.GetX(Tri[1]) - Particles.GetX(Tri[0]),
+				Particles.GetX(Tri[2]) - Particles.GetX(Tri[0])).Size();
 			for (int32 i = 0; i < 3; i++)
 			{
 				Particles.M(Tri[i]) += TriArea / 3;
@@ -153,7 +153,7 @@ namespace ChaosTest {
 		Points.SetNum(Particles.Size());
 
 		for (uint32 i = 0; i < Particles.Size(); i++)
-			Points[i] = Particles.X(i);
+			Points[i] = Particles.GetX(i);
 
 		return Points;
 	}
@@ -163,7 +163,7 @@ namespace ChaosTest {
 	{
 		for (uint32 i = 0; i < Particles.Size(); i++)
 		{
-			Particles.X(i) = Points[i];
+			Particles.SetX(i, Points[i]);
 			Particles.V(i) = Chaos::Softs::FSolverVec3(0);
 		}
 	}
diff --git a/Engine/Source/Programs/HeadlessChaos/Private/HeadlessChaosTestClustering.cpp b/Engine/Source/Programs/HeadlessChaos/Private/HeadlessChaosTestClustering.cpp
--- a/Engine/Source/Programs/HeadlessChaos/Private/HeadlessChaosTestClustering.cpp
+++ b/Engine/Source/Programs/HeadlessChaos/Private/HeadlessChaosTestClustering.cpp
@@ -35,8 +35,8 @@ namespace ChaosTest {
 		FPBDRigidParticleHandle* Box2 = AppendClusteredParticleBox(Particles, FVec3((FReal)100, (FReal)100, (FReal)100));
 		uint32 Box2Id = FirstId++;
 
-		Box2->X() = FVec3((FReal)100, (FReal)0, (FReal)0);
-		Box2->P() = Box2->X();
+		Box2->SetX(FVec3((FReal)100, (FReal)0, (FReal)0));
+		Box2->SetP(Box2->GetX());
 
 		Evolution.EnableParticle(Box1);
 		Evolution.EnableParticle(Box2);
@@ -51,12 +51,12 @@ namespace ChaosTest {
 		Evolution.GetRigidClustering().CreateClusterParticle(0, MoveTemp(ClusterChildren), ClusterParams, FImplicitObjectPtr(nullptr));
 		EXPECT_EQ(ClusteredParticles.Size(), 3);
 
-		FVec3 ClusterX = ClusteredParticles.X(2);
+		FVec3 ClusterX = ClusteredParticles.GetX(2);
 		FRotation3 ClusterRot = ClusteredParticles.GetR(2);
 
 		EXPECT_TRUE(ClusterX.Equals(FVec3 {(FReal)50, 0, 0}));
 		EXPECT_TRUE(ClusterRot.Equals(FRotation3::Identity));
-		EXPECT_TRUE(ClusterX.Equals(ClusteredParticles.P(2)));
+		EXPECT_TRUE(ClusterX.Equals(ClusteredParticles.GetP(2)));
 		EXPECT_TRUE(ClusterRot.Equals(ClusteredParticles.GetQ(2)));
 
 		FRigidTransform3 ClusterTM(ClusterX, ClusterRot);
@@ -86,8 +86,8 @@ namespace ChaosTest {
 		{
 			BoxIDs.Add(ClusteredParticles.Size());
 			FPBDRigidParticleHandle* Box = AppendClusteredParticleBox(Particles, FVec3((FReal)100, (FReal)100, (FReal)100));
-			Box->X() = FVec3((FReal)i * (FReal)100, (FReal)0, (FReal)0);
-			Box->P() = Box->X();
+			Box->SetX(FVec3((FReal)i * (FReal)100, (FReal)0, (FReal)0));
+			Box->SetP(Box->GetX());
 			Boxes.Add(Box);
 
 			Evolution.EnableParticle(Box);
@@ -208,8 +208,8 @@ namespace ChaosTest {
 		{
 			BoxIDs.Add(ClusteredParticles.Size());
 			FPBDRigidParticleHandle* Box = AppendClusteredParticleBox(Particles, FVec3((FReal)100, (FReal)100, (FReal)100));
-			Box->X() = FVec3((FReal)i * (FReal)100, (FReal)0, (FReal)0);
-			Box->P() = Box->X();
+			Box->SetX(FVec3((FReal)i * (FReal)100, (FReal)0, (FReal)0));
+			Box->SetP(Box->GetX());
 			Boxes.Add(Box);
 
 			Evolution.EnableParticle(Box);
diff --git a/Engine/Source/Programs/HeadlessChaos/Private/HeadlessChaosTestCollisions.cpp b/Engine/Source/Programs/HeadlessChaos/Private/HeadlessChaosTestCollisions.cpp
--- a/Engine/Source/Programs/HeadlessChaos/Private/HeadlessChaosTestCollisions.cpp
+++ b/Engine/Source/Programs/HeadlessChaos/Private/HeadlessChaosTestCollisions.cpp
@@ -15,9 +15,9 @@
 
 #define SMALL_THRESHOLD 1e-4
 
-#define RESET_PQ(Particle) Particle->P() = Particle->X(); Particle->SetQ(Particle->GetR());
-#define INVARIANT_XR_START(Particle) FVec3 InvariantPreX_##Particle = Particle->X(); FRotation3 InvariantPreR_##Particle = Particle->GetR()
-#define INVARIANT_XR_END(Particle) EXPECT_TRUE(InvariantPreX_##Particle.Equals(Particle->X())); EXPECT_TRUE(InvariantPreR_##Particle.Equals(Particle->GetR()))
+#define RESET_PQ(Particle) Particle->SetP(Particle->GetX()); Particle->SetQ(Particle->GetR());
+#define INVARIANT_XR_START(Particle) FVec3 InvariantPreX_##Particle = Particle->GetX(); FRotation3 InvariantPreR_##Particle = Particle->GetR()
+#define INVARIANT_XR_END(Particle) EXPECT_TRUE(InvariantPreX_##Particle.Equals(Particle->GetX())); EXPECT_TRUE(InvariantPreR_##Particle.Equals(Particle->GetR()))
 #define INVARIANT_VW_START(Particle) FVec3 InvariantPreV_##Particle = Particle->GetV(); FVec3 InvariantPreW_##Particle = Particle->GetW()
 #define INVARIANT_VW_END(Particle) EXPECT_TRUE(InvariantPreV_##Particle.Equals(Particle->GetV())); EXPECT_TRUE(InvariantPreW_##Particle.Equals(Particle->GetW()))
 
@@ -43,20 +43,20 @@ namespace ChaosTest {
 		Particles.GetParticleHandles().AddArray(&PerParticlePhysicsMaterials);
 
 		auto Box1 = AppendDynamicParticleBox(Particles);
-		Box1->X() = FVec3(1.f);
+		Box1->SetX(FVec3(1.f));
 		Box1->SetR(FRotation3(FQuat::Identity));
-		Box1->P() = Box1->X();
+		Box1->SetP(Box1->GetX());
 		Box1->SetQ(Box1->GetR());
 		Box1->AuxilaryValue(PhysicsMaterials) = MakeSerializable(PhysicsMaterial);
-		Box1->UpdateWorldSpaceState(FRigidTransform3(Box1->X(), Box1->GetR()), FVec3(0));
+		Box1->UpdateWorldSpaceState(FRigidTransform3(Box1->GetX(), Box1->GetR()), FVec3(0));
 
 		auto Box2 = AppendDynamicParticleBox(Particles);
-		Box2->X() = FVec3(0.5f, 0.5f, 1.9f);
+		Box2->SetX(FVec3(0.5f, 0.5f, 1.9f));
 		Box2->SetR(FRotation3(FQuat::Identity));
-		Box2->P() = Box2->X();
+		Box2->SetP(Box2->GetX());
 		Box2->SetQ(Box2->GetR());
 		Box2->AuxilaryValue(PhysicsMaterials) = MakeSerializable(PhysicsMaterial);
-		Box2->UpdateWorldSpaceState(FRigidTransform3(Box2->X(), Box2->GetR()), FVec3(0));
+		Box2->UpdateWorldSpaceState(FRigidTransform3(Box2->GetX(), Box2->GetR()), FVec3(0));
 
 		FPBDCollisionConstraintAccessor Collisions(Particles, Collided, PhysicsMaterials, PerParticlePhysicsMaterials, 1, 1);
 		Collisions.ComputeConstraints(0.f);
@@ -95,20 +95,20 @@ namespace ChaosTest {
 		Particles.GetParticleHandles().AddArray(&PerParticlePhysicsMaterials);
 
 		auto Box1 = AppendDynamicParticleConvexBox(Particles, FVec3(1.f) );
-		Box1->X() = FVec3(0.f);
+		Box1->SetX(FVec3(0.f));
 		Box1->SetR(FRotation3(FQuat::Identity));
-		Box1->P() = Box1->X();
+		Box1->SetP(Box1->GetX());
 		Box1->SetQ(Box1->GetR());
 		Box1->AuxilaryValue(PhysicsMaterials) = MakeSerializable(PhysicsMaterial);
-		Box1->UpdateWorldSpaceState(FRigidTransform3(Box1->X(), Box1->GetR()), FVec3(0));
+		Box1->UpdateWorldSpaceState(FRigidTransform3(Box1->GetX(), Box1->GetR()), FVec3(0));
 
 		auto Box2 = AppendDynamicParticleBox(Particles, FVec3(1.f) );
-		Box2->X() = FVec3(1.25f, 0.f, 0.f);
+		Box2->SetX(FVec3(1.25f, 0.f, 0.f));
 		Box2->SetR(FRotation3(FQuat::Identity));
-		Box2->P() = Box2->X();
+		Box2->SetP(Box2->GetX());
 		Box2->SetQ(Box2->GetR());
 		Box2->AuxilaryValue(PhysicsMaterials) = MakeSerializable(PhysicsMaterial);
-		Box2->UpdateWorldSpaceState(FRigidTransform3(Box2->X(), Box2->GetR()), FVec3(0));
+		Box2->UpdateWorldSpaceState(FRigidTransform3(Box2->GetX(), Box2->GetR()), FVec3(0));
 
 		FPBDCollisionConstraintAccessor Collisions(Particles, Collided, PhysicsMaterials, PerParticlePhysicsMaterials, 1, 1);
 		Collisions.ComputeConstraints(0.f);
@@ -146,16 +146,16 @@ namespace ChaosTest {
 
 		auto Floor = AppendStaticAnalyticFloor(Particles);
 		auto Box = AppendDynamicParticleBox(Particles);
-		Box->P() = FVec3(0, 1, 0);
+		Box->SetP(FVec3(0, 1, 0));
 		Box->SetQ(FRotation3(FQuat::Identity));
 		Box->SetV(FVec3(0, 0, -1));
 		Box->SetPreV(Box->GetV());
-		Box->X() = Box->P() - Box->GetV() * Dt;
+		Box->SetX(Box->GetP() - Box->GetV() * Dt);
 		Box->SetR(Box->GetQ());
 		Box->AuxilaryValue(PhysicsMaterials) = MakeSerializable(PhysicsMaterial);
-		Box->UpdateWorldSpaceState(FRigidTransform3(Box->P(), Box->GetQ()), FVec3(0));
+		Box->UpdateWorldSpaceState(FRigidTransform3(Box->GetP(), Box->GetQ()), FVec3(0));
 
-		const FReal InitialBoxZ = Box->X().Z;
+		const FReal InitialBoxZ = Box->GetX().Z;
 
 		FPBDCollisionConstraintAccessor Collisions(Particles, Collided, PhysicsMaterials, PerParticlePhysicsMaterials, 2, 5);
 
@@ -203,7 +203,7 @@ namespace ChaosTest {
 
 		// Box will not move because the default depentration velocity is zero
 		const FReal ExpectedBoxZ = InitialBoxZ + Collisions.CollisionConstraints.GetSolverSettings().DepenetrationVelocity * Dt;
-		EXPECT_NEAR(Box->P().Z, ExpectedBoxZ, 1.e-2f);
+		EXPECT_NEAR(Box->GetP().Z, ExpectedBoxZ, 1.e-2f);
 
 		// Velocity is below the restitution threshold, so expecting 0 velocity despite the fact that restitution is 1
 		EXPECT_TRUE(Box->GetV().Equals(FVec3(0)));
@@ -229,14 +229,14 @@ namespace ChaosTest {
 
 		auto Floor = AppendStaticConvexFloor(Particles);
 		auto Box = AppendDynamicParticleConvexBox( Particles, FVec3(50) );
-		Box->X() = FVec3(0, 0, 49);
+		Box->SetX(FVec3(0, 0, 49));
 		Box->SetR(FRotation3(FQuat::Identity));
 		Box->SetV(FVec3(0, 0, -1));
 		Box->SetPreV(Box->GetV());
-		Box->P() = Box->X();
+		Box->SetP(Box->GetX());
 		Box->SetQ(Box->GetR());
 		Box->AuxilaryValue(PhysicsMaterials) = MakeSerializable(PhysicsMaterial);
-		Box->UpdateWorldSpaceState(FRigidTransform3(Box->P(), Box->GetQ()), FVec3(0));
+		Box->UpdateWorldSpaceState(FRigidTransform3(Box->GetP(), Box->GetQ()), FVec3(0));
 
 		const FReal Dt = FReal(1) / FReal(24.);
 
@@ -302,16 +302,16 @@ namespace ChaosTest {
 
 		auto Floor = AppendStaticAnalyticFloor(Particles);
 		auto Box = AppendDynamicParticleBox(Particles);
-		Box->P() = FVec3(0, 1, 0);
+		Box->SetP(FVec3(0, 1, 0));
 		Box->SetQ(FRotation3(FQuat::Identity));
 		Box->SetV(FVec3(0, 0, -1));
 		Box->SetPreV(Box->GetV());
-		Box->X() = Box->P() - Box->GetV() * Dt;
+		Box->SetX(Box->GetP() - Box->GetV() * Dt);
 		Box->SetR(Box->GetR());
 		Box->AuxilaryValue(PhysicsMaterials) = MakeSerializable(PhysicsMaterial);
-		Box->UpdateWorldSpaceState(FRigidTransform3(Box->P(), Box->GetQ()), FVec3(0));
+		Box->UpdateWorldSpaceState(FRigidTransform3(Box->GetP(), Box->GetQ()), FVec3(0));
 
-		FReal InitialBoxZ = Box->X().Z;
+		FReal InitialBoxZ = Box->GetX().Z;
 
 		FPBDCollisionConstraintAccessor Collisions(Particles, Collided, PhysicsMaterials, PerParticlePhysicsMaterials, 2, 5);
 
@@ -352,7 +352,7 @@ namespace ChaosTest {
 
 		// Box will not move because the default depentration velocity is zero
 		const FReal ExpectedBoxZ = InitialBoxZ + Collisions.CollisionConstraints.GetSolverSettings().DepenetrationVelocity * Dt;
-		EXPECT_TRUE(FMath::IsNearlyEqual(Box->P().Z, ExpectedBoxZ, 1.e-2));
+		EXPECT_TRUE(FMath::IsNearlyEqual(Box->GetP().Z, ExpectedBoxZ, 1.e-2));
 	}
 
 	void CollisionBoxPlaneRestitution()
@@ -373,14 +373,14 @@ namespace ChaosTest {
 
 		auto Floor = AppendStaticAnalyticFloor(Particles);
 		auto Box = AppendDynamicParticleBox(Particles);
-		Box->P() = FVec3(0, 0, 0);
+		Box->SetP(FVec3(0, 0, 0));
 		Box->SetQ(FRotation3(FQuat::Identity));
 		Box->SetV(FVec3(0, 0, -100));
 		Box->SetPreV(Box->GetV());
-		Box->X() = Box->P() - Box->GetV() * Dt;
+		Box->SetX(Box->GetP() - Box->GetV() * Dt);
 		Box->SetR(Box->GetQ());
 		Box->AuxilaryValue(PhysicsMaterials) = MakeSerializable(PhysicsMaterial);
-		Box->UpdateWorldSpaceState(FRigidTransform3(Box->P(), Box->GetQ()), FVec3(0));
+		Box->UpdateWorldSpaceState(FRigidTransform3(Box->GetP(), Box->GetQ()), FVec3(0));
 
 		FPBDCollisionConstraintAccessor Collisions(Particles, Collided, PhysicsMaterials, PerParticlePhysicsMaterials, 2, 5);
 
@@ -424,7 +424,7 @@ namespace ChaosTest {
 		EXPECT_TRUE(Box->GetW().Equals(FVec3(0)));
 
 		// should end up outside the plane
-		EXPECT_GE(Box->P().Z, -Box->GetGeometry()->BoundingBox().Min().Z);
+		EXPECT_GE(Box->GetP().Z, -Box->GetGeometry()->BoundingBox().Min().Z);
 		EXPECT_TRUE(Box->GetQ().Equals(FQuat::Identity));
 	}
 
@@ -447,18 +447,18 @@ namespace ChaosTest {
 		const FReal Dt = FReal(1) / FReal(24.);
 
 		FGeometryParticleHandle* StaticCube = AppendStaticParticleBox(Particles, FVec3(100.0f));
-		StaticCube->X() = FVec3(0, 0, -50.0f);
-		StaticCube->UpdateWorldSpaceState(FRigidTransform3(StaticCube->X(), StaticCube->GetR()), FVec3(0));
+		StaticCube->SetX(FVec3(0, 0, -50.0f));
+		StaticCube->UpdateWorldSpaceState(FRigidTransform3(StaticCube->GetX(), StaticCube->GetR()), FVec3(0));
 
 		FPBDRigidParticleHandle* DynamicCube = AppendDynamicParticleBox(Particles, FVec3(100.0f));
-		DynamicCube->X() = FVec3(0, 0, 50);
+		DynamicCube->SetX(FVec3(0, 0, 50));
 		DynamicCube->SetR(FRotation3::FromIdentity());
 		DynamicCube->SetV(FVec3(0, 0, -100));
 		DynamicCube->SetPreV(DynamicCube->GetV());
-		DynamicCube->P() = DynamicCube->X() + DynamicCube->GetV() * Dt;
+		DynamicCube->SetP(DynamicCube->GetX() + DynamicCube->GetV() * Dt);
 		DynamicCube->SetQ(DynamicCube->GetR());
 		DynamicCube->AuxilaryValue(PhysicsMaterials) = MakeSerializable(PhysicsMaterial);
-		DynamicCube->UpdateWorldSpaceState(FRigidTransform3(DynamicCube->P(), DynamicCube->GetQ()), FVec3(0));
+		DynamicCube->UpdateWorldSpaceState(FRigidTransform3(DynamicCube->GetP(), DynamicCube->GetQ()), FVec3(0));
 
 		FPBDCollisionConstraintAccessor Collisions(Particles, Collided, PhysicsMaterials, PerParticlePhysicsMaterials, 2, 5);
 
@@ -513,21 +513,21 @@ namespace ChaosTest {
 		Particles.GetParticleHandles().AddArray(&PerParticlePhysicsMaterials);
 
 		auto StaticBox = AppendStaticParticleBox(Particles);
-		StaticBox->X() = FVec3(-0.05f, -0.05f, -0.1f);
+		StaticBox->SetX(FVec3(-0.05f, -0.05f, -0.1f));
 		StaticBox->AuxilaryValue(PhysicsMaterials) = MakeSerializable(PhysicsMaterial);
-		StaticBox->UpdateWorldSpaceState(FRigidTransform3(StaticBox->X(), StaticBox->GetR()), FVec3(0));
+		StaticBox->UpdateWorldSpaceState(FRigidTransform3(StaticBox->GetX(), StaticBox->GetR()), FVec3(0));
 
 		FReal Dt = FReal(1) / FReal(24.);
 
 		auto Box2 = AppendDynamicParticleBox(Particles);
 		FVec3 StartingPoint(0.5f);
-		Box2->P() = StartingPoint;
+		Box2->SetP(StartingPoint);
 		Box2->SetQ(Box2->GetR());
 		Box2->SetV(FVec3(0, 0, -1));
 		Box2->SetPreV(Box2->GetV());
-		Box2->X() = Box2->P() - Box2->GetV() * Dt;
+		Box2->SetX(Box2->GetP() - Box2->GetV() * Dt);
 		Box2->AuxilaryValue(PhysicsMaterials) = MakeSerializable(PhysicsMaterial);
-		Box2->UpdateWorldSpaceState(FRigidTransform3(Box2->P(), Box2->GetQ()), FVec3(0));
+		Box2->UpdateWorldSpaceState(FRigidTransform3(Box2->GetP(), Box2->GetQ()), FVec3(0));
 
 		FBox Region(FVector(FReal(.2)), FVector(FReal(.5)));
 
@@ -577,7 +577,7 @@ namespace ChaosTest {
 		EXPECT_TRUE(Box2->GetV().Size() < FVector(0, -1, 0).Size()); // slowed down  
 		EXPECT_TRUE(Box2->GetW().Size() > 0); // now has rotation 
 
-		EXPECT_FALSE(Box2->P().Equals(StartingPoint)); // moved
+		EXPECT_FALSE(Box2->GetP().Equals(StartingPoint)); // moved
 		EXPECT_FALSE(Box2->GetQ().Equals(FQuat::Identity)); // and rotated
 	}
 }
diff --git a/Engine/Source/Programs/HeadlessChaos/Private/HeadlessChaosTestConstraints.cpp b/Engine/Source/Programs/HeadlessChaos/Private/HeadlessChaosTestConstraints.cpp
--- a/Engine/Source/Programs/HeadlessChaos/Private/HeadlessChaosTestConstraints.cpp
+++ b/Engine/Source/Programs/HeadlessChaos/Private/HeadlessChaosTestConstraints.cpp
@@ -42,7 +42,7 @@ namespace ChaosTest {
 			Evolution.AddConstraintContainer(PositionConstraints);
 			Evolution.AdvanceOneTimeStep(0.1);
 			Evolution.EndFrame(0.1);
-			EXPECT_LT(Evolution.GetParticleHandles().Handle(0)->X().SizeSquared(), SMALL_NUMBER);
+			EXPECT_LT(Evolution.GetParticleHandles().Handle(0)->GetX().SizeSquared(), SMALL_NUMBER);
 		}
 		{
 			FParticleUniqueIndicesMultithreaded UniqueIndices;
@@ -65,16 +65,16 @@ namespace ChaosTest {
 			Evolution.AdvanceOneTimeStep(0.1);
 			Evolution.EndFrame(0.1);
 			auto& Handle = Evolution.GetParticleHandles().Handle(0);
-			EXPECT_LT(FMath::Abs(Handle->X()[0] - 0.5), (FReal)SMALL_NUMBER);
-			EXPECT_LT(FMath::Abs(Handle->X()[1] - 0.5), (FReal)SMALL_NUMBER);
-			EXPECT_LT(FMath::Abs(Handle->X()[2] - 0.5), (FReal)SMALL_NUMBER);
+			EXPECT_LT(FMath::Abs(Handle->GetX()[0] - 0.5), (FReal)SMALL_NUMBER);
+			EXPECT_LT(FMath::Abs(Handle->GetX()[1] - 0.5), (FReal)SMALL_NUMBER);
+			EXPECT_LT(FMath::Abs(Handle->GetX()[2] - 0.5), (FReal)SMALL_NUMBER);
 
 			Evolution.AdvanceOneTimeStep(0.1);
 			Evolution.EndFrame(0.1);
 
-			EXPECT_LT(FMath::Abs(Handle->X()[0] - 1), (FReal)SMALL_NUMBER);
-			EXPECT_LT(FMath::Abs(Handle->X()[1] - 1), (FReal)SMALL_NUMBER);
-			EXPECT_LT(FMath::Abs(Handle->X()[2] - 1), (FReal)SMALL_NUMBER);
+			EXPECT_LT(FMath::Abs(Handle->GetX()[0] - 1), (FReal)SMALL_NUMBER);
+			EXPECT_LT(FMath::Abs(Handle->GetX()[1] - 1), (FReal)SMALL_NUMBER);
+			EXPECT_LT(FMath::Abs(Handle->GetX()[2] - 1), (FReal)SMALL_NUMBER);
 		}
 	}
 
@@ -97,7 +97,7 @@ namespace ChaosTest {
 
 		Evolution.SetNumPositionIterations(Iterations);
 
-		Dynamics[1]->X() = FVec3(500, 0, 0);
+		Dynamics[1]->SetX(FVec3(500, 0, 0));
 		FVec3 JointConstraintPosition = FVec3(0, 0, 0);
 
 		TArray<FPBDRigidParticleHandle*> PositionParticles = { Dynamics[0] };
@@ -118,8 +118,8 @@ namespace ChaosTest {
 			Evolution.AdvanceOneTimeStep(Dt);
 			Evolution.EndFrame(Dt);
 
-			const auto& Pos0 = Dynamics[0]->X();
-			const auto& Pos1 = Dynamics[1]->X();
+			const auto& Pos0 = Dynamics[0]->GetX();
+			const auto& Pos1 = Dynamics[1]->GetX();
 			const float Delta0 = (Pos0 - FVec3(0, 0, 0)).Size();
 			const float Separation = (Pos1 - Pos0).Size();
 
@@ -159,7 +159,7 @@ namespace ChaosTest {
 
 			// chassis particle
 			auto* DynamicParticle = Evolution.CreateDynamicParticles(1)[0];
-			DynamicParticle->X() = FVec3(0, 10, 10);
+			DynamicParticle->SetX(FVec3(0, 10, 10));
 			//	DynamicParticle->R() = FRotation3::FromAxisAngle(FVec3(0,1,0), PI); // upside down
 			DynamicParticle->I() = TVec3<FRealSingle>(100.0f);
 			DynamicParticle->InvI() = TVec3<FRealSingle>(1.0f / 100.0f);
@@ -178,7 +178,7 @@ namespace ChaosTest {
 			Evolution.AdvanceOneTimeStep(0.1);
 			Evolution.EndFrame(0.1);
 
-			const FVec3& Pos = Evolution.GetParticleHandles().Handle(0)->X();
+			const FVec3& Pos = Evolution.GetParticleHandles().Handle(0)->GetX();
 			const FRotation3& Rot = Evolution.GetParticleHandles().Handle(0)->GetR();
 
 			//UE_LOG(LogChaos, Warning, TEXT("Pos %s"), *Pos.ToString());
@@ -205,7 +205,7 @@ namespace ChaosTest {
 
 			// chassis particle
 			auto* DynamicParticle = Evolution.CreateDynamicParticles(1)[0];
-			DynamicParticle->X() = FVec3(50, 10, 10);
+			DynamicParticle->SetX(FVec3(50, 10, 10));
 
 			// minimize rotation using high inertia
 			DynamicParticle->I() = TVec3<FRealSingle>(100000.0f);
@@ -226,7 +226,7 @@ namespace ChaosTest {
 			Evolution.AddConstraintContainer(SuspensionConstraints);
 			Evolution.EnableParticle(DynamicParticle);
 
-			const FVec3& Pos = Evolution.GetParticleHandles().Handle(0)->X();
+			const FVec3& Pos = Evolution.GetParticleHandles().Handle(0)->GetX();
 			const FRotation3& Rot = Evolution.GetParticleHandles().Handle(0)->GetR();
 
 			Evolution.AdvanceOneTimeStep(0.1);
@@ -268,7 +268,7 @@ namespace ChaosTest {
 		// chassis particle
 		auto* DynamicParticle = Evolution.CreateDynamicParticles(1)[0];
 		DynamicParticle->SetLinearEtherDrag(0.f);
-		DynamicParticle->X() = FVec3(0, 0, 10);
+		DynamicParticle->SetX(FVec3(0, 0, 10));
 		DynamicParticle->M() = Mass;
 		DynamicParticle->InvM() = 1.0f / Mass;
 		DynamicParticle->I() = TVec3<FRealSingle>(100000.0f);
@@ -303,7 +303,7 @@ namespace ChaosTest {
 		Evolution.AddConstraintContainer(SuspensionConstraints);
 		Evolution.EnableParticle(DynamicParticle);
 
-		const FVec3& Pos = Evolution.GetParticleHandles().Handle(0)->X();
+		const FVec3& Pos = Evolution.GetParticleHandles().Handle(0)->GetX();
 		const FRotation3& Rot = Evolution.GetParticleHandles().Handle(0)->GetR();
 		const float DeltaTime = 1.0f / 30.0f;
 
diff --git a/Engine/Source/Programs/HeadlessChaos/Private/HeadlessChaosTestForces.cpp b/Engine/Source/Programs/HeadlessChaos/Private/HeadlessChaosTestForces.cpp
--- a/Engine/Source/Programs/HeadlessChaos/Private/HeadlessChaosTestForces.cpp
+++ b/Engine/Source/Programs/HeadlessChaos/Private/HeadlessChaosTestForces.cpp
@@ -24,7 +24,7 @@ namespace ChaosTest
 		TArray<FPBDRigidParticleHandle*> Dynamics = Evolution.CreateDynamicParticles(1);
 		Evolution.EnableParticle(Dynamics[0]);
 		Evolution.AdvanceOneTimeStep(0.1);
-		EXPECT_LT(Dynamics[0]->X()[2], 0);
+		EXPECT_LT(Dynamics[0]->GetX()[2], 0);
 	}
 	
 	GTEST_TEST(AllEvolutions,Forces)
diff --git a/Engine/Source/Programs/HeadlessChaos/Private/HeadlessChaosTestImplicits.cpp b/Engine/Source/Programs/HeadlessChaos/Private/HeadlessChaosTestImplicits.cpp
--- a/Engine/Source/Programs/HeadlessChaos/Private/HeadlessChaosTestImplicits.cpp
+++ b/Engine/Source/Programs/HeadlessChaos/Private/HeadlessChaosTestImplicits.cpp
@@ -1705,15 +1705,15 @@ namespace ChaosTest {
 		{
 			FParticles Particles;
 			Particles.AddParticles(9);
-			Particles.X(0) = FVec3(-1, -1, -1);
-			Particles.X(1) = FVec3(-1, -1, 1);
-			Particles.X(2) = FVec3(-1, 1, -1);
-			Particles.X(3) = FVec3(-1, 1, 1);
-			Particles.X(4) = FVec3(1, -1, -1);
-			Particles.X(5) = FVec3(1, -1, 1);
-			Particles.X(6) = FVec3(1, 1, -1);
-			Particles.X(7) = FVec3(1, 1, 1);
-			Particles.X(8) = FVec3(0, 0, 0);
+			Particles.SetX(0, FVec3(-1, -1, -1));
+			Particles.SetX(1, FVec3(-1, -1, 1));
+			Particles.SetX(2, FVec3(-1, 1, -1));
+			Particles.SetX(3, FVec3(-1, 1, 1));
+			Particles.SetX(4, FVec3(1, -1, -1));
+			Particles.SetX(5, FVec3(1, -1, 1));
+			Particles.SetX(6, FVec3(1, 1, -1));
+			Particles.SetX(7, FVec3(1, 1, 1));
+			Particles.SetX(8, FVec3(0, 0, 0));
 			const FTriangleMesh TriMesh = FTriangleMesh::GetConvexHullFromParticles(Particles);
 			EXPECT_EQ(TriMesh.GetSurfaceElements().Num(), 12);
 			for (const auto& Tri : TriMesh.GetSurfaceElements())
@@ -1727,7 +1727,7 @@ namespace ChaosTest {
 			Vertices.SetNum((int32)Particles.Size());
 			for (int32 VertexIndex = 0; VertexIndex < (int32)Particles.Size(); ++VertexIndex)
 			{
-				Vertices[VertexIndex] = Particles.X(VertexIndex);
+				Vertices[VertexIndex] = Particles.GetX(VertexIndex);
 			}
 			FConvex Convex(Vertices, 0.0f);
 			const TArray<FConvex::FVec3Type>& CulledParticles = Convex.GetVertices();
@@ -1735,11 +1735,11 @@ namespace ChaosTest {
 
 			for (int32 Idx = 0; Idx < CulledParticles.Num(); ++Idx)
 			{
-				EXPECT_NE(Particles.X(8), (Chaos::TVector<FRealDouble, 3>)CulledParticles[Idx]);	//interior particle gone
+				EXPECT_NE(Particles.GetX(8), (Chaos::TVector<FRealDouble, 3>)CulledParticles[Idx]);	//interior particle gone
 				bool bFound = false;
 				for (uint32 InnerIdx = 0; InnerIdx < Particles.Size(); ++InnerIdx)	//remaining particles are from the original set
 				{
-					if (Particles.X(InnerIdx) == (Chaos::TVector<FRealDouble,3>)CulledParticles[Idx])
+					if (Particles.GetX(InnerIdx) == (Chaos::TVector<FRealDouble,3>)CulledParticles[Idx])
 					{
 						bFound = true;
 						break;
@@ -1753,12 +1753,12 @@ namespace ChaosTest {
 		{
 			FParticles Particles;
 			Particles.AddParticles(6);
-			Particles.X(0) = FVec3(-1, -1, -1);
-			Particles.X(1) = FVec3(1, -1, -1);
-			Particles.X(2) = FVec3(1, 1, -1);
-			Particles.X(3) = FVec3(0, 0, 0.5);
-			Particles.X(4) = (Particles.X(3) - Particles.X(1)) * 0.5 + Particles.X(1) + FVec3(0, 0, 0.1);
-			Particles.X(5) = Particles.X(4) + FVec3(-0.1, 0, 0);
+			Particles.SetX(0, FVec3(-1, -1, -1));
+			Particles.SetX(1, FVec3(1, -1, -1));
+			Particles.SetX(2, FVec3(1, 1, -1));
+			Particles.SetX(3, FVec3(0, 0, 0.5));
+			Particles.SetX(4, (Particles.GetX(3) - Particles.GetX(1)) * 0.5 + Particles.GetX(1) + FVec3(0, 0, 0.1));
+			Particles.SetX(5, Particles.GetX(4) + FVec3(-0.1, 0, 0));
 			const FTriangleMesh TriMesh = FTriangleMesh::GetConvexHullFromParticles(Particles);
 			//EXPECT_EQ(TriMesh.GetSurfaceElements().Num(), 6);
 		}
diff --git a/Engine/Source/Programs/HeadlessChaos/Private/HeadlessChaosTestJointBreak.cpp b/Engine/Source/Programs/HeadlessChaos/Private/HeadlessChaosTestJointBreak.cpp
--- a/Engine/Source/Programs/HeadlessChaos/Private/HeadlessChaosTestJointBreak.cpp
+++ b/Engine/Source/Programs/HeadlessChaos/Private/HeadlessChaosTestJointBreak.cpp
@@ -34,7 +34,7 @@ namespace ChaosTest {
 			Test.Evolution.AdvanceOneTimeStep(Dt);
 			Test.Evolution.EndFrame(Dt);
 		}
-		EXPECT_NEAR(Test.GetParticle(1)->X().Z, Test.ParticlePositions[1].Z, 1.0f);
+		EXPECT_NEAR(Test.GetParticle(1)->GetX().Z, Test.ParticlePositions[1].Z, 1.0f);
 
 		// Nothing should have broken
 		EXPECT_FALSE(bBrokenCallbackCalled);
@@ -54,7 +54,7 @@ namespace ChaosTest {
 			Test.Evolution.EndFrame(Dt);
 		}
 		FReal ExpectedZ = Test.ParticlePositions[1].Z - 0.5f * Gravity * (NumIts * Dt) * (NumIts * Dt);
-		EXPECT_NEAR(Test.GetParticle(1)->X().Z, ExpectedZ, 1.0f);
+		EXPECT_NEAR(Test.GetParticle(1)->GetX().Z, ExpectedZ, 1.0f);
 	}
 
 	GTEST_TEST(AllEvolutions, JointBreakTests_TestManualBreak)
diff --git a/Engine/Source/Programs/HeadlessChaos/Private/HeadlessChaosTestJointForces.cpp b/Engine/Source/Programs/HeadlessChaos/Private/HeadlessChaosTestJointForces.cpp
--- a/Engine/Source/Programs/HeadlessChaos/Private/HeadlessChaosTestJointForces.cpp
+++ b/Engine/Source/Programs/HeadlessChaos/Private/HeadlessChaosTestJointForces.cpp
@@ -203,7 +203,7 @@ namespace ChaosTest {
 		FGenericParticleHandle P1 = Test.GetParticle(1);
 
 		// Reposition the particle to have some extension in the spring
-		P1->InitTransform(P1->X() + FVec3(0,0,-Extension), P1->GetR());
+		P1->InitTransform(P1->GetX() + FVec3(0,0,-Extension), P1->GetR());
 
 		// Run the sim
 		Test.Evolution.AdvanceOneTimeStep(Dt);
@@ -395,7 +395,7 @@ namespace ChaosTest {
 		const FReal ExpectedVZ = 0;
 		const FReal ExpectedZ = Test.ParticlePositions[1].Z - Extension;
 		EXPECT_NEAR(P1->V().Z, ExpectedVZ, 1);
-		EXPECT_NEAR(P1->X().Z, ExpectedZ, 1);
+		EXPECT_NEAR(P1->GetX().Z, ExpectedZ, 1);
 	}
 
 	// Check that the maximum drive force setting honored for linear drives.
@@ -498,13 +498,13 @@ namespace ChaosTest {
 		const FReal ExpectedVZ = 0;
 		const FReal ExpectedZ = Test.ParticlePositions[1].Z - Extension;
 		EXPECT_NEAR(P1->V().Z, ExpectedVZ, 1);
-		EXPECT_NEAR(P1->X().Z, ExpectedZ, 1);
+		EXPECT_NEAR(P1->GetX().Z, ExpectedZ, 1);
 	}
 
 
 	// Check that the maximum drive force setting honored for linear drives.
 	// We repeat TestLinearDriveForceMode_MaxForcePreTest but with a max force which is less than K.X at the rest extension
-	// NOTE: Using Accleration mode on the drive but with adjusted stiffness and damping. 
+	// NOTE: Using Acceleration mode on the drive but with adjusted stiffness and damping. 
 	//
 	GTEST_TEST(JointForceTests, TestLinearDriveAccMode_MaxForce)
 	{
diff --git a/Engine/Source/Programs/HeadlessChaos/Private/HeadlessChaosTestJointPlasticity.cpp b/Engine/Source/Programs/HeadlessChaos/Private/HeadlessChaosTestJointPlasticity.cpp
--- a/Engine/Source/Programs/HeadlessChaos/Private/HeadlessChaosTestJointPlasticity.cpp
+++ b/Engine/Source/Programs/HeadlessChaos/Private/HeadlessChaosTestJointPlasticity.cpp
@@ -154,7 +154,7 @@ namespace ChaosTest {
 
 		Test.Create();
 
-		FReal Z = Test.SOAs.GetDynamicParticles().X(0)[2];
+		FReal Z = Test.SOAs.GetDynamicParticles().GetX(0)[2];
 		EXPECT_TRUE(FMath::IsNearlyEqual(Z, (FReal)50.));
 
 		// Run the sim
@@ -166,7 +166,7 @@ namespace ChaosTest {
 			//FVector Pos = Test.SOAs.GetDynamicParticles().X(0);
 			//std::cout << Pos.X << "," << Pos.Y << "," << Pos.Z << std::endl;
 		}
-		FReal ZPost = Test.SOAs.GetDynamicParticles().X(0)[2];
+		FReal ZPost = Test.SOAs.GetDynamicParticles().GetX(0)[2];
 
 		// Nothing should have reset
 		EXPECT_TRUE(FMath::IsNearlyEqual(Z, ZPost, (FReal)5.));
@@ -207,7 +207,7 @@ namespace ChaosTest {
 
 		Test.Create();
 
-		FReal Z = Test.SOAs.GetDynamicParticles().X(0)[2];
+		FReal Z = Test.SOAs.GetDynamicParticles().GetX(0)[2];
 		EXPECT_TRUE(FMath::IsNearlyEqual(Z, (FReal)50.));
 
 		// Run the sim
@@ -220,7 +220,7 @@ namespace ChaosTest {
 			//FVector Pos = Test.SOAs.GetDynamicParticles().X(0);
 			//std::cout << Pos.X << "," << Pos.Y << "," << Pos.Z << std::endl;
 		}
-		FReal ZPost = Test.SOAs.GetDynamicParticles().X(0)[2];
+		FReal ZPost = Test.SOAs.GetDynamicParticles().GetX(0)[2];
 
 		// The linear spring should have reset. 
 		EXPECT_TRUE(ZPost < Z - PlasticityLimit);
diff --git a/Engine/Source/Programs/HeadlessChaos/Private/HeadlessChaosTestJoints.cpp b/Engine/Source/Programs/HeadlessChaos/Private/HeadlessChaosTestJoints.cpp
--- a/Engine/Source/Programs/HeadlessChaos/Private/HeadlessChaosTestJoints.cpp
+++ b/Engine/Source/Programs/HeadlessChaos/Private/HeadlessChaosTestJoints.cpp
@@ -136,11 +136,11 @@ namespace ChaosTest {
 
 
 			// Joint position calculated from pose and local-space joint pos
-			const FVec3 Box2WorldSpaceJointPosition = Test.GetParticle(Box2Id)->GetR().RotateVector(Box2LocalSpaceJointPosition) + Test.GetParticle(Box2Id)->X();
+			const FVec3 Box2WorldSpaceJointPosition = Test.GetParticle(Box2Id)->GetR().RotateVector(Box2LocalSpaceJointPosition) + Test.GetParticle(Box2Id)->GetX();
 			EXPECT_LT((Box2WorldSpaceJointPosition - Test.JointPositions[0]).Size(), (FReal)0.1);
 
 			// Kinematic particle should not have moved
-			EXPECT_LT((Test.GetParticle(Box1Id)->X() - Test.ParticlePositions[0]).Size(), (FReal)0.1);
+			EXPECT_LT((Test.GetParticle(Box1Id)->GetX() - Test.ParticlePositions[0]).Size(), (FReal)0.1);
 		}
 	}
 
@@ -197,7 +197,7 @@ namespace ChaosTest {
 			// Nothing should have moved
 			for (int32 ParticleIndex = 0; ParticleIndex < Test.ParticlePositions.Num(); ++ParticleIndex)
 			{
-				EXPECT_LT((Test.GetParticle(ParticleIndex)->X() - Test.ParticlePositions[ParticleIndex]).Size(), (FReal)0.1) << "Initial configuration instability on frame " << i;
+				EXPECT_LT((Test.GetParticle(ParticleIndex)->GetX() - Test.ParticlePositions[ParticleIndex]).Size(), (FReal)0.1) << "Initial configuration instability on frame " << i;
 			}
 		}
 
@@ -211,16 +211,16 @@ namespace ChaosTest {
 			Test.Evolution.EndFrame(Dt);
 
 			// Kinematic particle should have moved to animated position
-			EXPECT_LT((Test.GetParticle(Box1Id)->X() - RootPosition).Size(), (FReal)0.1 * BoxSize) << "Post-move instability on frame " << i;
+			EXPECT_LT((Test.GetParticle(Box1Id)->GetX() - RootPosition).Size(), (FReal)0.1 * BoxSize) << "Post-move instability on frame " << i;
 
 			// Particles should remain fixed distance apart (joint point is at Box1 location)
 			// NOTE: when using linear joints the error can be moderately large
-			const FVec3 Delta = Test.GetParticle(Box2Id)->X() - Test.GetParticle(Box1Id)->X();
+			const FVec3 Delta = Test.GetParticle(Box2Id)->GetX() - Test.GetParticle(Box1Id)->GetX();
 			const FReal Distance = Delta.Size();
 			EXPECT_NEAR(Distance, ExpectedDistance, (FReal)0.15 * BoxSize) << "Post-move instability on frame " << i;
 
 			// Joint position calculted from pose and local-space joint pos
-			const FVec3 Box2WorldSpaceJointPosition = Test.GetParticle(Box2Id)->GetR().RotateVector(Box2LocalSpaceJointPosition) + Test.GetParticle(Box2Id)->X();
+			const FVec3 Box2WorldSpaceJointPosition = Test.GetParticle(Box2Id)->GetR().RotateVector(Box2LocalSpaceJointPosition) + Test.GetParticle(Box2Id)->GetX();
 			EXPECT_LT((Box2WorldSpaceJointPosition - RootPosition).Size(), (FReal)0.15 * BoxSize) << "Post-move instability on frame " << i;
 		}
 	}
@@ -286,15 +286,15 @@ namespace ChaosTest {
 			Test.Evolution.EndFrame(Dt);
 
 			// Kinematic particle should have moved to animated position
-			EXPECT_LT((Test.GetParticle(Box1Id)->X() - RootPosition).Size(), (FReal)1) << "Failed on frame " << i;
+			EXPECT_LT((Test.GetParticle(Box1Id)->GetX() - RootPosition).Size(), (FReal)1) << "Failed on frame " << i;
 
 			// Particles should remain fixed distance apart (joint point is at Box1 location)
-			const FVec3 Delta = Test.GetParticle(Box2Id)->CastToRigidParticle()->P() - Test.GetParticle(Box1Id)->X();
+			const FVec3 Delta = Test.GetParticle(Box2Id)->CastToRigidParticle()->GetP() - Test.GetParticle(Box1Id)->GetX();
 			const FReal Distance = Delta.Size();
 			EXPECT_NEAR(Distance, ExpectedDistance, (FReal)1) << "Failed on frame " << i;
 
 			// Joint position calculated from pose and local-space joint pos
-			const FVec3 Box2WorldSpaceJointPosition = Test.GetParticle(Box2Id)->GetR().RotateVector(Box2LocalSpaceJointPosition) + Test.GetParticle(Box2Id)->X();
+			const FVec3 Box2WorldSpaceJointPosition = Test.GetParticle(Box2Id)->GetR().RotateVector(Box2LocalSpaceJointPosition) + Test.GetParticle(Box2Id)->GetX();
 			EXPECT_LT((Box2WorldSpaceJointPosition - RootPosition).Size(), (FReal)1) << "Failed on frame " << i;
 		}
 	}
@@ -368,7 +368,7 @@ namespace ChaosTest {
 			{
 				const int32 ParticleIndex1 = Test.JointParticleIndices[JointIndex][0];
 				const int32 ParticleIndex2 = Test.JointParticleIndices[JointIndex][1];
-				const FVec3 Delta = Test.GetParticle(ParticleIndex2)->CastToRigidParticle()->P() - Test.GetParticle(ParticleIndex1)->X();
+				const FVec3 Delta = Test.GetParticle(ParticleIndex2)->CastToRigidParticle()->GetP() - Test.GetParticle(ParticleIndex1)->GetX();
 				const FReal Distance = Delta.Size();
 				const FReal ExpectedDistance = (Test.ParticlePositions[ParticleIndex2] - Test.ParticlePositions[ParticleIndex1]).Size();
 				EXPECT_NEAR(Distance, ExpectedDistance, AcceptableDistanceError) << "Joint " << JointIndex << " on frame " << FrameIndex;
@@ -454,7 +454,7 @@ namespace ChaosTest {
 			{
 				const int32 ParticleIndex1 = Test.JointParticleIndices[JointIndex][0];
 				const int32 ParticleIndex2 = Test.JointParticleIndices[JointIndex][1];
-				const FVec3 Delta = Test.GetParticle(ParticleIndex2)->CastToRigidParticle()->P() - Test.GetParticle(ParticleIndex1)->X();
+				const FVec3 Delta = Test.GetParticle(ParticleIndex2)->CastToRigidParticle()->GetP() - Test.GetParticle(ParticleIndex1)->GetX();
 				const FReal Distance = Delta.Size();
 				const FReal ExpectedDistance = (Test.ParticlePositions[ParticleIndex2] - Test.ParticlePositions[ParticleIndex1]).Size();
 				EXPECT_NEAR(Distance, ExpectedDistance, AcceptableDistanceError) << "Joint " << JointIndex << " on frame " << FrameIndex;
@@ -488,10 +488,10 @@ namespace ChaosTest {
 
 		auto StaticBox = AppendStaticParticleBox(Particles, FVec3((FReal)100, (FReal)100, (FReal)100));
 		auto Box2 = AppendDynamicParticleBox(Particles, FVec3((FReal)100, (FReal)100, (FReal)100));
-		StaticBox->X() = FVec3((FReal)0, (FReal)0, (FReal)1000);
+		StaticBox->SetX(FVec3((FReal)0, (FReal)0, (FReal)1000));
 
-		Box2->X() = FVec3((FReal)500, (FReal)0, (FReal)1000);
-		Box2->P() = Box2->X();
+		Box2->SetX(FVec3((FReal)500, (FReal)0, (FReal)1000));
+		Box2->SetP(Box2->GetX());
 		THandleArray<FChaosPhysicsMaterial> PhysicalMaterials;
 		TEvolution Evolution(Particles, PhysicalMaterials);
 		TVec2<FGeometryParticleHandle*> ConstrainedParticles = TVec2<FGeometryParticleHandle*>(StaticBox, Box2);
@@ -513,7 +513,7 @@ namespace ChaosTest {
 		{
 			Evolution.AdvanceOneTimeStep(Dt);
 			Evolution.EndFrame(Dt);
-			EXPECT_LT(FMath::Abs((Box2->GetR().RotateVector(FVec3((FReal)-100, (FReal)0, (FReal)0)) + Box2->X() - Points[0]).Size() - 300.f), 0.1);
+			EXPECT_LT(FMath::Abs((Box2->GetR().RotateVector(FVec3((FReal)-100, (FReal)0, (FReal)0)) + Box2->GetX() - Points[0]).Size() - 300.f), 0.1);
 		}
 	}
 
@@ -535,10 +535,10 @@ namespace ChaosTest {
 
 			auto& StaticBox = *AppendStaticParticleBox(Particles, FVec3((FReal)100, (FReal)100, (FReal)100));
 			auto& Box2 = *AppendDynamicParticleBox(Particles, FVec3((FReal)100, (FReal)100, (FReal)100));
-			StaticBox.X() = FVec3((FReal)0, (FReal)0, (FReal)500);
+			StaticBox.SetX(FVec3((FReal)0, (FReal)0, (FReal)500));
 
-			Box2.X() = FVec3((FReal)500, (FReal)0, (FReal)1000);
-			Box2.P() = Box2.X();
+			Box2.SetX(FVec3((FReal)500, (FReal)0, (FReal)1000));
+			Box2.SetP(Box2.GetX());
 
 			THandleArray<FChaosPhysicsMaterial> PhysicalMaterials;
 			TEvolution Evolution(Particles, PhysicalMaterials);
@@ -559,7 +559,7 @@ namespace ChaosTest {
 				Evolution.AdvanceOneTimeStep(Dt);
 				Evolution.EndFrame(Dt);
 			}
-			EXPECT_LT(Box2.X()[2], 0);
+			EXPECT_LT(Box2.GetX()[2], 0);
 		}
 
 		{
@@ -568,10 +568,10 @@ namespace ChaosTest {
 
 			auto& StaticBox = *AppendStaticParticleBox(Particles, FVec3((FReal)100, (FReal)100, (FReal)100));
 			auto& Box2 = *AppendDynamicParticleBox(Particles, FVec3((FReal)100, (FReal)100, (FReal)100));
-			StaticBox.X() = FVec3((FReal)0, (FReal)0, (FReal)500);
+			StaticBox.SetX(FVec3((FReal)0, (FReal)0, (FReal)500));
 
-			Box2.X() = FVec3((FReal)500, (FReal)0, (FReal)1000);
-			Box2.P() = Box2.X();
+			Box2.SetX(FVec3((FReal)500, (FReal)0, (FReal)1000));
+			Box2.SetP(Box2.GetX());
 
 			THandleArray<FChaosPhysicsMaterial> PhysicalMaterials;
 			TEvolution Evolution(Particles, PhysicalMaterials);
@@ -592,7 +592,7 @@ namespace ChaosTest {
 				Evolution.AdvanceOneTimeStep(Dt);
 				Evolution.EndFrame(Dt);
 			}
-			EXPECT_GT(Box2.X()[2], 0);
+			EXPECT_GT(Box2.GetX()[2], 0);
 		}
 	}
 
@@ -647,14 +647,14 @@ namespace ChaosTest {
 			Test.Evolution.EndFrame(Dt);
 
 			// box 1 not simulating so would expect to not have moved
-			EXPECT_LT(FMath::Abs(Test.ParticlePositions[0].X - Test.GetParticle(Box1Id)->X().X), (FReal)0.1);
-			EXPECT_LT(FMath::Abs(Test.ParticlePositions[0].Y - Test.GetParticle(Box1Id)->X().Y), (FReal)0.1);
-			EXPECT_LT(FMath::Abs(Test.ParticlePositions[0].Z - Test.GetParticle(Box1Id)->X().Z), (FReal)0.09);
+			EXPECT_LT(FMath::Abs(Test.ParticlePositions[0].X - Test.GetParticle(Box1Id)->GetX().X), (FReal)0.1);
+			EXPECT_LT(FMath::Abs(Test.ParticlePositions[0].Y - Test.GetParticle(Box1Id)->GetX().Y), (FReal)0.1);
+			EXPECT_LT(FMath::Abs(Test.ParticlePositions[0].Z - Test.GetParticle(Box1Id)->GetX().Z), (FReal)0.09);
 
 			// box 2 should fall under gravity & not have moved in X or Y, constraint should not 'Apply' if other particle is disabled
-			EXPECT_LT(FMath::Abs(Test.ParticlePositions[1].X - Test.GetParticle(Box2Id)->X().X), (FReal)0.1);
-			EXPECT_LT(FMath::Abs(Test.ParticlePositions[1].Y - Test.GetParticle(Box2Id)->X().Y), (FReal)0.1);
-			EXPECT_GT(FMath::Abs(Test.ParticlePositions[1].Z - Test.GetParticle(Box2Id)->X().Z), (FReal)0.09);
+			EXPECT_LT(FMath::Abs(Test.ParticlePositions[1].X - Test.GetParticle(Box2Id)->GetX().X), (FReal)0.1);
+			EXPECT_LT(FMath::Abs(Test.ParticlePositions[1].Y - Test.GetParticle(Box2Id)->GetX().Y), (FReal)0.1);
+			EXPECT_GT(FMath::Abs(Test.ParticlePositions[1].Z - Test.GetParticle(Box2Id)->GetX().Z), (FReal)0.09);
 		}
 	}
 
@@ -749,8 +749,8 @@ namespace ChaosTest {
 			Test.Evolution.EndFrame(Dt);
 
 			// Neither particle should have moved
-			EXPECT_LT((Test.GetParticle(0)->X() - Test.ParticlePositions[0]).Size(), (FReal)0.1);
-			EXPECT_LT((Test.GetParticle(1)->X() - Test.ParticlePositions[1]).Size(), (FReal)0.1);
+			EXPECT_LT((Test.GetParticle(0)->GetX() - Test.ParticlePositions[0]).Size(), (FReal)0.1);
+			EXPECT_LT((Test.GetParticle(1)->GetX() - Test.ParticlePositions[1]).Size(), (FReal)0.1);
 		}
 	}
 
diff --git a/Engine/Source/Programs/HeadlessChaos/Private/HeadlessChaosTestMassProperties.cpp b/Engine/Source/Programs/HeadlessChaos/Private/HeadlessChaosTestMassProperties.cpp
--- a/Engine/Source/Programs/HeadlessChaos/Private/HeadlessChaosTestMassProperties.cpp
+++ b/Engine/Source/Programs/HeadlessChaos/Private/HeadlessChaosTestMassProperties.cpp
@@ -174,14 +174,14 @@ namespace ChaosTest {
 	{
 		Chaos::FParticles Particles;
 		Particles.AddParticles(8);
-		Particles.X(0) = FVec3(-1, -1, -1);
-		Particles.X(1) = FVec3(-1, -1, 1);
-		Particles.X(2) = FVec3(-1, 1, -1);
-		Particles.X(3) = FVec3(-1, 1, 1);
-		Particles.X(4) = FVec3(1, -1, -1);
-		Particles.X(5) = FVec3(1, -1, 1);
-		Particles.X(6) = FVec3(1, 1, -1);
-		Particles.X(7) = FVec3(1, 1, 1);
+		Particles.SetX(0, FVec3(-1, -1, -1));
+		Particles.SetX(1, FVec3(-1, -1, 1));
+		Particles.SetX(2, FVec3(-1, 1, -1));
+		Particles.SetX(3, FVec3(-1, 1, 1));
+		Particles.SetX(4, FVec3(1, -1, -1));
+		Particles.SetX(5, FVec3(1, -1, 1));
+		Particles.SetX(6, FVec3(1, 1, -1));
+		Particles.SetX(7, FVec3(1, 1, 1));
 		TArray<Chaos::TVec3<int32>> Faces;
 		Faces.SetNum(12);
 		Faces[0] = TVec3<int32>(0, 4, 5);
diff --git a/Engine/Source/Programs/HeadlessChaos/Private/HeadlessChaosTestMinEvolution.cpp b/Engine/Source/Programs/HeadlessChaos/Private/HeadlessChaosTestMinEvolution.cpp
--- a/Engine/Source/Programs/HeadlessChaos/Private/HeadlessChaosTestMinEvolution.cpp
+++ b/Engine/Source/Programs/HeadlessChaos/Private/HeadlessChaosTestMinEvolution.cpp
@@ -64,20 +64,20 @@ namespace ChaosTest
 
 		// Set up Particles
 		// @todo(ccaulfield) this needs to be easier
-		Particles[0]->X() = FVec3(-50, 0, 0);
+		Particles[0]->SetX(FVec3(-50, 0, 0));
 		Particles[0]->M() = 1.0f;
 		Particles[0]->I() = TVec3<FRealSingle>(100.0f, 100.0f, 100.0f);
 		Particles[0]->InvM() = 1.0f;
 		Particles[0]->InvI() = TVec3<FRealSingle>(1.0f / 100.0f, 1.0f / 100.0f, 1.0f / 100.0f);
-		Particles[0]->AuxilaryValue(ParticlePrevXs) = Particles[0]->X();
+		Particles[0]->AuxilaryValue(ParticlePrevXs) = Particles[0]->GetX();
 		Particles[0]->AuxilaryValue(ParticlePrevRs) = Particles[0]->GetR();
 
-		Particles[1]->X() = FVec3(50, 0, 0);
+		Particles[1]->SetX(FVec3(50, 0, 0));
 		Particles[1]->M() = 1.0f;
 		Particles[1]->I() = TVec3<FRealSingle>(100.0f, 100.0f, 100.0f);
 		Particles[1]->InvM() = 1.0f;
 		Particles[1]->InvI() = TVec3<FRealSingle>(1.0f / 100.0f, 1.0f / 100.0f, 1.0f / 100.0f);
-		Particles[1]->AuxilaryValue(ParticlePrevXs) = Particles[1]->X();
+		Particles[1]->AuxilaryValue(ParticlePrevXs) = Particles[1]->GetX();
 		Particles[1]->AuxilaryValue(ParticlePrevRs) = Particles[1]->GetR();
 
 		// Spring connectors at particle centres
@@ -97,8 +97,8 @@ namespace ChaosTest
 		}
 
 		// Particles should be separated by the spring's rest length
-		FVec3 P0 = Particles[0]->X();
-		FVec3 P1 = Particles[1]->X();
+		FVec3 P0 = Particles[0]->GetX();
+		FVec3 P1 = Particles[1]->GetX();
 		FReal Distance01 = (P0 - P1).Size();
 		EXPECT_NEAR(Distance01, Spring->GetRestLength(), 0.1f);
 	}
diff --git a/Engine/Source/Programs/HeadlessChaos/Private/HeadlessChaosTestMostOpposing.cpp b/Engine/Source/Programs/HeadlessChaos/Private/HeadlessChaosTestMostOpposing.cpp
--- a/Engine/Source/Programs/HeadlessChaos/Private/HeadlessChaosTestMostOpposing.cpp
+++ b/Engine/Source/Programs/HeadlessChaos/Private/HeadlessChaosTestMostOpposing.cpp
@@ -28,13 +28,12 @@ namespace ChaosTest
 
 			FTriangleMeshImplicitObject::ParticlesType Particles;
 			Particles.AddParticles(6);
-			Particles.X(0) = FVec3(1, 1, 1);
-			Particles.X(1) = FVec3(5, 1, 1);
-			Particles.X(2) = FVec3(1, 5, 1);
-
-			Particles.X(3) = FVec3(1, 1, 1);
-			Particles.X(4) = FVec3(1, 5, 1);
-			Particles.X(5) = FVec3(1, 1, -5);
+			Particles.SetX(0, FVec3(1, 1, 1));
+			Particles.SetX(1, FVec3(5, 1, 1));
+			Particles.SetX(2, FVec3(1, 5, 1));
+			Particles.SetX(3, FVec3(1, 1, 1));
+			Particles.SetX(4, FVec3(1, 5, 1));
+			Particles.SetX(5, FVec3(1, 1, -5));
 
 			TArray<TVec3<int32>> Indices;
 			Indices.Emplace(0, 1, 2);
@@ -75,24 +74,24 @@ namespace ChaosTest
 			FTriangleMeshImplicitObject::ParticlesType Particles;
 			Particles.AddParticles(12);
 			// in z-y plane
-			Particles.X(0) = FVec3(0, 0, 0);
-			Particles.X(1) = FVec3(0, 1, 0);
-			Particles.X(2) = FVec3(0, 0, 1);
+			Particles.SetX(0, FVec3(0, 0, 0));
+			Particles.SetX(1, FVec3(0, 1, 0));
+			Particles.SetX(2, FVec3(0, 0, 1));
 
 			// In x-z plane
-			Particles.X(3) = FVec3(0, 0, 0);
-			Particles.X(4) = FVec3(1, 0, 0);
-			Particles.X(5) = FVec3(0, 0, 1);
+			Particles.SetX(3, FVec3(0, 0, 0));
+			Particles.SetX(4, FVec3(1, 0, 0));
+			Particles.SetX(5, FVec3(0, 0, 1));
 
 			// In x-y plane
-			Particles.X(6) = FVec3(0, 0, 0);
-			Particles.X(7) = FVec3(1, 0, 0);
-			Particles.X(8) = FVec3(0, 1, 0);
+			Particles.SetX(6, FVec3(0, 0, 0));
+			Particles.SetX(7, FVec3(1, 0, 0));
+			Particles.SetX(8, FVec3(0, 1, 0));
 
 			// One 45 degree slanted triangle
-			Particles.X(9) = FVec3(1, 0, 0);
-			Particles.X(10) = FVec3(1, 1, 0);
-			Particles.X(11) = FVec3(0, 0, 1);
+			Particles.SetX(9, FVec3(1, 0, 0));
+			Particles.SetX(10,  FVec3(1, 1, 0));
+			Particles.SetX(11,  FVec3(0, 0, 1));
 
 			TArray<TVec3<int32>> Indices;
 			Indices.Emplace(0, 1, 2);
@@ -148,15 +147,14 @@ namespace ChaosTest
 			FTriangleMeshImplicitObject::ParticlesType Particles;
 			Particles.AddParticles(6);
 			// in z-y plane
-			Particles.X(0) = FVec3(0, 0, 0);
-			Particles.X(1) = FVec3(0, 100, 0);
-			Particles.X(2) = FVec3(0, 0, 100);
+			Particles.SetX(0, FVec3(0, 0, 0));
+			Particles.SetX(1, FVec3(0, 100, 0));
+			Particles.SetX(2, FVec3(0, 0, 100));
 
 			// In x-z plane
-			Particles.X(3) = FVec3(0, 0, 0);
-			Particles.X(4) = FVec3(100, 0, 0);
-			Particles.X(5) = FVec3(0, 0, 100);
-
+			Particles.SetX(3, FVec3(0, 0, 0));
+			Particles.SetX(4, FVec3(100, 0, 0));
+			Particles.SetX(5, FVec3(0, 0, 100));
 
 			TArray<TVec3<int32>> Indices;
 			Indices.Emplace(0, 1, 2);
