Make consistent API by using GetX, SetX, GetP and SetP for physics particles

diff --git a/Engine/Source/Programs/HeadlessChaos/Private/HeadlessChaosTestMove.cpp b/Engine/Source/Programs/HeadlessChaos/Private/HeadlessChaosTestMove.cpp
--- a/Engine/Source/Programs/HeadlessChaos/Private/HeadlessChaosTestMove.cpp
+++ b/Engine/Source/Programs/HeadlessChaos/Private/HeadlessChaosTestMove.cpp
@@ -26,23 +26,23 @@ namespace ChaosTest
 		// Drop the box on the ground and wait until it sleeps
 		Test.TickUntilSleep();
 		EXPECT_TRUE(Box->IsSleeping());
-		EXPECT_NEAR(Box->X().Z, 0.5 * BoxSize.Z, DistanceTolerance);
+		EXPECT_NEAR(Box->GetX().Z, 0.5 * BoxSize.Z, DistanceTolerance);
 
 		// Move the box horizontally and up into the air
-		const FVec3 NewBoxPos = Box->X() + FVec3(100, 0, 100);
+		const FVec3 NewBoxPos = Box->GetX() + FVec3(100, 0, 100);
 		Test.GetEvolution().SetParticleTransformSwept(Box, NewBoxPos, BoxRot, false);
 
 		// The box should be at its new location and awake
 		EXPECT_FALSE(Box->IsSleeping());
-		EXPECT_NEAR(Box->X().X, NewBoxPos.X, DistanceTolerance);
-		EXPECT_NEAR(Box->X().Y, NewBoxPos.Y, DistanceTolerance);
-		EXPECT_NEAR(Box->X().Z, NewBoxPos.Z, DistanceTolerance);
+		EXPECT_NEAR(Box->GetX().X, NewBoxPos.X, DistanceTolerance);
+		EXPECT_NEAR(Box->GetX().Y, NewBoxPos.Y, DistanceTolerance);
+		EXPECT_NEAR(Box->GetX().Z, NewBoxPos.Z, DistanceTolerance);
 
 		// The box should fall to the ground again
 		Test.TickUntilSleep();
 		EXPECT_TRUE(Box->IsSleeping());
-		EXPECT_NEAR(Box->X().X, NewBoxPos.X, DistanceTolerance);
-		EXPECT_NEAR(Box->X().Z, 0.5 * BoxSize.Z, DistanceTolerance);
+		EXPECT_NEAR(Box->GetX().X, NewBoxPos.X, DistanceTolerance);
+		EXPECT_NEAR(Box->GetX().Z, 0.5 * BoxSize.Z, DistanceTolerance);
 	}
 
 	GTEST_TEST(MoveTests, TestMoveWithSweepNoHit)
@@ -65,23 +65,23 @@ namespace ChaosTest
 		// Drop the box on the ground and wait until it sleeps
 		Test.TickUntilSleep();
 		EXPECT_TRUE(Box->IsSleeping());
-		EXPECT_NEAR(Box->X().Z, 0.5 * BoxSize.Z, DistanceTolerance);
+		EXPECT_NEAR(Box->GetX().Z, 0.5 * BoxSize.Z, DistanceTolerance);
 
 		// Move the box horizontally and up into the air
-		const FVec3 NewBoxPos = Box->X() + FVec3(100, 0, 100);
+		const FVec3 NewBoxPos = Box->GetX() + FVec3(100, 0, 100);
 		Test.GetEvolution().SetParticleTransformSwept(Box, NewBoxPos, BoxRot, false);
 
 		// The box should be at its new location and awake
 		EXPECT_FALSE(Box->IsSleeping());
-		EXPECT_NEAR(Box->X().X, NewBoxPos.X, DistanceTolerance);
-		EXPECT_NEAR(Box->X().Y, NewBoxPos.Y, DistanceTolerance);
-		EXPECT_NEAR(Box->X().Z, NewBoxPos.Z, DistanceTolerance);
+		EXPECT_NEAR(Box->GetX().X, NewBoxPos.X, DistanceTolerance);
+		EXPECT_NEAR(Box->GetX().Y, NewBoxPos.Y, DistanceTolerance);
+		EXPECT_NEAR(Box->GetX().Z, NewBoxPos.Z, DistanceTolerance);
 
 		// The box should fall to the ground again
 		Test.TickUntilSleep();
 		EXPECT_TRUE(Box->IsSleeping());
-		EXPECT_NEAR(Box->X().X, NewBoxPos.X, DistanceTolerance);
-		EXPECT_NEAR(Box->X().Z, 0.5 * BoxSize.Z, DistanceTolerance);
+		EXPECT_NEAR(Box->GetX().X, NewBoxPos.X, DistanceTolerance);
+		EXPECT_NEAR(Box->GetX().Z, 0.5 * BoxSize.Z, DistanceTolerance);
 	}
 
 	GTEST_TEST(MoveTests, TestMoveWithSweepBlockingHit)
@@ -107,22 +107,22 @@ namespace ChaosTest
 		// Drop the box on the ground and wait until it sleeps
 		Test.TickUntilSleep();
 		EXPECT_TRUE(Box->IsSleeping());
-		EXPECT_NEAR(Box->X().Z, 0.5 * BoxSize.Z, DistanceTolerance);
+		EXPECT_NEAR(Box->GetX().Z, 0.5 * BoxSize.Z, DistanceTolerance);
 
 		// Move the box horizontally and up into the air
 		Test.GetEvolution().SetParticleTransformSwept(Box, BoxTargetPos, BoxRot, false);
 
 		// The box should have been stopped on its way to its new location
 		EXPECT_FALSE(Box->IsSleeping());
-		EXPECT_NEAR(Box->X().X, 50, DistanceTolerance);
-		EXPECT_NEAR(Box->X().Y, 0, DistanceTolerance);
-		EXPECT_NEAR(Box->X().Z, 50 + 0.5 * BoxSize.Z, DistanceTolerance);
+		EXPECT_NEAR(Box->GetX().X, 50, DistanceTolerance);
+		EXPECT_NEAR(Box->GetX().Y, 0, DistanceTolerance);
+		EXPECT_NEAR(Box->GetX().Z, 50 + 0.5 * BoxSize.Z, DistanceTolerance);
 
 		// The box should fall to the ground again
 		Test.TickUntilSleep();
 		EXPECT_TRUE(Box->IsSleeping());
-		EXPECT_NEAR(Box->X().X, 50, DistanceTolerance);
-		EXPECT_NEAR(Box->X().Z, 0.5 * BoxSize.Z, DistanceTolerance);
+		EXPECT_NEAR(Box->GetX().X, 50, DistanceTolerance);
+		EXPECT_NEAR(Box->GetX().Z, 0.5 * BoxSize.Z, DistanceTolerance);
 	}
 
 }
\ No newline at end of file
diff --git a/Engine/Source/Programs/HeadlessChaos/Private/HeadlessChaosTestParticleHandle.cpp b/Engine/Source/Programs/HeadlessChaos/Private/HeadlessChaosTestParticleHandle.cpp
--- a/Engine/Source/Programs/HeadlessChaos/Private/HeadlessChaosTestParticleHandle.cpp
+++ b/Engine/Source/Programs/HeadlessChaos/Private/HeadlessChaosTestParticleHandle.cpp
@@ -273,24 +273,24 @@ namespace ChaosTest {
 	template <typename TGeometry, typename TKinematicGeometry, typename TPBDRigid>
 	void ParticleHandleTestHelper(TGeometry* Geometry, TKinematicGeometry* KinematicGeometry, TPBDRigid* PBDRigid)
 	{
-		EXPECT_EQ(Geometry->X()[0], 0);	//default constructor
-		EXPECT_EQ(Geometry->X()[1], 0);
-		EXPECT_EQ(Geometry->X()[2], 0);
+		EXPECT_EQ(Geometry->GetX()[0], 0);	//default constructor
+		EXPECT_EQ(Geometry->GetX()[1], 0);
+		EXPECT_EQ(Geometry->GetX()[2], 0);
 
 		EXPECT_EQ(KinematicGeometry->GetV()[0], 0);	//default constructor
 		EXPECT_EQ(KinematicGeometry->GetV()[1], 0);
 		EXPECT_EQ(KinematicGeometry->GetV()[2], 0);
 
-		EXPECT_EQ(PBDRigid->X()[0], 0);	//default constructor of base
-		EXPECT_EQ(PBDRigid->X()[1], 0);
-		EXPECT_EQ(PBDRigid->X()[2], 0);
+		EXPECT_EQ(PBDRigid->GetX()[0], 0);	//default constructor of base
+		EXPECT_EQ(PBDRigid->GetX()[1], 0);
+		EXPECT_EQ(PBDRigid->GetX()[2], 0);
 		EXPECT_EQ(PBDRigid->GetV()[0], 0);
 		EXPECT_EQ(PBDRigid->GetV()[1], 0);
 		EXPECT_EQ(PBDRigid->GetV()[2], 0);
 		EXPECT_EQ(PBDRigid->M(), 1);
 
 		PBDRigid->SetX(FVec3(1, 2, 3));
-		EXPECT_EQ(PBDRigid->X()[0], 1);
+		EXPECT_EQ(PBDRigid->GetX()[0], 1);
 		KinematicGeometry->SetV(FVec3(3, 3, 3));
 		EXPECT_EQ(KinematicGeometry->GetV()[0], 3);
 
@@ -301,7 +301,7 @@ namespace ChaosTest {
 
 		TGeometry* DynamicAsStatic = PBDRigid;
 		EXPECT_EQ(DynamicAsStatic->ObjectState(), EObjectStateType::Dynamic);
-		EXPECT_EQ(DynamicAsStatic->X()[0], 1);
+		EXPECT_EQ(DynamicAsStatic->GetX()[0], 1);
 
 		//more polymorphism
 		ParticleHandleTestHelperObjectState(PBDRigid);
@@ -333,7 +333,7 @@ namespace ChaosTest {
 				//fake step and write to physics side
 				for (auto& Particle : SOAs.GetAllParticlesView())
 				{
-					Particle.X() = FVec3(Count);
+					Particle.SetX(FVec3(Count));
 					Count += 1;
 				}
 			}
@@ -342,7 +342,7 @@ namespace ChaosTest {
 			{
 				for (const auto& Particle : SOAs.GetAllParticlesView())
 				{
-					Particle.GTGeometryParticle()->SetX(Particle.X());
+					Particle.GTGeometryParticle()->SetX(Particle.GetX());
 				}
 			}
 
@@ -373,7 +373,7 @@ namespace ChaosTest {
 
 			for (auto& Particle : SOAs.GetAllParticlesView())
 			{
-				EXPECT_TRUE(Particle.X()[0] != 1);
+				EXPECT_TRUE(Particle.GetX()[0] != 1);
 			}
 		}
 	}
@@ -388,7 +388,7 @@ namespace ChaosTest {
 			FGeometryParticleHandle* ThirdParticle = nullptr;
 			for (auto& Particle : SOAs.GetAllParticlesView())
 			{
-				Particle.X() = FVec3(Count);
+				Particle.SetX(FVec3(Count));
 				if (Count == 2)
 				{
 					ThirdParticle = Particle.Handle();
@@ -396,7 +396,7 @@ namespace ChaosTest {
 
 				Count += 1;
 			}
-			EXPECT_EQ(ThirdParticle->X()[0], 2);
+			EXPECT_EQ(ThirdParticle->GetX()[0], 2);
 
 			SOAs.DestroyParticle(ThirdParticle);
 			//default behavior is swap dynamics at end
@@ -405,11 +405,11 @@ namespace ChaosTest {
 			{
 				if (Count == 2)
 				{
-					EXPECT_EQ(Particle.X()[0], 9);
+					EXPECT_EQ(Particle.GetX()[0], 9);
 				}
 				else
 				{
-					EXPECT_EQ(Particle.X()[0], Count);
+					EXPECT_EQ(Particle.GetX()[0], Count);
 				}
 
 				Count += 1;
@@ -425,7 +425,7 @@ namespace ChaosTest {
 			FGeometryParticleHandle* ThirdParticle = nullptr;
 			for (auto& Particle : SOAs.GetAllParticlesView())
 			{
-				Particle.X() = FVec3(Count);
+				Particle.SetX(FVec3(Count));
 				if (Count == 2)
 				{
 					ThirdParticle = Particle.Handle();
@@ -433,7 +433,7 @@ namespace ChaosTest {
 
 				Count += 1;
 			}
-			EXPECT_EQ(ThirdParticle->X()[0], 2);
+			EXPECT_EQ(ThirdParticle->GetX()[0], 2);
 
 			/*
 			//For now we're just disabling removing clustered all together
@@ -506,7 +506,7 @@ namespace ChaosTest {
 				for (auto Itr = MakeParticleIterator(SOAViews); Itr; ++Itr)
 				{
 					//set X back to 0 for all particles
-					Itr->X() = FVec3(0);
+					Itr->SetX(FVec3(0));
 					EXPECT_EQ(Itr->Handle(), GeomHandles[Count]);
 					//implicit const
 					TConstParticleIterator<FGeometryParticles>& ConstItr = Itr;
@@ -517,7 +517,7 @@ namespace ChaosTest {
 				for (auto Itr = MakeConstParticleIterator(SOAViews); Itr; ++Itr)
 				{
 					//check Xs are back to 0
-					EXPECT_EQ(Itr->X()[0], 0);
+					EXPECT_EQ(Itr->GetX()[0], 0);
 				}
 
 				Count = 0;
@@ -542,9 +542,9 @@ namespace ChaosTest {
 				for (auto Itr = MakeParticleIterator(MoveTemp(SOAViews)); Itr; ++Itr)
 				{
 					//set P to 1,1,1
-					Itr->P() = FVec3(1);
+					Itr->SetP(FVec3(1));
 					EXPECT_EQ(Itr->Handle(), PBDRigidHandles[Count++]);
-					EXPECT_EQ(Itr->Handle()->P()[0], Itr->P()[0]);	//handle type is deduced from iterator type
+					EXPECT_EQ(Itr->Handle()->GetP()[0], Itr->GetP()[0]);	//handle type is deduced from iterator type
 				}
 				EXPECT_EQ(Count, 1);
 			}
@@ -556,8 +556,8 @@ namespace ChaosTest {
 				auto PartialDynamics = SOAsWithHandles.CreateDynamicParticles(10);
 
 				TArray<FPBDRigidParticleHandle*> ActiveParticles = { PartialDynamics[3], PartialDynamics[5] };
-				PartialDynamics[3]->X() = FVec3(3);
-				PartialDynamics[5]->X() = FVec3(5);
+				PartialDynamics[3]->SetX(FVec3(3));
+				PartialDynamics[5]->SetX(FVec3(5));
 				
 				TArray<TSOAView<FPBDRigidParticles>> SOAViews = { PBDRigidParticles.Get(), &ActiveParticles, PBDRigidParticles.Get() };
 				int32 Count = 0;
@@ -565,12 +565,12 @@ namespace ChaosTest {
 				{
 					if (Count == 1)
 					{
-						EXPECT_EQ(Itr->X()[0], 3);
+						EXPECT_EQ(Itr->GetX()[0], 3);
 					}
 
 					if (Count == 2)
 					{
-						EXPECT_EQ(Itr->X()[0], 5);
+						EXPECT_EQ(Itr->GetX()[0], 5);
 					}
 					++Count;
 				}
@@ -600,7 +600,7 @@ namespace ChaosTest {
 			FReal Count = 0;
 			for (auto& Kinematic : KinematicParticles)
 			{
-				Kinematic->X() = FVec3(Count);
+				Kinematic->SetX(FVec3(Count));
 				SOAs.DisableParticle(Kinematic);
 				Count += 1;
 			}
@@ -608,9 +608,9 @@ namespace ChaosTest {
 			EXPECT_EQ(SOAs.GetNonDisabledView().Num(), 6);
 
 			//values are still set
-			EXPECT_EQ(KinematicParticles[0]->X()[0], 0);
-			EXPECT_EQ(KinematicParticles[1]->X()[0], 1);
-			EXPECT_EQ(KinematicParticles[2]->X()[0], 2);
+			EXPECT_EQ(KinematicParticles[0]->GetX()[0], 0);
+			EXPECT_EQ(KinematicParticles[1]->GetX()[0], 1);
+			EXPECT_EQ(KinematicParticles[2]->GetX()[0], 2);
 
 			//move to enabled
 			for (auto& Kinematic : KinematicParticles)
diff --git a/Engine/Source/Programs/HeadlessChaos/Private/HeadlessChaosTestPerf.cpp b/Engine/Source/Programs/HeadlessChaos/Private/HeadlessChaosTestPerf.cpp
--- a/Engine/Source/Programs/HeadlessChaos/Private/HeadlessChaosTestPerf.cpp
+++ b/Engine/Source/Programs/HeadlessChaos/Private/HeadlessChaosTestPerf.cpp
@@ -26,7 +26,7 @@ namespace ChaosTest {
 	Chaos::FGeometryParticleHandle* AddFloor(TEvolution& Evolution)
 	{
 		auto Static = Evolution.CreateStaticParticles(1)[0];
-		Static->X() = FVec3(0, 0, 0);
+		Static->SetX(FVec3(0, 0, 0));
 		Static->SetGeometry(MakeImplicitObjectPtr<TPlane<FReal, 3>>(FVec3(0, 0, 0), FVec3(0, 0, 1)));
 		return Static;
 	}
@@ -86,7 +86,7 @@ namespace ChaosTest {
 			{
 				auto BoxDynamic = DynamicParticles[DynamicParticleIndex++];
 				BoxDynamic->SetGeometry(Box);
-				BoxDynamic->X() = Offset + FVec3(x * Interval, y * Interval, Height);
+				BoxDynamic->SetX(Offset + FVec3(x * Interval, y * Interval, Height));
 				BoxDynamic->I() = TVec3<FRealSingle>(100000.);
 				BoxDynamic->InvI() = TVec3<FRealSingle>(1. / 100000.);
 				Evolution.SetPhysicsMaterial(BoxDynamic, Material);
@@ -94,7 +94,7 @@ namespace ChaosTest {
 
 				auto SphereDynamic = DynamicParticles[DynamicParticleIndex++];
 				SphereDynamic->SetGeometry(Sphere);
-				SphereDynamic->X() = Offset + FVec3(x * Interval, y * Interval, Height * (FReal)0.5);
+				SphereDynamic->SetX(Offset + FVec3(x * Interval, y * Interval, Height * (FReal)0.5));
 				SphereDynamic->I() = TVec3<FRealSingle>(100000.);
 				SphereDynamic->InvI() = TVec3<FRealSingle>(1. / 100000.);
 				Evolution.SetPhysicsMaterial(SphereDynamic, Material);
diff --git a/Engine/Source/Programs/HeadlessChaos/Private/HeadlessChaosTestPhysicsScene.cpp b/Engine/Source/Programs/HeadlessChaos/Private/HeadlessChaosTestPhysicsScene.cpp
--- a/Engine/Source/Programs/HeadlessChaos/Private/HeadlessChaosTestPhysicsScene.cpp
+++ b/Engine/Source/Programs/HeadlessChaos/Private/HeadlessChaosTestPhysicsScene.cpp
@@ -82,14 +82,14 @@ namespace ChaosTest {
 			{
 				if (bQueryGameThread)
 				{
-					FTransform ParticleTransform(Data.Payload.GetExternalGeometryParticle_ExternalThread()->R(), Data.Payload.GetExternalGeometryParticle_ExternalThread()->X());
+					FTransform ParticleTransform(Data.Payload.GetExternalGeometryParticle_ExternalThread()->R(), Data.Payload.GetExternalGeometryParticle_ExternalThread()->GetX());
 					const FVec3 DirLocal = ParticleTransform.InverseTransformVectorNoScale(CurData.Dir);
 					const FVec3 StartLocal = ParticleTransform.InverseTransformPositionNoScale(Start);
 					bHit = Data.Payload.GetExternalGeometryParticle_ExternalThread()->GetGeometry()->Raycast(StartLocal, DirLocal, CurData.CurrentLength, 0, OutTime, OutPos, OutNorm, FaceIdx);
 				}
 				else
 				{
-					FTransform ParticleTransform(Data.Payload.GetGeometryParticleHandle_PhysicsThread()->GetR(), Data.Payload.GetGeometryParticleHandle_PhysicsThread()->X());
+					FTransform ParticleTransform(Data.Payload.GetGeometryParticleHandle_PhysicsThread()->GetR(), Data.Payload.GetGeometryParticleHandle_PhysicsThread()->GetX());
 					const FVec3 DirLocal = ParticleTransform.InverseTransformVectorNoScale(CurData.Dir);
 					const FVec3 StartLocal = ParticleTransform.InverseTransformPositionNoScale(Start);
 					bHit = Data.Payload.GetGeometryParticleHandle_PhysicsThread()->GetGeometry()->Raycast(StartLocal, DirLocal, CurData.CurrentLength, 0, OutTime, OutPos, OutNorm, FaceIdx);
@@ -948,7 +948,7 @@ namespace ChaosTest {
 					FPBDRigidsEvolution* Evolution = Scene.GetSolver()->GetEvolution();
 					const auto& SOA = Evolution->GetParticles();
 					const auto& InternalProxy = *SOA.GetAllParticlesView().Begin();
-					EXPECT_EQ(InternalProxy.X()[0], 0);
+					EXPECT_EQ(InternalProxy.GetX()[0], 0);
 				}
 			}
 
@@ -965,7 +965,7 @@ namespace ChaosTest {
 				FPBDRigidsEvolution* Evolution = Scene.GetSolver()->GetEvolution();
 				const auto& SOA = Evolution->GetParticles();
 				const auto& InternalProxy = *SOA.GetAllParticlesView().Begin();
-				EXPECT_EQ(InternalProxy.X()[0], 5);
+				EXPECT_EQ(InternalProxy.GetX()[0], 5);
 			}
 
 			//make sure commands are also deferred
diff --git a/Engine/Source/Programs/HeadlessChaos/Private/HeadlessChaosTestProbes.cpp b/Engine/Source/Programs/HeadlessChaos/Private/HeadlessChaosTestProbes.cpp
--- a/Engine/Source/Programs/HeadlessChaos/Private/HeadlessChaosTestProbes.cpp
+++ b/Engine/Source/Programs/HeadlessChaos/Private/HeadlessChaosTestProbes.cpp
@@ -43,13 +43,13 @@ namespace ChaosTest
 		Dynamic->ShapesArray()[0]->SetIsProbe(true);
 
 		// Positions
-		Static->X() = FVec3(0, 0, 0);
-		Dynamic->X() = FVec3(-(2 * BoxHalfSize) - Separation, 0, 0);
+		Static->SetX( FVec3(0, 0, 0));
+		Dynamic->SetX(FVec3(-(2 * BoxHalfSize) - Separation, 0, 0));
 		Dynamic->SetV(FVec3(InitialSpeed, 0, 0));
 		Dynamic->SetCCDEnabled(false);
 
 		// The position of the static has changed and statics don't automatically update bounds, so update explicitly
-		Static->UpdateWorldSpaceState(TRigidTransform<FReal, 3>(Static->X(), Static->GetR()), FVec3(0));
+		Static->UpdateWorldSpaceState(TRigidTransform<FReal, 3>(Static->GetX(), Static->GetR()), FVec3(0));
 
 		// Make sure the particles would collide if Dynamic wasn't a probe
 		::ChaosTest::SetParticleSimDataToCollide({ Static,Dynamic });
@@ -94,13 +94,13 @@ namespace ChaosTest
 		Dynamic->ShapesArray()[0]->SetIsProbe(true);
 
 		// Positions
-		Static->X() = FVec3(0, 0, 0);
-		Dynamic->X() = FVec3(-(2 * BoxHalfSize) - Separation, 0, 0);
+		Static->SetX(FVec3(0, 0, 0));
+		Dynamic->SetX(FVec3(-(2 * BoxHalfSize) - Separation, 0, 0));
 		Dynamic->SetV(FVec3(InitialSpeed, 0, 0));
 		Dynamic->SetCCDEnabled(true);
 
 		// The position of the static has changed and statics don't automatically update bounds, so update explicitly
-		Static->UpdateWorldSpaceState(TRigidTransform<FReal, 3>(Static->X(), Static->GetR()), FVec3(0));
+		Static->UpdateWorldSpaceState(TRigidTransform<FReal, 3>(Static->GetX(), Static->GetR()), FVec3(0));
 
 		// Make sure the particles would collide if Dynamic wasn't a probe
 		::ChaosTest::SetParticleSimDataToCollide({ Static,Dynamic });
diff --git a/Engine/Source/Programs/HeadlessChaos/Private/HeadlessChaosTestRaycast.cpp b/Engine/Source/Programs/HeadlessChaos/Private/HeadlessChaosTestRaycast.cpp
--- a/Engine/Source/Programs/HeadlessChaos/Private/HeadlessChaosTestRaycast.cpp
+++ b/Engine/Source/Programs/HeadlessChaos/Private/HeadlessChaosTestRaycast.cpp
@@ -350,9 +350,9 @@ namespace ChaosTest
 
 		FTriangleMeshImplicitObject::ParticlesType Particles;
 		Particles.AddParticles(3);
-		Particles.X(0) = FVec3(1, 1, 1);
-		Particles.X(1) = FVec3(5, 1, 1);
-		Particles.X(2) = FVec3(1, 5, 1);
+		Particles.SetX(0, FVec3(1, 1, 1));
+		Particles.SetX(1, FVec3(5, 1, 1));
+		Particles.SetX(2, FVec3(1, 5, 1));
 		TArray<TVec3<int32>> Indices;
 		Indices.Emplace(0, 1, 2);
 		FTriangleMeshImplicitObject Tri(MoveTemp(Particles), MoveTemp(Indices), MoveTemp(DummyMaterials));
@@ -441,9 +441,9 @@ namespace ChaosTest
 
 		FTriangleMeshImplicitObject::ParticlesType Particles;
 		Particles.AddParticles(3);
-		Particles.X(0) = FVec3(1, 1, 1);
-		Particles.X(1) = FVec3(1, 1, 2);
-		Particles.X(2) = FVec3(1, 1, 3);
+		Particles.SetX(0, FVec3(1, 1, 1));
+		Particles.SetX(1, FVec3(1, 1, 2));
+		Particles.SetX(2, FVec3(1, 1, 3));
 		TArray<TVec3<int32>> Indices;
 		Indices.Emplace(0, 1, 2);
 		FTriangleMeshImplicitObject Tri(MoveTemp(Particles), MoveTemp(Indices), MoveTemp(DummyMaterials));
diff --git a/Engine/Source/Programs/HeadlessChaos/Private/HeadlessChaosTestRewind.cpp b/Engine/Source/Programs/HeadlessChaos/Private/HeadlessChaosTestRewind.cpp
--- a/Engine/Source/Programs/HeadlessChaos/Private/HeadlessChaosTestRewind.cpp
+++ b/Engine/Source/Programs/HeadlessChaos/Private/HeadlessChaosTestRewind.cpp
@@ -150,7 +150,7 @@ namespace ChaosTest {
 				const FReal LastInputTime = SimDt <= 1 ? TimeStart : TimeEnd - 1;	//latest gt time associated with this interval
 
 				const auto ParticleState = RewindData->GetPastStateAtFrame(*Proxy->GetHandle_LowLevel(), SimStep);
-				EXPECT_EQ(ParticleState.X()[2], 100 - FMath::FloorToInt(LastInputTime));	//We teleported on GT so no interpolation
+				EXPECT_EQ(ParticleState.GetX()[2], 100 - FMath::FloorToInt(LastInputTime));	//We teleported on GT so no interpolation
 
 				if (LastInputTime < 3)
 				{
@@ -1720,7 +1720,7 @@ namespace ChaosTest {
 					ExpectedXZ = 10;
 				}
 
-				EXPECT_NEAR(ParticleState.X()[2], ExpectedXZ, 1e-4);
+				EXPECT_NEAR(ParticleState.GetX()[2], ExpectedXZ, 1e-4);
 				EXPECT_NEAR(ParticleState.GetV()[2], ExpectedVZ, 1e-4);
 
 				ExpectedVZ -= SimDt;
@@ -3871,7 +3871,7 @@ namespace ChaosTest {
 
 			for (const auto& Dynamic : NonDisabledDyanmic)
 			{
-				Frame.X.Add(Dynamic.X());
+				Frame.X.Add(Dynamic.GetX());
 				Frame.R.Add(Dynamic.GetR());
 			}
 			History.Add(MoveTemp(Frame));
diff --git a/Engine/Source/Programs/HeadlessChaos/Private/HeadlessChaosTestRotation.cpp b/Engine/Source/Programs/HeadlessChaos/Private/HeadlessChaosTestRotation.cpp
--- a/Engine/Source/Programs/HeadlessChaos/Private/HeadlessChaosTestRotation.cpp
+++ b/Engine/Source/Programs/HeadlessChaos/Private/HeadlessChaosTestRotation.cpp
@@ -156,8 +156,8 @@ namespace ChaosTest {
 
 			AppendAnalyticBox(Particles, Dims);
 
-			Particles.X(ParticleIndex) = FVec3::ZeroVector;
-			Particles.P(ParticleIndex) = FVec3::ZeroVector;
+			Particles.SetX(ParticleIndex, FVec3::ZeroVector);
+			Particles.SetP(ParticleIndex, FVec3::ZeroVector);
 			Particles.SetR(ParticleIndex, InitialRotations[ParticleIndex]);
 			Particles.SetQ(ParticleIndex, FinalRotations[ParticleIndex]);
 			Particles.SetV(ParticleIndex, FVec3::ZeroVector);
diff --git a/Engine/Source/Programs/HeadlessChaos/Private/HeadlessChaosTestSerialization.cpp b/Engine/Source/Programs/HeadlessChaos/Private/HeadlessChaosTestSerialization.cpp
--- a/Engine/Source/Programs/HeadlessChaos/Private/HeadlessChaosTestSerialization.cpp
+++ b/Engine/Source/Programs/HeadlessChaos/Private/HeadlessChaosTestSerialization.cpp
@@ -338,8 +338,8 @@ namespace ChaosTest
 			OriginalParticles.SetR(1, FRotation3::Identity);
 			OriginalParticles.SetGeometry(0, OriginalSpheres[0]);
 			OriginalParticles.SetGeometry(1, OriginalSpheres[1]);
-			OriginalParticles.X(0) = FVec3(100, 1, 2);
-			OriginalParticles.X(1) = FVec3(0, 1, 2);
+			OriginalParticles.SetX(0, FVec3(100, 1, 2));
+			OriginalParticles.SetX(1, FVec3(0, 1, 2));
 			OriginalParticles.SetR(0, FRotation3::Identity);
 			OriginalParticles.SetR(1, FRotation3::Identity);
 
@@ -369,7 +369,7 @@ namespace ChaosTest
 			TArray<int32> FinalIntersections;
 			for (int32 Potential : PotentialIntersections)
 			{
-				FRigidTransform3 TM(SerializedParticles.X(Potential), SerializedParticles.GetR(Potential));
+				FRigidTransform3 TM(SerializedParticles.GetX(Potential), SerializedParticles.GetR(Potential));
 				const FAABB3 Bounds = SerializedParticles.GetGeometry(Potential)->BoundingBox().TransformedAABB(TM);
 				if (Bounds.Intersects(QueryBox))
 				{
@@ -399,8 +399,8 @@ namespace ChaosTest
 
 		Particles.Acceleration(0) = F[0];
 		Particles.Acceleration(1) = F[1];
-		Particles.X(0) = X[0];
-		Particles.X(1) = X[1];
+		Particles.SetX(0, X[0]);
+		Particles.SetX(1, X[1]);
 		Particles.RotationOfMass(0) = FRotation3::FromIdentity();
 		Particles.RotationOfMass(1) = FRotation3::FromIdentity();
 
@@ -415,8 +415,8 @@ namespace ChaosTest
 			EXPECT_EQ(TestParticles.Size(), Particles.Size());
 			EXPECT_EQ(TestParticles.Acceleration(0), Particles.Acceleration(0));
 			EXPECT_EQ(TestParticles.Acceleration(1), Particles.Acceleration(1));
-			EXPECT_EQ(TestParticles.X(0), Particles.X(0));
-			EXPECT_EQ(TestParticles.X(1), Particles.X(1));
+			EXPECT_EQ(TestParticles.GetX(0), Particles.GetX(0));
+			EXPECT_EQ(TestParticles.GetX(1), Particles.GetX(1));
 		}
 	}
 
@@ -432,9 +432,9 @@ namespace ChaosTest
 		Particles.AddParticles(3);
 		Particles.SetR(0, FRotation3::Identity);
 		Particles.SetR(1, FRotation3::Identity);
-		Particles.X(0) = FVec3(15, 1, 2);
-		Particles.X(1) = FVec3(0, 2, 2);
-		Particles.X(2) = FVec3(0, 2, 2);
+		Particles.SetX(0, FVec3(15, 1, 2));
+		Particles.SetX(1, FVec3(0, 2, 2));
+		Particles.SetX(2, FVec3(0, 2, 2));
 		Particles.SetR(0, FRotation3::Identity);
 		Particles.SetR(1, FRotation3::Identity);
 		Particles.SetR(2, FRotation3::Identity);
diff --git a/Engine/Source/Programs/HeadlessChaos/Private/HeadlessChaosTestShapePadding.cpp b/Engine/Source/Programs/HeadlessChaos/Private/HeadlessChaosTestShapePadding.cpp
--- a/Engine/Source/Programs/HeadlessChaos/Private/HeadlessChaosTestShapePadding.cpp
+++ b/Engine/Source/Programs/HeadlessChaos/Private/HeadlessChaosTestShapePadding.cpp
@@ -47,20 +47,20 @@ namespace ChaosTest {
 		Particles.GetParticleHandles().AddArray(&PerParticlePhysicsMaterials);
 
 		auto Box0 = AppendDynamicParticleBoxMargin(Particles, Size, Margin0);
-		Box0->X() = FVec3(0, 0, 0);
+		Box0->SetX(FVec3(0, 0, 0));
 		Box0->SetR(FRotation3(FQuat::Identity));
 		Box0->SetV(FVec3(0));
 		Box0->SetPreV(Box0->GetV());
-		Box0->P() = Box0->X();
+		Box0->SetP(Box0->GetX());
 		Box0->SetQ(Box0->GetR());
 		Box0->AuxilaryValue(PhysicsMaterials) = MakeSerializable(PhysicsMaterial);
 
 		auto Box1 = AppendDynamicParticleBoxMargin(Particles, Size, Margin1);
-		Box1->X() = Delta;
+		Box1->SetX(Delta);
 		Box1->SetR(FRotation3(FQuat::Identity));
 		Box1->SetV(FVec3(0));
 		Box1->SetPreV(Box1->GetV());
-		Box1->P() = Box1->X();
+		Box1->SetP(Box1->GetX());
 		Box1->SetQ(Box1->GetR());
 		Box1->AuxilaryValue(PhysicsMaterials) = MakeSerializable(PhysicsMaterial);
 
@@ -69,7 +69,7 @@ namespace ChaosTest {
 
 		const FReal Tolerance = 2.0f * KINDA_SMALL_NUMBER;
 
-		// Boxes should have a margin, limited by the minumum box dimension
+		// Boxes should have a margin, limited by the minimum box dimension
 		float ExpectedMargin0 = BoxImplicit0->ClampedMargin(Margin0);
 		float ExpectedMargin1 = BoxImplicit0->ClampedMargin(Margin1);
 		EXPECT_NEAR(BoxImplicit0->GetMargin(), ExpectedMargin0, Tolerance);
@@ -166,20 +166,20 @@ namespace ChaosTest {
 		Particles.GetParticleHandles().AddArray(&PerParticlePhysicsMaterials);
 
 		auto Box0 = AppendDynamicParticleConvexBoxMargin(Particles, 0.5f * Size, Margin0);
-		Box0->X() = FVec3(0, 0, 0);
+		Box0->SetX(FVec3(0, 0, 0));
 		Box0->SetR(FRotation3(FQuat::Identity));
 		Box0->SetV(FVec3(0));
 		Box0->SetPreV(Box0->GetV());
-		Box0->P() = Box0->X();
+		Box0->SetP(Box0->GetX());
 		Box0->SetQ(Box0->GetR());
 		Box0->AuxilaryValue(PhysicsMaterials) = MakeSerializable(PhysicsMaterial);
 
 		auto Box1 = AppendDynamicParticleConvexBoxMargin(Particles, 0.5f * Size, Margin1);
-		Box1->X() = Delta;
+		Box1->SetX(Delta);
 		Box1->SetR(FRotation3(FQuat::Identity));
 		Box1->SetV(FVec3(0));
 		Box1->SetPreV(Box1->GetV());
-		Box1->P() = Box1->X();
+		Box1->SetP(Box1->GetX());
 		Box1->SetQ(Box1->GetR());
 		Box1->AuxilaryValue(PhysicsMaterials) = MakeSerializable(PhysicsMaterial);
 
@@ -258,7 +258,7 @@ namespace ChaosTest {
 	TEST(CollisionTests, DISABLED_TestConvexConvexCollisionMarginTooLarge)
 	{
 		// If the margin is too large, the margin should be limited
-		// @todo(chaos): fix this for convex - we do not have margin limits implemeted
+		// @todo(chaos): fix this for convex - we do not have margin limits implemented
 		TestConvexConvexCollisionMargin(15, 15, FVec3(20, 100, 100), FVec3(0, -100, 0), 0.0f, FVec3(0, 1, 0));
 		TestConvexConvexCollisionMargin(15, 15, FVec3(20, 100, 100), FVec3(20, 0, 0), 0.0f, FVec3(-1, 0, 0));
 	}
@@ -290,11 +290,11 @@ namespace ChaosTest {
 		Particles.GetParticleHandles().AddArray(&PerParticlePhysicsMaterials);
 
 		auto Box0 = AppendDynamicParticleBoxMargin(Particles, Size, Margin0);
-		Box0->X() = FVec3(0, 0, 0);
+		Box0->SetX(FVec3(0, 0, 0));
 		Box0->SetR(FRotation3(FQuat::Identity));
 		Box0->SetV(FVec3(0));
 		Box0->SetPreV(Box0->GetV());
-		Box0->P() = Box0->X();
+		Box0->SetP(Box0->GetX());
 		Box0->SetQ(Box0->GetR());
 		Box0->AuxilaryValue(PhysicsMaterials) = MakeSerializable(PhysicsMaterial);
 
diff --git a/Engine/Source/Programs/HeadlessChaos/Private/HeadlessChaosTestSim.cpp b/Engine/Source/Programs/HeadlessChaos/Private/HeadlessChaosTestSim.cpp
--- a/Engine/Source/Programs/HeadlessChaos/Private/HeadlessChaosTestSim.cpp
+++ b/Engine/Source/Programs/HeadlessChaos/Private/HeadlessChaosTestSim.cpp
@@ -79,13 +79,13 @@ namespace ChaosTest {
 
 		Evolution.SetPhysicsMaterial(Dynamic, MakeSerializable(PhysicsMaterial));
 
-		Static->X() = FVec3(10, 10, 10);
-		Dynamic->X() = FVec3(10, 10, 150);
+		Static->SetX(FVec3(10, 10, 10));
+		Dynamic->SetX(FVec3(10, 10, 150));
 		Dynamic->I() = TVec3<FRealSingle>(100000.0f);
 		Dynamic->InvI() = TVec3<FRealSingle>(1.0f / 100000.0f);
 
 		// The position of the static has changed and statics don't automatically update bounds, so update explicitly
-		Static->UpdateWorldSpaceState(FRigidTransform3(Static->X(), Static->GetR()), FVec3(0));
+		Static->UpdateWorldSpaceState(FRigidTransform3(Static->GetX(), Static->GetR()), FVec3(0));
 
 		::ChaosTest::SetParticleSimDataToCollide({ Static,Dynamic });
 
@@ -100,7 +100,7 @@ namespace ChaosTest {
 			Evolution.EndFrame(Dt);
 		}
 
-		EXPECT_NEAR(Dynamic->X().Z, 110, 1);
+		EXPECT_NEAR(Dynamic->GetX().Z, 110, 1);
 	}
 
 	GTEST_TEST(AllEvolutions, SimTests_BoxBoxSimTest)
@@ -119,9 +119,9 @@ namespace ChaosTest {
 		Static->SetGeometry(StaticBox);
 		Dynamic->SetGeometry(DynamicBox);
 
-		Static->X() = FVec3(10, 10, 10);
-		Static->UpdateWorldSpaceState(FRigidTransform3(Static->X(), Static->GetR()), FVec3(0));
-		Dynamic->X() = FVec3(10, 10, 120);
+		Static->SetX(FVec3(10, 10, 10));
+		Static->UpdateWorldSpaceState(FRigidTransform3(Static->GetX(), Static->GetR()), FVec3(0));
+		Dynamic->SetX(FVec3(10, 10, 120));
 		Dynamic->I() = TVec3<FRealSingle>(100000.0f);
 		Dynamic->InvI() = TVec3<FRealSingle>(1.0f / 100000.0f);
 
@@ -138,7 +138,7 @@ namespace ChaosTest {
 			Evolution.EndFrame(Dt);
 		}
 
-		EXPECT_NEAR(Dynamic->X().Z, 110, 5);
+		EXPECT_NEAR(Dynamic->GetX().Z, 110, 5);
 	}
 
 	// This test will fail because the inertia of the dynamic box is very low. The mass and inertia are both 1.0, 
@@ -165,9 +165,9 @@ namespace ChaosTest {
 		Static->SetGeometry(StaticBox);
 		Dynamic->SetGeometry(DynamicBox);
 
-		Static->X() = FVec3(10, 10, 10);
-		Static->UpdateWorldSpaceState(FRigidTransform3(Static->X(), Static->GetR()), FVec3(0));
-		Dynamic->X() = FVec3(10, 10, 300);
+		Static->SetX(FVec3(10, 10, 10));
+		Static->UpdateWorldSpaceState(FRigidTransform3(Static->GetX(), Static->GetR()), FVec3(0));
+		Dynamic->SetX(FVec3(10, 10, 300));
 		Dynamic->I() = TVec3<FRealSingle>(1);
 		Dynamic->InvI() = TVec3<FRealSingle>(1);
 
@@ -182,7 +182,7 @@ namespace ChaosTest {
 			Evolution.EndFrame(1 / 60.f);
 		}
 
-		EXPECT_NEAR(Dynamic->X().Z, 110, 10);
+		EXPECT_NEAR(Dynamic->GetX().Z, 110, 10);
 	}
 
 	GTEST_TEST(AllEvolutions, SimTests_SleepAndWakeSimTest)
@@ -212,10 +212,10 @@ namespace ChaosTest {
 		Dynamic1->SetGeometry(DynamicBox);
 		Dynamic2->SetGeometry(DynamicBox);
 
-		Static->X() = FVec3(10, 10, 10);
-		Static->UpdateWorldSpaceState(FRigidTransform3(Static->X(), Static->GetR()), FVec3(0));
-		Dynamic1->X() = FVec3(10, 10, 120);
-		Dynamic2->X() = FVec3(10, 10, 400);
+		Static->SetX(FVec3(10, 10, 10));
+		Static->UpdateWorldSpaceState(FRigidTransform3(Static->GetX(), Static->GetR()), FVec3(0));
+		Dynamic1->SetX(FVec3(10, 10, 120));
+		Dynamic2->SetX(FVec3(10, 10, 400));
 
 		Evolution.SetPhysicsMaterial(Dynamic1, MakeSerializable(PhysicsMaterial));
 		Evolution.SetPhysicsMaterial(Dynamic2, MakeSerializable(PhysicsMaterial));
@@ -244,8 +244,8 @@ namespace ChaosTest {
 			{
 				Dynamic1WentToSleep = true;
 
-				EXPECT_LT(Dynamic1->X().Z, 120);
-				EXPECT_GT(Dynamic1->X().Z, 100);
+				EXPECT_LT(Dynamic1->GetX().Z, 120);
+				EXPECT_GT(Dynamic1->GetX().Z, 100);
 			}
 
 			// later the Dynamic2 collides with Dynamic1 waking it up again
@@ -361,8 +361,8 @@ namespace ChaosTest {
 		PhysicsMaterial->SleepingAngularThreshold = 20;
 		PhysicsMaterial->SleepCounterThreshold = 5;
 
-		Dynamic1->X() = FVec3(-200, 0, 0);
-		Dynamic2->X() = FVec3(200, 0, 0);
+		Dynamic1->SetX(FVec3(-200, 0, 0));
+		Dynamic2->SetX(FVec3(200, 0, 0));
 
 		Chaos::FImplicitObjectPtr DynamicBox(new TBox<FReal, 3>(FVec3(-50, -50, -50), FVec3(50, 50, 50)));
 		Dynamic1->SetGeometry(DynamicBox);
diff --git a/Engine/Source/Programs/HeadlessChaos/Private/HeadlessChaosTestTriangleMesh.cpp b/Engine/Source/Programs/HeadlessChaos/Private/HeadlessChaosTestTriangleMesh.cpp
--- a/Engine/Source/Programs/HeadlessChaos/Private/HeadlessChaosTestTriangleMesh.cpp
+++ b/Engine/Source/Programs/HeadlessChaos/Private/HeadlessChaosTestTriangleMesh.cpp
@@ -15,14 +15,14 @@ namespace ChaosTest
 		FTriangleMesh TriMesh;
 		TParticles<double, 3> Particles;
 		Particles.AddParticles(3);
-		Particles.X(0) = FVec3(0, 0, 0);
-		Particles.X(1) = FVec3(1, 0, 0);
-		Particles.X(2) = FVec3(0, 1, 0);
+		Particles.SetX(0, FVec3(0, 0, 0));
+		Particles.SetX(1, FVec3(1, 0, 0));
+		Particles.SetX(2, FVec3(0, 1, 0));
 		TArray<TVec3<int32>> Elements;
 		Elements.Add(TVec3<int32>(0, 1, 2));
 		TriMesh.Init(MoveTemp(Elements));
 
-		TTriangle<double> Tri(Particles.X(0), Particles.X(1), Particles.X(2));
+		TTriangle<double> Tri(Particles.GetX(0), Particles.GetX(1), Particles.GetX(2));
 		Chaos::FVec3 Center = Tri.GetCentroid();
 		Chaos::FVec3 Norm = Tri.GetNormal();
 
@@ -59,65 +59,65 @@ namespace ChaosTest
 
 		// Vertices
 
-		Pt = Particles.X(0);
+		Pt = Particles.GetX(0);
 		Hit = TriMesh.SmoothProject(BVH, Particles.XArray(), PointNormals, Pt, HitTriIdx, Weights);
 		EXPECT_TRUE(Hit);
 		EXPECT_VECTOR_NEAR(Weights, FVec3(1,0,0), 0.01);
 
-		Pt = Particles.X(1);
+		Pt = Particles.GetX(1);
 		Hit = TriMesh.SmoothProject(BVH, Particles.XArray(), PointNormals, Pt, HitTriIdx, Weights);
 		EXPECT_TRUE(Hit);
 		EXPECT_VECTOR_NEAR(Weights, FVec3(0, 1, 0), 0.01);
 
-		Pt = Particles.X(2);
+		Pt = Particles.GetX(2);
 		Hit = TriMesh.SmoothProject(BVH, Particles.XArray(), PointNormals, Pt, HitTriIdx, Weights);
 		EXPECT_TRUE(Hit);
 		EXPECT_VECTOR_NEAR(Weights, FVec3(0, 0, 1), 0.01);
 
 		// Edges
 
-		Pt = (Particles.X(0) + Particles.X(1)) * 0.5;
+		Pt = (Particles.GetX(0) + Particles.GetX(1)) * 0.5;
 		Hit = TriMesh.SmoothProject(BVH, Particles.XArray(), PointNormals, Pt, HitTriIdx, Weights);
 		EXPECT_TRUE(Hit);
 		EXPECT_VECTOR_NEAR(Weights, FVec3(.5, .5, 0), 0.01);
 
-		Pt = (Particles.X(1) + Particles.X(2)) * 0.5;
+		Pt = (Particles.GetX(1) + Particles.GetX(2)) * 0.5;
 		Hit = TriMesh.SmoothProject(BVH, Particles.XArray(), PointNormals, Pt, HitTriIdx, Weights);
 		EXPECT_TRUE(Hit);
 		EXPECT_VECTOR_NEAR(Weights, FVec3(0, .5, .5), 0.01);
 
-		Pt = (Particles.X(0) + Particles.X(2)) * 0.5;
+		Pt = (Particles.GetX(0) + Particles.GetX(2)) * 0.5;
 		Hit = TriMesh.SmoothProject(BVH, Particles.XArray(), PointNormals, Pt, HitTriIdx, Weights);
 		EXPECT_TRUE(Hit);
 		EXPECT_VECTOR_NEAR(Weights, FVec3(.5, 0, .5), 0.01);
 
 		// Off vertices - Expect false
 
-		Pt = Center + (Particles.X(0) - Center) * 1.1;
+		Pt = Center + (Particles.GetX(0) - Center) * 1.1;
 		Hit = TriMesh.SmoothProject(BVH, Particles.XArray(), PointNormals, Pt, HitTriIdx, Weights);
 		EXPECT_FALSE(Hit);
 
-		Pt = Center + (Particles.X(1) - Center) * 1.1;
+		Pt = Center + (Particles.GetX(1) - Center) * 1.1;
 		Hit = TriMesh.SmoothProject(BVH, Particles.XArray(), PointNormals, Pt, HitTriIdx, Weights);
 		EXPECT_FALSE(Hit);
 
-		Pt = Center + (Particles.X(2) - Center) * 1.1;
+		Pt = Center + (Particles.GetX(2) - Center) * 1.1;
 		Hit = TriMesh.SmoothProject(BVH, Particles.XArray(), PointNormals, Pt, HitTriIdx, Weights);
 		EXPECT_FALSE(Hit);
 
 		// Off edges - Expect false
 
-		Pt = (Particles.X(0) + Particles.X(1)) * 0.5;
+		Pt = (Particles.GetX(0) + Particles.GetX(1)) * 0.5;
 		Pt = Center + (Pt - Center) * 1.1;
 		Hit = TriMesh.SmoothProject(BVH, Particles.XArray(), PointNormals, Pt, HitTriIdx, Weights);
 		EXPECT_FALSE(Hit);
 
-		Pt = (Particles.X(1) + Particles.X(2)) * 0.5;
+		Pt = (Particles.GetX(1) + Particles.GetX(2)) * 0.5;
 		Pt = Center + (Pt - Center) * 1.1;
 		Hit = TriMesh.SmoothProject(BVH, Particles.XArray(), PointNormals, Pt, HitTriIdx, Weights);
 		EXPECT_FALSE(Hit);
 
-		Pt = (Particles.X(0) + Particles.X(2)) * 0.5;
+		Pt = (Particles.GetX(0) + Particles.GetX(2)) * 0.5;
 		Pt = Center + (Pt - Center) * 1.1;
 		Hit = TriMesh.SmoothProject(BVH, Particles.XArray(), PointNormals, Pt, HitTriIdx, Weights);
 		EXPECT_FALSE(Hit);
diff --git a/Engine/Source/Programs/HeadlessChaos/Private/HeadlessChaosTestUtility.cpp b/Engine/Source/Programs/HeadlessChaos/Private/HeadlessChaosTestUtility.cpp
--- a/Engine/Source/Programs/HeadlessChaos/Private/HeadlessChaosTestUtility.cpp
+++ b/Engine/Source/Programs/HeadlessChaos/Private/HeadlessChaosTestUtility.cpp
@@ -25,11 +25,11 @@ namespace ChaosTest {
 		InParticles.AddParticles(1);
 		int32 RigidBodyIndex = InParticles.Size() - 1;
 
-		InParticles.X(RigidBodyIndex) = FVec3(0, 0, 0);
+		InParticles.SetX(RigidBodyIndex, FVec3(0, 0, 0));
 		InParticles.SetV(RigidBodyIndex, FVec3(0, 0, 0));
 		InParticles.SetR(RigidBodyIndex, FRotation3::MakeFromEuler(FVec3(0, 0, 0)).GetNormalized());
 		InParticles.SetW(RigidBodyIndex, FVec3(0, 0, 0));
-		InParticles.P(RigidBodyIndex) = InParticles.X(RigidBodyIndex);
+		InParticles.SetP(RigidBodyIndex, InParticles.GetX(RigidBodyIndex));
 		InParticles.SetQ(RigidBodyIndex, InParticles.GetR(RigidBodyIndex));
 
 		InParticles.M(RigidBodyIndex) = 1.0;
@@ -47,11 +47,11 @@ namespace ChaosTest {
 		InParticles.AddParticles(1);
 		int32 RigidBodyIndex = InParticles.Size() - 1;
 
-		InParticles.X(RigidBodyIndex) = FVec3(0, 0, 0);
+		InParticles.SetX(RigidBodyIndex, FVec3(0, 0, 0));
 		InParticles.SetV(RigidBodyIndex, FVec3(0, 0, 0));
 		InParticles.SetR(RigidBodyIndex, FRotation3::MakeFromEuler(FVec3(0, 0, 0)).GetNormalized());
 		InParticles.SetW(RigidBodyIndex, FVec3(0, 0, 0));
-		InParticles.P(RigidBodyIndex) = InParticles.X(RigidBodyIndex);
+		InParticles.SetP(RigidBodyIndex, InParticles.GetX(RigidBodyIndex));
 		InParticles.SetQ(RigidBodyIndex, InParticles.GetR(RigidBodyIndex));
 
 		InParticles.M(RigidBodyIndex) = 1.0;
@@ -66,7 +66,7 @@ namespace ChaosTest {
 
 	void InitAnalyticBox2(FKinematicGeometryParticleHandle* Particle, FVec3 Scale)
 	{
-		Particle->X() = FVec3(0, 0, 0);
+		Particle->SetX(FVec3(0, 0, 0));
 		Particle->SetV(FVec3(0, 0, 0));
 		Particle->SetR(FRotation3::MakeFromEuler(FVec3(0, 0, 0)).GetNormalized());
 		Particle->SetW(FVec3(0, 0, 0));
@@ -75,7 +75,7 @@ namespace ChaosTest {
 		FPBDRigidParticleHandle* DynamicParticle = Particle->CastToRigidParticle();
 		if(DynamicParticle && DynamicParticle->ObjectState() == EObjectStateType::Dynamic)
 		{
-			DynamicParticle->P() = Particle->X();
+			DynamicParticle->SetP(Particle->GetX());
 			DynamicParticle->SetQ(Particle->GetR());
 
 			DynamicParticle->M() = 1.0;
@@ -90,11 +90,11 @@ namespace ChaosTest {
 		InParticles.AddParticles(1);
 		int32 RigidBodyIndex = InParticles.Size() - 1;
 
-		InParticles.X(RigidBodyIndex) = FVec3(0, 0, 0);
+		InParticles.SetX(RigidBodyIndex, FVec3(0, 0, 0));
 		InParticles.SetV(RigidBodyIndex, FVec3(0, 0, 0));
 		InParticles.SetR(RigidBodyIndex, FRotation3::MakeFromEuler(FVec3(0, 0, 0)).GetNormalized());
 		InParticles.SetW(RigidBodyIndex, FVec3(0, 0, 0));
-		InParticles.P(RigidBodyIndex) = InParticles.X(RigidBodyIndex);
+		InParticles.SetP(RigidBodyIndex, InParticles.GetX(RigidBodyIndex));
 		InParticles.SetQ(RigidBodyIndex, InParticles.GetR(RigidBodyIndex));
 
 		check(Scale.X == Scale.Y && Scale.X == Scale.Z);
@@ -109,14 +109,14 @@ namespace ChaosTest {
 		int32 CollisionIndex = 0;
 		InParticles.CollisionParticlesInitIfNeeded(RigidBodyIndex);
 		InParticles.CollisionParticles(RigidBodyIndex)->AddParticles(8);
-		InParticles.CollisionParticles(RigidBodyIndex)->X(CollisionIndex++) = FVec3(-Scale[0] / 2.0, -Scale[1] / 2.0, -Scale[2] / 2.0);
-		InParticles.CollisionParticles(RigidBodyIndex)->X(CollisionIndex++) = FVec3(+Scale[0] / 2.0, -Scale[1] / 2.0, -Scale[2] / 2.0);
-		InParticles.CollisionParticles(RigidBodyIndex)->X(CollisionIndex++) = FVec3(-Scale[0] / 2.0, +Scale[1] / 2.0, -Scale[2] / 2.0);
-		InParticles.CollisionParticles(RigidBodyIndex)->X(CollisionIndex++) = FVec3(+Scale[0] / 2.0, +Scale[1] / 2.0, -Scale[2] / 2.0);
-		InParticles.CollisionParticles(RigidBodyIndex)->X(CollisionIndex++) = FVec3(-Scale[0] / 2.0, -Scale[1] / 2.0, +Scale[2] / 2.0);
-		InParticles.CollisionParticles(RigidBodyIndex)->X(CollisionIndex++) = FVec3(+Scale[0] / 2.0, -Scale[1] / 2.0, +Scale[2] / 2.0);
-		InParticles.CollisionParticles(RigidBodyIndex)->X(CollisionIndex++) = FVec3(-Scale[0] / 2.0, +Scale[1] / 2.0, +Scale[2] / 2.0);
-		InParticles.CollisionParticles(RigidBodyIndex)->X(CollisionIndex++) = FVec3(+Scale[0] / 2.0, +Scale[1] / 2.0, +Scale[2] / 2.0);
+		InParticles.CollisionParticles(RigidBodyIndex)->SetX(CollisionIndex++, FVec3(-Scale[0] / 2.0, -Scale[1] / 2.0, -Scale[2] / 2.0));
+		InParticles.CollisionParticles(RigidBodyIndex)->SetX(CollisionIndex++, FVec3(+Scale[0] / 2.0, -Scale[1] / 2.0, -Scale[2] / 2.0));
+		InParticles.CollisionParticles(RigidBodyIndex)->SetX(CollisionIndex++, FVec3(-Scale[0] / 2.0, +Scale[1] / 2.0, -Scale[2] / 2.0));
+		InParticles.CollisionParticles(RigidBodyIndex)->SetX(CollisionIndex++, FVec3(+Scale[0] / 2.0, +Scale[1] / 2.0, -Scale[2] / 2.0));
+		InParticles.CollisionParticles(RigidBodyIndex)->SetX(CollisionIndex++, FVec3(-Scale[0] / 2.0, -Scale[1] / 2.0, +Scale[2] / 2.0));
+		InParticles.CollisionParticles(RigidBodyIndex)->SetX(CollisionIndex++, FVec3(+Scale[0] / 2.0, -Scale[1] / 2.0, +Scale[2] / 2.0));
+		InParticles.CollisionParticles(RigidBodyIndex)->SetX(CollisionIndex++, FVec3(-Scale[0] / 2.0, +Scale[1] / 2.0, +Scale[2] / 2.0));
+		InParticles.CollisionParticles(RigidBodyIndex)->SetX(CollisionIndex++, FVec3(+Scale[0] / 2.0, +Scale[1] / 2.0, +Scale[2] / 2.0));
 
 		if (elements != nullptr)
 		{
@@ -156,12 +156,12 @@ namespace ChaosTest {
 
 	void InitDynamicParticleBox2(FPBDRigidParticleHandle* Particle, const FVec3& Scale, FReal Margin, TArray<TVector<int32, 3>>* OutElements)
 	{
-		Particle->X() = FVec3(0, 0, 0);
+		Particle->SetX(FVec3(0, 0, 0));
 		Particle->SetV(FVec3(0, 0, 0));
 		Particle->SetR(FRotation3::MakeFromEuler(FVec3(0, 0, 0)).GetNormalized());
 		Particle->SetW(FVec3(0, 0, 0));
 
-		Particle->P() = Particle->X();
+		Particle->SetP(Particle->GetX());
 		Particle->SetQ(Particle->GetR());
 
 		// Assume unit mass - this gets scaled externally by the actual mass
@@ -176,14 +176,14 @@ namespace ChaosTest {
 		int32 CollisionIndex = 0;
 		Particle->CollisionParticlesInitIfNeeded();
 		Particle->CollisionParticles()->AddParticles(8);
-		Particle->CollisionParticles()->X(CollisionIndex++) = FVec3(-Scale[0] / 2.0, -Scale[1] / 2.0, -Scale[2] / 2.0);
-		Particle->CollisionParticles()->X(CollisionIndex++) = FVec3(+Scale[0] / 2.0, -Scale[1] / 2.0, -Scale[2] / 2.0);
-		Particle->CollisionParticles()->X(CollisionIndex++) = FVec3(-Scale[0] / 2.0, +Scale[1] / 2.0, -Scale[2] / 2.0);
-		Particle->CollisionParticles()->X(CollisionIndex++) = FVec3(+Scale[0] / 2.0, +Scale[1] / 2.0, -Scale[2] / 2.0);
-		Particle->CollisionParticles()->X(CollisionIndex++) = FVec3(-Scale[0] / 2.0, -Scale[1] / 2.0, +Scale[2] / 2.0);
-		Particle->CollisionParticles()->X(CollisionIndex++) = FVec3(+Scale[0] / 2.0, -Scale[1] / 2.0, +Scale[2] / 2.0);
-		Particle->CollisionParticles()->X(CollisionIndex++) = FVec3(-Scale[0] / 2.0, +Scale[1] / 2.0, +Scale[2] / 2.0);
-		Particle->CollisionParticles()->X(CollisionIndex++) = FVec3(+Scale[0] / 2.0, +Scale[1] / 2.0, +Scale[2] / 2.0);
+		Particle->CollisionParticles()->SetX(CollisionIndex++, FVec3(-Scale[0] / 2.0, -Scale[1] / 2.0, -Scale[2] / 2.0));
+		Particle->CollisionParticles()->SetX(CollisionIndex++, FVec3(+Scale[0] / 2.0, -Scale[1] / 2.0, -Scale[2] / 2.0));
+		Particle->CollisionParticles()->SetX(CollisionIndex++, FVec3(-Scale[0] / 2.0, +Scale[1] / 2.0, -Scale[2] / 2.0));
+		Particle->CollisionParticles()->SetX(CollisionIndex++, FVec3(+Scale[0] / 2.0, +Scale[1] / 2.0, -Scale[2] / 2.0));
+		Particle->CollisionParticles()->SetX(CollisionIndex++, FVec3(-Scale[0] / 2.0, -Scale[1] / 2.0, +Scale[2] / 2.0));
+		Particle->CollisionParticles()->SetX(CollisionIndex++, FVec3(+Scale[0] / 2.0, -Scale[1] / 2.0, +Scale[2] / 2.0));
+		Particle->CollisionParticles()->SetX(CollisionIndex++, FVec3(-Scale[0] / 2.0, +Scale[1] / 2.0, +Scale[2] / 2.0));
+		Particle->CollisionParticles()->SetX(CollisionIndex++, FVec3(+Scale[0] / 2.0, +Scale[1] / 2.0, +Scale[2] / 2.0));
 
 		Particle->SetLocalBounds(Particle->GetGeometry()->BoundingBox());
 		Particle->UpdateWorldSpaceState(FRigidTransform3::Identity, FVec3(0));
@@ -227,12 +227,12 @@ namespace ChaosTest {
 
 	void InitDynamicParticleSphere2(FPBDRigidParticleHandle* Particle, const FVec3& Scale, TArray<TVec3<int32>>* OutElements) 
 	{
-		Particle->X() = FVec3(0, 0, 0);
+		Particle->SetX(FVec3(0, 0, 0));
 		Particle->SetV(FVec3(0, 0, 0));
 		Particle->SetR(FRotation3::MakeFromEuler(FVec3(0, 0, 0)).GetNormalized());
 		Particle->SetW(FVec3(0, 0, 0));
 
-		Particle->P() = Particle->X();
+		Particle->SetP(Particle->GetX());
 		Particle->SetQ(Particle->GetR());
 
 		check(Scale.X == Scale.Y && Scale.X == Scale.Z);
@@ -248,12 +248,12 @@ namespace ChaosTest {
 		Particle->CollisionParticlesInitIfNeeded();
 		Particle->CollisionParticles()->AddParticles(6);
 
-		Particle->CollisionParticles()->X(CollisionIndex++) = FVec3(-Scale[0] / 2.0, 0, 0);
-		Particle->CollisionParticles()->X(CollisionIndex++) = FVec3(+Scale[0] / 2.0, 0, 0);
-		Particle->CollisionParticles()->X(CollisionIndex++) = FVec3(0, -Scale[1] / 2.0, 0);
-		Particle->CollisionParticles()->X(CollisionIndex++) = FVec3(0, +Scale[1] / 2.0, 0);
-		Particle->CollisionParticles()->X(CollisionIndex++) = FVec3(0, 0, -Scale[2] / 2.0);
-		Particle->CollisionParticles()->X(CollisionIndex++) = FVec3(0, 0, +Scale[2] / 2.0);
+		Particle->CollisionParticles()->SetX(CollisionIndex++, FVec3(-Scale[0] / 2.0, 0, 0));
+		Particle->CollisionParticles()->SetX(CollisionIndex++, FVec3(+Scale[0] / 2.0, 0, 0));
+		Particle->CollisionParticles()->SetX(CollisionIndex++, FVec3(0, -Scale[1] / 2.0, 0));
+		Particle->CollisionParticles()->SetX(CollisionIndex++, FVec3(0, +Scale[1] / 2.0, 0));
+		Particle->CollisionParticles()->SetX(CollisionIndex++, FVec3(0, 0, -Scale[2] / 2.0));
+		Particle->CollisionParticles()->SetX(CollisionIndex++, FVec3(0, 0, +Scale[2] / 2.0));
 
 		if (OutElements != nullptr)
 		{
@@ -290,12 +290,12 @@ namespace ChaosTest {
 
 	void InitDynamicParticleCylinder2(FPBDRigidParticleHandle* Particle, const FVec3& Scale, TArray<TVec3<int32>>* OutElements, bool Tapered) 
 	{
-		Particle->X() = FVec3(0, 0, 0);
+		Particle->SetX(FVec3(0, 0, 0));
 		Particle->SetV(FVec3(0, 0, 0));
 		Particle->SetR(FRotation3::MakeFromEuler(FVec3(0, 0, 0)).GetNormalized());
 		Particle->SetW(FVec3(0, 0, 0));
 
-		Particle->P() = Particle->X();
+		Particle->SetP(Particle->GetX());
 		Particle->SetQ(Particle->GetR());
 
 		check(Scale.X == Scale.Y && Scale.X == Scale.Z);
@@ -318,14 +318,14 @@ namespace ChaosTest {
 		Particle->CollisionParticlesInitIfNeeded();
 		Particle->CollisionParticles()->AddParticles(8);
 
-		Particle->CollisionParticles()->X(CollisionIndex++) = FVec3(-Scale[0] / 2.0, 0, +Scale[2] / 2.0);
-		Particle->CollisionParticles()->X(CollisionIndex++) = FVec3(-Scale[0] / 2.0, 0, -Scale[2] / 2.0);
-		Particle->CollisionParticles()->X(CollisionIndex++) = FVec3(+Scale[0] / 2.0, 0, +Scale[2] / 2.0);
-		Particle->CollisionParticles()->X(CollisionIndex++) = FVec3(+Scale[0] / 2.0, 0, -Scale[2] / 2.0);
-		Particle->CollisionParticles()->X(CollisionIndex++) = FVec3(0, -Scale[1] / 2.0, +Scale[2] / 2.0);
-		Particle->CollisionParticles()->X(CollisionIndex++) = FVec3(0, -Scale[1] / 2.0, -Scale[2] / 2.0);
-		Particle->CollisionParticles()->X(CollisionIndex++) = FVec3(0, +Scale[1] / 2.0, +Scale[2] / 2.0);
-		Particle->CollisionParticles()->X(CollisionIndex++) = FVec3(0, +Scale[1] / 2.0, -Scale[2] / 2.0);
+		Particle->CollisionParticles()->SetX(CollisionIndex++, FVec3(-Scale[0] / 2.0, 0, +Scale[2] / 2.0));
+		Particle->CollisionParticles()->SetX(CollisionIndex++, FVec3(-Scale[0] / 2.0, 0, -Scale[2] / 2.0));
+		Particle->CollisionParticles()->SetX(CollisionIndex++, FVec3(+Scale[0] / 2.0, 0, +Scale[2] / 2.0));
+		Particle->CollisionParticles()->SetX(CollisionIndex++, FVec3(+Scale[0] / 2.0, 0, -Scale[2] / 2.0));
+		Particle->CollisionParticles()->SetX(CollisionIndex++, FVec3(0, -Scale[1] / 2.0, +Scale[2] / 2.0));
+		Particle->CollisionParticles()->SetX(CollisionIndex++, FVec3(0, -Scale[1] / 2.0, -Scale[2] / 2.0));
+		Particle->CollisionParticles()->SetX(CollisionIndex++, FVec3(0, +Scale[1] / 2.0, +Scale[2] / 2.0));
+		Particle->CollisionParticles()->SetX(CollisionIndex++, FVec3(0, +Scale[1] / 2.0, -Scale[2] / 2.0));
 
 		if (OutElements != nullptr)
 		{
@@ -405,7 +405,7 @@ namespace ChaosTest {
 
 	void InitStaticParticleBox(FGeometryParticleHandle* Particle, const FVec3& Scale, TArray<TVec3<int32>>* OutElements)
 	{
-		Particle->X() = FVec3(0, 0, 0);
+		Particle->SetX(FVec3(0, 0, 0));
 		Particle->SetR(FRotation3::FromIdentity());
 
 		FReal ScaleSq = Scale.X * Scale.X;
@@ -465,7 +465,7 @@ namespace ChaosTest {
 		InParticles.AddParticles(1);
 		int32 RigidBodyIndex = InParticles.Size() - 1;
 
-		InParticles.X(RigidBodyIndex) = FVec3(0, 0, 0);
+		InParticles.SetX(RigidBodyIndex, FVec3(0, 0, 0));
 		InParticles.SetV(RigidBodyIndex, FVec3(0, 0, 0));
 		InParticles.SetR(RigidBodyIndex, FRotation3::MakeFromEuler(FVec3(0, 0, 0)).GetNormalized());
 		InParticles.SetW(RigidBodyIndex, FVec3(0, 0, 0));
@@ -476,7 +476,7 @@ namespace ChaosTest {
 		InParticles.SetGeometry(RigidBodyIndex, MakeImplicitObjectPtr<TPlane<FReal, 3>>(FVec3(0, 0, 0), FVec3(0, 0, 1)));
 		InParticles.SetObjectState(RigidBodyIndex, EObjectStateType::Kinematic);
 
-		InParticles.P(RigidBodyIndex) = InParticles.X(RigidBodyIndex);
+		InParticles.SetP(RigidBodyIndex, InParticles.GetX(RigidBodyIndex));
 		InParticles.SetQ(RigidBodyIndex, InParticles.GetR(RigidBodyIndex));
 
 		return RigidBodyIndex;
@@ -487,7 +487,7 @@ namespace ChaosTest {
 		TArray<FKinematicGeometryParticleHandle*> Particles = SOAs.CreateKinematicParticles(1);
 		FKinematicGeometryParticleHandle* Particle = Particles[0];
 
-		Particle->X() = FVec3(0, 0, 0);
+		Particle->SetX(FVec3(0, 0, 0));
 		Particle->SetV(FVec3(0, 0, 0));
 		Particle->SetR(FRotation3::MakeFromEuler(FVec3(0, 0, 0)).GetNormalized());
 		Particle->SetW(FVec3(0, 0, 0));
@@ -535,7 +535,7 @@ namespace ChaosTest {
 		Chaos::FAABB3 BoundingBox(FVec3(0), FVec3(0));
 		for (int32 CollisionParticleIndex = 0; CollisionParticleIndex < (int32)SurfaceParticles.Size(); CollisionParticleIndex++)
 		{
-			BoundingBox.GrowToInclude(SurfaceParticles.X(CollisionParticleIndex));
+			BoundingBox.GrowToInclude(SurfaceParticles.GetX(CollisionParticleIndex));
 		}
 
 		// build cell domain
@@ -570,11 +570,11 @@ namespace ChaosTest {
 		Cube[7] = FConvex::FVec3Type{ 1, 1, 1 } * ScaleFloat;
 		Cube[8] = FConvex::FVec3Type{ 0, 0, 0 };
 
-		InParticles.X() = FVec3(0, 0, 0);
+		InParticles.SetX(FVec3(0, 0, 0));
 		InParticles.SetV(FVec3(0, 0, 0));
 		InParticles.SetR(FRotation3::MakeFromEuler(FVec3(0, 0, 0)).GetNormalized());
 		InParticles.SetW(FVec3(0, 0, 0));
-		InParticles.P() = InParticles.X();
+		InParticles.SetP(InParticles.GetX());
 		InParticles.SetQ(InParticles.GetR());
 
 		// TODO: Change this error prone API to set bounds more automatically. This is easy to forget
@@ -614,23 +614,23 @@ namespace ChaosTest {
 	/**/
 	FVec3 ObjectSpacePoint(FPBDRigidParticles& InParticles, const int32 Index, const FVec3& WorldSpacePoint)
 	{
-		FRigidTransform3 LocalToWorld(InParticles.X(Index), InParticles.GetR(Index));
+		FRigidTransform3 LocalToWorld(InParticles.GetX(Index), InParticles.GetR(Index));
 		return LocalToWorld.InverseTransformPosition(WorldSpacePoint);
 	}
 
 	FVec3 ObjectSpacePoint(FGeometryParticleHandle& Particle, const FVec3& WorldSpacePoint)
 	{
-		FRigidTransform3 LocalToWorld(Particle.X(), Particle.GetR());
+		FRigidTransform3 LocalToWorld(Particle.GetX(), Particle.GetR());
 		return LocalToWorld.InverseTransformPosition(WorldSpacePoint);
 	}
 
 	/**/
 	FReal PhiWithNormal(FPBDRigidParticles& InParticles, const int32 Index, const FVec3& WorldSpacePoint, FVec3& Normal)
 	{
-		FRigidTransform3(InParticles.X(Index), InParticles.GetR(Index));
+		FRigidTransform3(InParticles.GetX(Index), InParticles.GetR(Index));
 		FVec3 BodySpacePoint = ObjectSpacePoint(InParticles, Index, WorldSpacePoint);
 		FReal LocalPhi = InParticles.GetGeometry(Index)->PhiWithNormal(BodySpacePoint, Normal);
-		Normal = FRigidTransform3(InParticles.X(Index), InParticles.GetR(Index)).TransformVector(Normal);
+		Normal = FRigidTransform3(InParticles.GetX(Index), InParticles.GetR(Index)).TransformVector(Normal);
 		return LocalPhi;
 	}
 
@@ -644,10 +644,10 @@ namespace ChaosTest {
 	/**/
 	FReal PhiWithNormal(FGeometryParticleHandle& Particle, const FVec3& WorldSpacePoint, FVec3& Normal)
 	{
-		FRigidTransform3(Particle.X(), Particle.GetR());
+		FRigidTransform3(Particle.GetX(), Particle.GetR());
 		FVec3 BodySpacePoint = ObjectSpacePoint(Particle, WorldSpacePoint);
 		FReal LocalPhi = Particle.GetGeometry()->PhiWithNormal(BodySpacePoint, Normal);
-		Normal = FRigidTransform3(Particle.X(), Particle.GetR()).TransformVector(Normal);
+		Normal = FRigidTransform3(Particle.GetX(), Particle.GetR()).TransformVector(Normal);
 		return LocalPhi;
 	}
 
diff --git a/Engine/Source/Programs/HeadlessChaos/Private/HeadlessChaosTestVehicles.cpp b/Engine/Source/Programs/HeadlessChaos/Private/HeadlessChaosTestVehicles.cpp
--- a/Engine/Source/Programs/HeadlessChaos/Private/HeadlessChaosTestVehicles.cpp
+++ b/Engine/Source/Programs/HeadlessChaos/Private/HeadlessChaosTestVehicles.cpp
@@ -1521,7 +1521,7 @@ namespace ChaosTest
 
 		Evolution.SetPhysicsMaterial(Dynamic, MakeSerializable(PhysicsMaterial));
 
-		Dynamic->X() = FVec3(10, 10, 20);
+		Dynamic->SetX(FVec3(10, 10, 20));
 		Dynamic->M() = BodyMass;
 		Dynamic->InvM() = 1.0f / BodyMass;
 		Dynamic->I() = TVec3<FRealSingle>(100000.0f);
@@ -1534,7 +1534,7 @@ namespace ChaosTest
 		for (int i = 0; i < 500; ++i)
 		{
 			// latest body transform
-			const FTransform BodyTM(Dynamic->GetR(), Dynamic->X());
+			const FTransform BodyTM(Dynamic->GetR(), Dynamic->GetX());
 
 			for (int SpringIdx = 0; SpringIdx < 4; SpringIdx++)
 			{
@@ -1569,7 +1569,7 @@ namespace ChaosTest
 
 		float Tolerance = 0.5f; // half cm
 		float ExpectedRestingPosition = (10.f + PlaneZPos + WheelRadius);
-		EXPECT_LT(Dynamic->X().Z - ExpectedRestingPosition, Tolerance);
+		EXPECT_LT(Dynamic->GetX().Z - ExpectedRestingPosition, Tolerance);
 	}
 
 	GTEST_TEST(AllTraits, VehicleTest_WheelAcceleratingLongitudinalSlip_VaryingDelta)
diff --git a/Engine/Source/Programs/HeadlessChaos/Public/HeadlessChaosTestConstraints.h b/Engine/Source/Programs/HeadlessChaos/Public/HeadlessChaosTestConstraints.h
--- a/Engine/Source/Programs/HeadlessChaos/Public/HeadlessChaosTestConstraints.h
+++ b/Engine/Source/Programs/HeadlessChaos/Public/HeadlessChaosTestConstraints.h
@@ -80,7 +80,7 @@ namespace ChaosTest
 
 		void ResetParticle(FGeometryParticleHandle* Particle, const FVec3& Position, const FRotation3& Rotation, const FVec3& Velocity, const FVec3& AngularVelocity)
 		{
-			Particle->X() = Position;
+			Particle->SetX(Position);
 			Particle->SetR(Rotation);
 			if (FKinematicGeometryParticleHandle* KinParticle = Particle->CastToKinematicParticle())
 			{
@@ -89,7 +89,7 @@ namespace ChaosTest
 			}
 			if (FPBDRigidParticleHandle* PBDParticle = Particle->CastToRigidParticle())
 			{
-				PBDParticle->P() = Position;
+				PBDParticle->SetP(Position);
 				PBDParticle->SetQ(Rotation);
 			}
 		}
diff --git a/Engine/Source/Programs/HeadlessChaos/Public/HeadlessChaosTestEvolution.h b/Engine/Source/Programs/HeadlessChaos/Public/HeadlessChaosTestEvolution.h
--- a/Engine/Source/Programs/HeadlessChaos/Public/HeadlessChaosTestEvolution.h
+++ b/Engine/Source/Programs/HeadlessChaos/Public/HeadlessChaosTestEvolution.h
@@ -158,7 +158,7 @@ namespace ChaosTest
 
 		void ResetParticle(FGeometryParticleHandle* Particle, const FVec3& Position, const FRotation3& Rotation, const FVec3& Velocity, const FVec3& AngularVelocity)
 		{
-			Particle->X() = Position;
+			Particle->SetX(Position);
 			Particle->SetR(Rotation);
 			if (FKinematicGeometryParticleHandle* KinParticle = Particle->CastToKinematicParticle())
 			{
@@ -167,7 +167,7 @@ namespace ChaosTest
 			}
 			if (FPBDRigidParticleHandle* PBDParticle = Particle->CastToRigidParticle())
 			{
-				PBDParticle->P() = Position;
+				PBDParticle->SetP(Position);
 				PBDParticle->SetQ(Rotation);
 			}
 		}
diff --git a/Engine/Source/Runtime/Engine/Private/Collision/CollisionConversions.cpp b/Engine/Source/Runtime/Engine/Private/Collision/CollisionConversions.cpp
--- a/Engine/Source/Runtime/Engine/Private/Collision/CollisionConversions.cpp
+++ b/Engine/Source/Runtime/Engine/Private/Collision/CollisionConversions.cpp
@@ -110,7 +110,7 @@ static FVector FindGeomOpposingNormal(ECollisionShapeType QueryGeomType, const T
 			}
 			else
 			{
-				const FTransform ActorTM(Hit.Actor->GetR(), Hit.Actor->X());
+				const FTransform ActorTM(Hit.Actor->GetR(), Hit.Actor->GetX());
 				const FVector LocalInNormal = ActorTM.InverseTransformVectorNoScale(InNormal);
 				const FVector LocalTraceDirectionDenorm = ActorTM.InverseTransformVectorNoScale(TraceDirectionDenorm);
 				const FVector LocalNormal = Shape->GetGeometry()->FindGeometryOpposingNormal(LocalTraceDirectionDenorm, Hit.FaceIndex, LocalInNormal);
diff --git a/Engine/Source/Runtime/Engine/Private/PhysicsEngine/Experimental/ChaosCooking.cpp b/Engine/Source/Runtime/Engine/Private/PhysicsEngine/Experimental/ChaosCooking.cpp
--- a/Engine/Source/Runtime/Engine/Private/PhysicsEngine/Experimental/ChaosCooking.cpp
+++ b/Engine/Source/Runtime/Engine/Private/PhysicsEngine/Experimental/ChaosCooking.cpp
@@ -55,7 +55,7 @@ namespace Chaos
 			const int32 NumVerts = FinalVerts.Num();
 			for(int32 VertIndex = 0; VertIndex < NumVerts; ++VertIndex)
 			{
-				TriMeshParticles.X(VertIndex) = FinalVerts[VertIndex];
+				TriMeshParticles.SetX(VertIndex, FinalVerts[VertIndex]);
 			}
 
 			// Build chaos triangle list. #BGTODO Just make the clean function take these types instead of double copying
@@ -238,7 +238,7 @@ namespace Chaos
 			const int32 NumVerts = FinalVerts.Num();
 			for(int32 VertIndex = 0; VertIndex < NumVerts; ++VertIndex)
 			{
-				TriMeshParticles.X(VertIndex) = FinalVerts[VertIndex];
+				TriMeshParticles.SetX(VertIndex, FinalVerts[VertIndex]);
 			}
 
 			// Build chaos triangle list. #BGTODO Just make the clean function take these types instead of double copying
diff --git a/Engine/Source/Runtime/Engine/Private/PhysicsEngine/Experimental/PhysScene_Chaos.cpp b/Engine/Source/Runtime/Engine/Private/PhysicsEngine/Experimental/PhysScene_Chaos.cpp
--- a/Engine/Source/Runtime/Engine/Private/PhysicsEngine/Experimental/PhysScene_Chaos.cpp
+++ b/Engine/Source/Runtime/Engine/Private/PhysicsEngine/Experimental/PhysScene_Chaos.cpp
@@ -468,7 +468,7 @@ static TUniquePtr<FPhysScene_ChaosPauseHandler> PhysScene_ChaosPauseHandler;
 
 static void CopyParticleData(Chaos::FPBDRigidParticles& ToParticles, const int32 ToIndex, Chaos::FPBDRigidParticles& FromParticles, const int32 FromIndex)
 {
-	ToParticles.X(ToIndex) = FromParticles.X(FromIndex);
+	ToParticles.SetX(ToIndex, FromParticles.GetX(FromIndex));
 	ToParticles.SetR(ToIndex, FromParticles.GetR(FromIndex));
 	ToParticles.SetV(ToIndex, FromParticles.GetV(FromIndex));
 	ToParticles.SetW(ToIndex, FromParticles.GetW(FromIndex));
@@ -1359,7 +1359,7 @@ void FPhysScene_Chaos::PopulateReplicationCache(const int32 PhysicsStep)
 		else
 		{
 			FRigidBodyState& ReplicationState = ReplicationData.GetState();
-			ReplicationState.Position = Handle->X();
+			ReplicationState.Position = Handle->GetX();
 			ReplicationState.Quaternion = Handle->GetR();
 			ReplicationState.LinVel = Handle->GetV();
 			ReplicationState.AngVel = Handle->GetW();
diff --git a/Engine/Source/Runtime/Engine/Private/PhysicsEngine/ImmediatePhysics/ImmediatePhysicsChaos/ImmediatePhysicsActorHandle_Chaos.cpp b/Engine/Source/Runtime/Engine/Private/PhysicsEngine/ImmediatePhysics/ImmediatePhysicsChaos/ImmediatePhysicsActorHandle_Chaos.cpp
--- a/Engine/Source/Runtime/Engine/Private/PhysicsEngine/ImmediatePhysics/ImmediatePhysicsChaos/ImmediatePhysicsActorHandle_Chaos.cpp
+++ b/Engine/Source/Runtime/Engine/Private/PhysicsEngine/ImmediatePhysics/ImmediatePhysicsChaos/ImmediatePhysicsActorHandle_Chaos.cpp
@@ -233,7 +233,7 @@ namespace ImmediatePhysics_Chaos
 			{
 				ParticleHandle->SetHasBounds(true);
 				ParticleHandle->SetLocalBounds(Geometry->BoundingBox());
-				ParticleHandle->UpdateWorldSpaceState(FRigidTransform3(ParticleHandle->X(), ParticleHandle->GetR()), FVec3(0));
+				ParticleHandle->UpdateWorldSpaceState(FRigidTransform3(ParticleHandle->GetX(), ParticleHandle->GetR()), FVec3(0));
 			}
 
 			if (FKinematicGeometryParticleHandle* Kinematic = ParticleHandle->CastToKinematicParticle())
@@ -361,22 +361,22 @@ namespace ImmediatePhysics_Chaos
 
 		// Initialize the bounds. Important because if the particle never moves its 
 		// bounds will never get updated (see FPBDMinEvolution::ApplyKinematicTargets) 
-		ParticleHandle->UpdateWorldSpaceState(FRigidTransform3(ParticleHandle->X(), ParticleHandle->GetR()), FVec3(0));
+		ParticleHandle->UpdateWorldSpaceState(FRigidTransform3(ParticleHandle->GetX(), ParticleHandle->GetR()), FVec3(0));
 	}
 
 	void FActorHandle::SetWorldTransform(const FTransform& WorldTM)
 	{
 		using namespace Chaos;
 
-		ParticleHandle->X() = WorldTM.GetTranslation();
+		ParticleHandle->SetX(WorldTM.GetTranslation());
 		ParticleHandle->SetR(WorldTM.GetRotation());
 
 		FPBDRigidParticleHandle* Dynamic = ParticleHandle->CastToRigidParticle();
 		if(Dynamic && Dynamic->ObjectState() == Chaos::EObjectStateType::Dynamic)
 		{
-			Dynamic->P() = Dynamic->X();
+			Dynamic->SetP(Dynamic->GetX());
 			Dynamic->SetQf(Dynamic->GetRf());
-			Dynamic->AuxilaryValue(ParticlePrevXs) = Dynamic->P();
+			Dynamic->AuxilaryValue(ParticlePrevXs) = Dynamic->GetP();
 			Dynamic->AuxilaryValue(ParticlePrevRs) = Dynamic->GetQ();
 		}
 	}
diff --git a/Engine/Source/Runtime/Engine/Private/PhysicsEngine/ImmediatePhysics/ImmediatePhysicsChaos/ImmediatePhysicsSimulation_Chaos.cpp b/Engine/Source/Runtime/Engine/Private/PhysicsEngine/ImmediatePhysics/ImmediatePhysicsChaos/ImmediatePhysicsSimulation_Chaos.cpp
--- a/Engine/Source/Runtime/Engine/Private/PhysicsEngine/ImmediatePhysics/ImmediatePhysicsChaos/ImmediatePhysicsSimulation_Chaos.cpp
+++ b/Engine/Source/Runtime/Engine/Private/PhysicsEngine/ImmediatePhysics/ImmediatePhysicsChaos/ImmediatePhysicsSimulation_Chaos.cpp
@@ -457,7 +457,7 @@ namespace ImmediatePhysics_Chaos
 		ActorHandle->GetParticle()->AuxilaryValue(Implementation->ParticleMaterials) = MakeSerializable(Material);
 		ActorHandle->GetParticle()->AuxilaryValue(Implementation->PerParticleMaterials) = MoveTemp(Material);
 		ActorHandle->GetParticle()->AuxilaryValue(Implementation->CollidedParticles) = false;
-		ActorHandle->GetParticle()->AuxilaryValue(Implementation->ParticlePrevXs) = ActorHandle->GetParticle()->X();
+		ActorHandle->GetParticle()->AuxilaryValue(Implementation->ParticlePrevXs) = ActorHandle->GetParticle()->GetX();
 		ActorHandle->GetParticle()->AuxilaryValue(Implementation->ParticlePrevRs) = ActorHandle->GetParticle()->GetR();
 
 		Implementation->bActorsDirty = true;
diff --git a/Engine/Source/Runtime/Engine/Private/PhysicsEngine/PhysicsReplication.cpp b/Engine/Source/Runtime/Engine/Private/PhysicsEngine/PhysicsReplication.cpp
--- a/Engine/Source/Runtime/Engine/Private/PhysicsEngine/PhysicsReplication.cpp
+++ b/Engine/Source/Runtime/Engine/Private/PhysicsEngine/PhysicsReplication.cpp
@@ -1261,7 +1261,7 @@ bool FPhysicsReplicationAsync::DefaultReplication(Chaos::FPBDRigidParticleHandle
 
 	// Get Current state
 	FRigidBodyState CurrentState;
-	CurrentState.Position = Handle->X();
+	CurrentState.Position = Handle->GetX();
 	CurrentState.Quaternion = Handle->GetR();
 	CurrentState.AngVel = Handle->GetW();
 	CurrentState.LinVel = Handle->GetV();
@@ -1491,7 +1491,7 @@ bool FPhysicsReplicationAsync::PredictiveInterpolation(Chaos::FPBDRigidParticleH
 
 	// If target velocity is low enough, check the distance from the current position to the source position of our target to see if it's low enough to early out of replication
 	const bool bXCanEarlyOut = (PhysicsReplicationCVars::PredictiveInterpolationCVars::bEarlyOutWithVelocity || Target.TargetState.LinVel.SizeSquared() < UE_KINDA_SMALL_NUMBER) &&
-		(Target.PrevPosTarget - Handle->X()).SizeSquared() < PhysicsReplicationCVars::PredictiveInterpolationCVars::EarlyOutDistanceSqr;
+		(Target.PrevPosTarget - Handle->GetX()).SizeSquared() < PhysicsReplicationCVars::PredictiveInterpolationCVars::EarlyOutDistanceSqr;
 
 	// Early out if we are within range of target, also apply target sleep state
 	if (bXCanEarlyOut)
@@ -1522,7 +1522,7 @@ bool FPhysicsReplicationAsync::PredictiveInterpolation(Chaos::FPBDRigidParticleH
 
 	// CurrentState
 	FRigidBodyState CurrentState;
-	CurrentState.Position = Handle->X();
+	CurrentState.Position = Handle->GetX();
 	CurrentState.Quaternion = Handle->GetR();
 	CurrentState.LinVel = Handle->GetV();
 	CurrentState.AngVel = Handle->GetW(); // Note: Current angular velocity is in Radians
@@ -1758,7 +1758,7 @@ bool FPhysicsReplicationAsync::ResimulationReplication(Chaos::FPBDRigidParticleH
 	const float ResimErrorThreshold = SettingsCurrent.ResimulationSettings.GetResimulationErrorThreshold(Chaos::FPhysicsSolverBase::ResimulationErrorThreshold());
 	const Chaos::FGeometryParticleState PastState = RewindData->GetPastStateAtFrame(*Handle, LocalFrame, RewindPhase);
 
-	const FVector ErrorOffset = (Target.TargetState.Position - PastState.X());
+	const FVector ErrorOffset = (Target.TargetState.Position - PastState.GetX());
 	const float ErrorDistance = ErrorOffset.Size();
 	const bool ShouldTriggerResim = ErrorDistance >= ResimErrorThreshold;
 
@@ -1768,7 +1768,7 @@ bool FPhysicsReplicationAsync::ResimulationReplication(Chaos::FPBDRigidParticleH
 	{
 		UE_LOG(LogTemp, Log, TEXT("Apply Rigid body state at local frame %d with offset = %d"), LocalFrame, Target.FrameOffset);
 		UE_LOG(LogTemp, Log, TEXT("Particle Position Error = %f | Should Trigger Resim = %s | Server Frame = %d | Client Frame = %d"), ErrorDistance, (ShouldTriggerResim ? TEXT("True") : TEXT("False")), Target.ServerFrame, LocalFrame);
-		UE_LOG(LogTemp, Log, TEXT("Particle Target Position = %s | Current Position = %s"), *Target.TargetState.Position.ToString(), *PastState.X().ToString());
+		UE_LOG(LogTemp, Log, TEXT("Particle Target Position = %s | Current Position = %s"), *Target.TargetState.Position.ToString(), *PastState.GetX().ToString());
 		UE_LOG(LogTemp, Log, TEXT("Particle Target Velocity = %s | Current Velocity = %s"), *Target.TargetState.LinVel.ToString(), *PastState.GetV().ToString());
 		UE_LOG(LogTemp, Log, TEXT("Particle Target Quaternion = %s | Current Quaternion = %s"), *Target.TargetState.Quaternion.ToString(), *PastState.GetR().ToString());
 		UE_LOG(LogTemp, Log, TEXT("Particle Target Omega = %s | Current Omega= %s"), *Target.TargetState.AngVel.ToString(), *PastState.GetW().ToString());
@@ -1779,9 +1779,9 @@ bool FPhysicsReplicationAsync::ResimulationReplication(Chaos::FPBDRigidParticleH
 			const FColor DebugColor = FLinearColor::LerpUsingHSV(FLinearColor::Green, FLinearColor::Red, ColorLerp).ToFColor(false);
 
 			Chaos::FDebugDrawQueue::GetInstance().DrawDebugBox(Target.TargetState.Position, FVector(BoxSize, BoxSize, BoxSize), Target.TargetState.Quaternion, FColor::Orange, true, CharacterMovementCVars::NetCorrectionLifetime, 0, 1.0f);
-			Chaos::FDebugDrawQueue::GetInstance().DrawDebugBox(PastState.X(), FVector(6, 6, 6), PastState.GetR(), DebugColor, true, CharacterMovementCVars::NetCorrectionLifetime, 0, 1.0f);
+			Chaos::FDebugDrawQueue::GetInstance().DrawDebugBox(PastState.GetX(), FVector(6, 6, 6), PastState.GetR(), DebugColor, true, CharacterMovementCVars::NetCorrectionLifetime, 0, 1.0f);
 
-			Chaos::FDebugDrawQueue::GetInstance().DrawDebugDirectionalArrow(PastState.X(), Target.TargetState.Position, 5.0f, FColor::MakeRandomSeededColor(LocalFrame), true, CharacterMovementCVars::NetCorrectionLifetime, 0, 0.5f);
+			Chaos::FDebugDrawQueue::GetInstance().DrawDebugDirectionalArrow(PastState.GetX(), Target.TargetState.Position, 5.0f, FColor::MakeRandomSeededColor(LocalFrame), true, CharacterMovementCVars::NetCorrectionLifetime, 0, 0.5f);
 		}
 	}
 #endif
@@ -1805,7 +1805,7 @@ bool FPhysicsReplicationAsync::ResimulationReplication(Chaos::FPBDRigidParticleH
 			{
 				// Calculate correction to position
 				const float CorrectionAmountX = SettingsCurrent.ResimulationSettings.GetPosStabilityMultiplier() / NumPredictedFrames; // Same result as (ErrorOffset / NumPredictedFrames) * PosStabilityMultiplier
-				const FVector CorrectedX = Handle->X() + (ErrorOffset * CorrectionAmountX);
+				const FVector CorrectedX = Handle->GetX() + (ErrorOffset * CorrectionAmountX);
 
 				// Calculate correction to rotation
 				const float CorrectionAmountR = (1.f / NumPredictedFrames) * SettingsCurrent.ResimulationSettings.GetRotStabilityMultiplier();
@@ -1817,7 +1817,7 @@ bool FPhysicsReplicationAsync::ResimulationReplication(Chaos::FPBDRigidParticleH
 	#if !(UE_BUILD_SHIPPING || UE_BUILD_TEST)
 				if (Chaos::FPhysicsSolverBase::CanDebugNetworkPhysicsPrediction())
 				{
-					Chaos::FDebugDrawQueue::GetInstance().DrawDebugDirectionalArrow(Handle->X(), CorrectedX, 5.0f, FColor::MakeRandomSeededColor(LocalFrame), true, CharacterMovementCVars::NetCorrectionLifetime, 0, 0.5f);
+					Chaos::FDebugDrawQueue::GetInstance().DrawDebugDirectionalArrow(Handle->GetX(), CorrectedX, 5.0f, FColor::MakeRandomSeededColor(LocalFrame), true, CharacterMovementCVars::NetCorrectionLifetime, 0, 0.5f);
 				}
 	#endif
 				// Apply correction
diff --git a/Engine/Source/Runtime/Experimental/Chaos/Private/Chaos/BVHParticles.cpp b/Engine/Source/Runtime/Experimental/Chaos/Private/Chaos/BVHParticles.cpp
--- a/Engine/Source/Runtime/Experimental/Chaos/Private/Chaos/BVHParticles.cpp
+++ b/Engine/Source/Runtime/Experimental/Chaos/Private/Chaos/BVHParticles.cpp
@@ -44,7 +44,7 @@ namespace Chaos
 		AddParticles(Other.Size());
 		for(int32 i = Other.Size() - 1; 0 <= i; i--)
 		{
-			X(i) = Other.X(i);
+			SetX(i, Other.GetX(i));
 		}
 		MBVH = new TBoundingVolumeHierarchy<FParticles, TArray<int32>>(*this, CollisionParticlesBVHDepth);
 	}
diff --git a/Engine/Source/Runtime/Experimental/Chaos/Private/Chaos/CCDUtilities.cpp b/Engine/Source/Runtime/Experimental/Chaos/Private/Chaos/CCDUtilities.cpp
--- a/Engine/Source/Runtime/Experimental/Chaos/Private/Chaos/CCDUtilities.cpp
+++ b/Engine/Source/Runtime/Experimental/Chaos/Private/Chaos/CCDUtilities.cpp
@@ -347,10 +347,10 @@ namespace Chaos
 		// to get the position at that time. This is required to handle objects with a CoM offset from the actor position.
 		// We will undo this manipulation at the end.
 		// NOTE: We do not modify the previous rotation R here - we just use the current rotation Q everywhere
-		// @todo(chaos): we should store the sweep positions in CCDParticle and onlymodify the original particle at the end
+		// @todo(chaos): we should store the sweep positions in CCDParticle and only modify the original particle at the end
 		for (int32 i = ParticleStart; i < ParticleStart + ParticleNum; i++)
 		{
-			GroupedCCDParticles[i]->Particle->X() = GroupedCCDParticles[i]->Particle->P() - GroupedCCDParticles[i]->Particle->GetV() * Dt;
+			GroupedCCDParticles[i]->Particle->SetX(GroupedCCDParticles[i]->Particle->GetP() - GroupedCCDParticles[i]->Particle->GetV() * Dt);
 		}
 
 		// Sort constraints based on TOI
@@ -497,7 +497,7 @@ namespace Chaos
 		{
 			const FVec3 CoMPrev = GroupedCCDParticles[i]->Particle->PCom() - GroupedCCDParticles[i]->Particle->GetV() * Dt;
 			const FVec3 CoMOffsetPrev = GroupedCCDParticles[i]->Particle->GetR() * GroupedCCDParticles[i]->Particle->CenterOfMass();
-			GroupedCCDParticles[i]->Particle->X() = CoMPrev - CoMOffsetPrev;
+			GroupedCCDParticles[i]->Particle->SetX(CoMPrev - CoMOffsetPrev);
 		}
 	}
 
@@ -699,7 +699,7 @@ namespace Chaos
 						{
 							AdvanceParticleXToTOI(AffectedCCDParticle, IslandTOI, Dt);
 						}
-						ParticleStartWorldTransforms[j] = FRigidTransform3(AffectedParticle->X(), AffectedParticle->GetR());
+						ParticleStartWorldTransforms[j] = FRigidTransform3(AffectedParticle->GetX(), AffectedParticle->GetR());
 					}
 					else
 					{
@@ -711,7 +711,7 @@ namespace Chaos
 						}
 						else // Static case
 						{
-							ParticleStartWorldTransforms[j] = FRigidTransform3(AffectedParticle->X(), AffectedParticle->GetR());
+							ParticleStartWorldTransforms[j] = FRigidTransform3(AffectedParticle->GetX(), AffectedParticle->GetR());
 						}
 					}
 				}
@@ -794,7 +794,7 @@ namespace Chaos
 		{
 			TPBDRigidParticleHandle<FReal, 3>* Particle = CCDParticle->Particle;
 			const FReal RestDt = (TOI - CCDParticle->TOI) * Dt;
-			Particle->X() = Particle->X() + Particle->GetV() * RestDt;
+			Particle->SetX(Particle->GetX() + Particle->GetV() * RestDt);
 			CCDParticle->TOI = TOI;
 		}
 	}
@@ -803,20 +803,20 @@ namespace Chaos
 	{
 		TPBDRigidParticleHandle<FReal, 3>* Particle = CCDParticle->Particle;
 		const FReal RestDt = (1.f - CCDParticle->TOI) * Dt;
-		Particle->P() = Particle->X() + Particle->GetV() * RestDt;
+		Particle->SetP(Particle->GetX() + Particle->GetV() * RestDt);
 	}
 
 	void FCCDManager::ClipParticleP(FCCDParticle *CCDParticle) const
 	{
 		TPBDRigidParticleHandle<FReal, 3>* Particle = CCDParticle->Particle;
-		Particle->P() = Particle->X();
+		Particle->SetP(Particle->GetX());
 	}
 
 	void FCCDManager::ClipParticleP(FCCDParticle *CCDParticle, const FVec3 Offset) const
 	{
 		TPBDRigidParticleHandle<FReal, 3>* Particle = CCDParticle->Particle;
-		Particle->X() += Offset;
-		Particle->P() = Particle->X();
+		Particle->SetX(Particle->GetX() + Offset);
+		Particle->SetP(Particle->GetX());
 	}
 
 	void FCCDManager::ApplyImpulse(FCCDConstraint *CCDConstraint)
@@ -871,7 +871,7 @@ namespace Chaos
 
 	void FCCDManager::UpdateSweptConstraints(const FReal Dt, Private::FCollisionConstraintAllocator *CollisionAllocator)
 	{
-		// Buld the set of collision whose contact data will be out of date because we moved one or both of its particles. 
+		// Build the set of collision whose contact data will be out of date because we moved one or both of its particles. 
 		// This is all collision constraints, including non-swept ones, for any particle that was relocated by the CCD sweep 
 		// logic executed in ApplySweptConstraints (i.e., contents of CCDConstraints)
 		// @todo(chaos): we could calculate the size of the Collisions array in Init
@@ -1035,14 +1035,14 @@ namespace Chaos
 					if (Bias0 > UE_SMALL_NUMBER)
 					{
 						const FVec3 Correction0 = -Bias0 * Phi * WorldNormal;
-						P0->P() = P0->P() + Correction0;
+						P0->SetP(P0->P() + Correction0);
 						P0->SetV(P0->V() - Bias0 * VelNormal * WorldNormal);
 					}
 
 					if (Bias1 > UE_SMALL_NUMBER)
 					{
 						const FVec3 Correction1 = Bias1 * Phi * WorldNormal;
-						P1->P() = P1->P() + Correction1;
+						P1->SetP(P1->P() + Correction1);
 						P1->SetV(P1->V() + Bias1 * VelNormal * WorldNormal);
 					}
 				}
@@ -1057,7 +1057,7 @@ namespace Chaos
 		for (FCCDParticle& CCDParticle : CCDParticles)
 		{
 			TPBDRigidParticleHandle<FReal, 3>* Particle = CCDParticle.Particle;
-			Particle->X() = Particle->P() - Particle->GetV() * Dt;
+			Particle->SetX(Particle->GetP() - Particle->GetV() * Dt);
 		}
 	}
 
@@ -1146,8 +1146,8 @@ namespace Chaos
 		// For non-rigids, DeltaX is zero and use R for rotation.
 		const auto Rigid0 = Particle0.CastToRigidParticle();
 		const auto Rigid1 = Particle1.CastToRigidParticle();
-		const FVec3 DeltaX0 = Rigid0 ? Rigid0->P() - Rigid0->X() : FVec3::ZeroVector;
-		const FVec3 DeltaX1 = Rigid1 ? Rigid1->P() - Rigid1->X() : FVec3::ZeroVector;
+		const FVec3 DeltaX0 = Rigid0 ? Rigid0->GetP() - Rigid0->GetX() : FVec3::ZeroVector;
+		const FVec3 DeltaX1 = Rigid1 ? Rigid1->GetP() - Rigid1->GetX() : FVec3::ZeroVector;
 		const FQuat& R0 = Rigid0 ? Rigid0->GetQ() : Particle0.GetR();
 		const FQuat& R1 = Rigid1 ? Rigid1->GetQ() : Particle1.GetR();
 		return DeltaExceedsThreshold(
diff --git a/Engine/Source/Runtime/Experimental/Chaos/Private/Chaos/ChaosDebugDraw.cpp b/Engine/Source/Runtime/Experimental/Chaos/Private/Chaos/ChaosDebugDraw.cpp
--- a/Engine/Source/Runtime/Experimental/Chaos/Private/Chaos/ChaosDebugDraw.cpp
+++ b/Engine/Source/Runtime/Experimental/Chaos/Private/Chaos/ChaosDebugDraw.cpp
@@ -506,7 +506,7 @@ namespace Chaos
 				{
 					for (int32 ParticleIndex = 0; ParticleIndex < (int32)CollisionParticles->Size(); ++ParticleIndex)
 					{
-						const FVec3 P = ShapeTransform.TransformPosition(CollisionParticles->X(ParticleIndex));
+						const FVec3 P = ShapeTransform.TransformPosition(CollisionParticles->GetX(ParticleIndex));
 						FDebugDrawQueue::GetInstance().DrawDebugPoint(P, Color, false, Duration, 0, Settings.PointSize);
 					}
 				}
@@ -881,7 +881,7 @@ namespace Chaos
 					{
 						for (int32 ParticleIndex = 0; ParticleIndex < (int32)Particles->Size(); ++ParticleIndex)
 						{
-							FVec3 P = ShapeTransform.TransformPosition(Particles->X(ParticleIndex));
+							FVec3 P = ShapeTransform.TransformPosition(Particles->GetX(ParticleIndex));
 							FDebugDrawQueue::GetInstance().DrawDebugPoint(P, ShapeColor, false, Duration, uint8(Settings.DrawPriority), Settings.PointSize);
 						}
 					}
@@ -929,7 +929,7 @@ namespace Chaos
 
 		void DrawParticleShapesImpl(const FRigidTransform3& SpaceTransform, const FGeometryParticleHandle* Particle, const FColor& InColor, const FChaosDebugDrawSettings& Settings)
 		{
-			const FVec3 P = SpaceTransform.TransformPosition(Particle->ObjectState() == EObjectStateType::Dynamic ? Particle->CastToRigidParticle()->P() : Particle->X());
+			const FVec3 P = SpaceTransform.TransformPosition(Particle->ObjectState() == EObjectStateType::Dynamic ? Particle->CastToRigidParticle()->GetP() : Particle->GetX());
 			const FRotation3 Q = SpaceTransform.GetRotation() * (Particle->ObjectState() == EObjectStateType::Dynamic ? Particle->CastToRigidParticle()->GetQ() : Particle->GetR());
 			const FRigidTransform3 ParticleSpaceTransform = FRigidTransform3(P, Q);
 
@@ -990,7 +990,7 @@ namespace Chaos
 			{
 				if (Union->GetBVH() != nullptr)
 				{
-					FVec3 P = SpaceTransform.TransformPosition(Particle->X());
+					FVec3 P = SpaceTransform.TransformPosition(Particle->GetX());
 					FRotation3 Q = SpaceTransform.GetRotation() * (Particle->GetR());
 
 					DrawBVHImpl(Particle->Handle(), FRigidTransform3(P, Q), Union->GetBVH(), InColor, 0.0f, Settings);
@@ -1804,7 +1804,7 @@ namespace Chaos
 						}
 					}
 
-					return FTransform{ Particle->GetR(), Particle->X() };
+					return FTransform{ Particle->GetR(), Particle->GetX() };
 				};
 
 				return GetClusterTransformImpl(Particle, GetClusterTransformImpl);
@@ -1923,7 +1923,7 @@ namespace Chaos
 
 						if (bChaosDebugDrawConnectionGraphShowInternalStrains)
 						{
-							FDebugDrawQueue::GetInstance().DrawDebugString(Child->X(), FString::Printf(TEXT("%.1f"), Child->GetInternalStrains()), nullptr, FColor::White, UE_KINDA_SMALL_NUMBER, false, 1.0);
+							FDebugDrawQueue::GetInstance().DrawDebugString(Child->GetX(), FString::Printf(TEXT("%.1f"), Child->GetInternalStrains()), nullptr, FColor::White, UE_KINDA_SMALL_NUMBER, false, 1.0);
 						}
 
 						const FConnectivityEdgeArray& Edges = Child->ConnectivityEdges();
diff --git a/Engine/Source/Runtime/Experimental/Chaos/Private/Chaos/ClusterUnionManager.cpp b/Engine/Source/Runtime/Experimental/Chaos/Private/Chaos/ClusterUnionManager.cpp
--- a/Engine/Source/Runtime/Experimental/Chaos/Private/Chaos/ClusterUnionManager.cpp
+++ b/Engine/Source/Runtime/Experimental/Chaos/Private/Chaos/ClusterUnionManager.cpp
@@ -42,7 +42,7 @@ namespace Chaos
 			}
 			else
 			{
-				const FRigidTransform3 ChildWorldTM(Child->X(), Child->GetR());
+				const FRigidTransform3 ChildWorldTM(Child->GetX(), Child->GetR());
 				Frame = ChildWorldTM.GetRelativeTransform(ClusterWorldTM);
 			}
 
@@ -54,7 +54,7 @@ namespace Chaos
 			if (!ClusterUnion.ChildParticles.IsEmpty() && ClusterUnion.Geometry->GetType() == ImplicitObjectType::Union)
 			{
 				constexpr FReal kAddThickness = 5.0;
-				const FRigidTransform3 ClusterWorldTM(ClusterUnion.InternalCluster->X(), ClusterUnion.InternalCluster->GetR());
+				const FRigidTransform3 ClusterWorldTM(ClusterUnion.InternalCluster->GetX(), ClusterUnion.InternalCluster->GetR());
 
 				// Use the acceleration structure of the cluster union itself to make finding overlaps easy.
 				const FImplicitObjectUnion& ShapeUnion = ClusterUnion.Geometry->GetObjectChecked<FImplicitObjectUnion>();
@@ -387,7 +387,7 @@ namespace Chaos
 		}
 
 		// TODO: Can we do something better than a union?
-		const FRigidTransform3 ClusterWorldTM(Union.InternalCluster->X(), Union.InternalCluster->GetR());
+		const FRigidTransform3 ClusterWorldTM(Union.InternalCluster->GetX(), Union.InternalCluster->GetR());
 		TArray<Chaos::FImplicitObjectPtr> Objects;
 		Objects.Reserve(Union.ChildParticles.Num());
 
@@ -689,7 +689,7 @@ namespace Chaos
 			// the deferred cluster union update properties is called, the position of the particle will get reset to the position of the cluster union.
 			if (FPBDRigidClusteredParticleHandle* ClusterParticle = Particle->CastToClustered())
 			{
-				const FRigidTransform3 ClusterWorldTM(Cluster->InternalCluster->X(), Cluster->InternalCluster->GetR());
+				const FRigidTransform3 ClusterWorldTM(Cluster->InternalCluster->GetX(), Cluster->InternalCluster->GetR());
 				const FRigidTransform3 Frame = GetParticleRigidFrameInClusterUnion(Particle, ClusterWorldTM);
 				ClusterParticle->SetChildToParent(Frame);
 			}
@@ -1106,7 +1106,7 @@ namespace Chaos
 		MEvolution.InvalidateParticle(ClusterUnion.InternalCluster);
 		ImplicitUnion->SetAllowBVH(false);
 
-		const FRigidTransform3 ClusterWorldTM(ClusterUnion.InternalCluster->X(), ClusterUnion.InternalCluster->GetR());
+		const FRigidTransform3 ClusterWorldTM(ClusterUnion.InternalCluster->GetX(), ClusterUnion.InternalCluster->GetR());
 		if (!PendingGeometryAdditions.IsEmpty())
 		{
 			ModifyAdditionOfChildrenToClusterUnionGeometry(
@@ -1357,7 +1357,7 @@ namespace Chaos
 
 						// Update the child's world transform to be consistent with its ChildToParent transform
 						const FPBDRigidClusteredParticleHandle* ParentHandle = ClusterUnion->InternalCluster;
-						const FRigidTransform3 ParticleToWorld = ChildToParent * FRigidTransform3(ParentHandle->X(), ParentHandle->GetR());
+						const FRigidTransform3 ParticleToWorld = ChildToParent * FRigidTransform3(ParentHandle->GetX(), ParentHandle->GetR());
 						ClusteredParticle->SetX(ParticleToWorld.GetTranslation());
 						ClusteredParticle->SetP(ParticleToWorld.GetTranslation());
 						ClusteredParticle->SetR(ParticleToWorld.GetRotation());
diff --git a/Engine/Source/Runtime/Experimental/Chaos/Private/Chaos/ClusterUnionStressSolver.cpp b/Engine/Source/Runtime/Experimental/Chaos/Private/Chaos/ClusterUnionStressSolver.cpp
--- a/Engine/Source/Runtime/Experimental/Chaos/Private/Chaos/ClusterUnionStressSolver.cpp
+++ b/Engine/Source/Runtime/Experimental/Chaos/Private/Chaos/ClusterUnionStressSolver.cpp
@@ -44,7 +44,7 @@ namespace Chaos
 
 		// todo(chaos) : change this logic to be less hacky and be able to handle dynamic structures 
 		//				 we could also look for contact point and joints that cause reaction to the 
-		return (Particle->IsKinematic() && Particle->X().Z < 200.0);
+		return (Particle->IsKinematic() && Particle->GetX().Z < 200.0);
 	}
 
 	void FClusterUnionStressSolver::FNode::AddMassContribution(double MassRatio, const FVec3& CenterOfMass)
diff --git a/Engine/Source/Runtime/Experimental/Chaos/Private/Chaos/CollisionResolutionUtil.cpp b/Engine/Source/Runtime/Experimental/Chaos/Private/Chaos/CollisionResolutionUtil.cpp
--- a/Engine/Source/Runtime/Experimental/Chaos/Private/Chaos/CollisionResolutionUtil.cpp
+++ b/Engine/Source/Runtime/Experimental/Chaos/Private/Chaos/CollisionResolutionUtil.cpp
@@ -555,11 +555,11 @@ namespace Chaos
 					{
 						if (NormalAveraging && UpdateType != ECollisionUpdateType::Any)	//if we just want one don't bother with normal
 						{
-							SampleObjectNormalAverageHelper(Object, ObjectTransform, SampleToObjectTM, SampleParticles.X(i), CullingDistance, WeightSum, AvgContact);
+							SampleObjectNormalAverageHelper(Object, ObjectTransform, SampleToObjectTM, SampleParticles.GetX(i), CullingDistance, WeightSum, AvgContact);
 						}
 						else
 						{
-							if (SampleObjectNoNormal(Object, ObjectTransform, SampleToObjectTM, SampleParticles.X(i), CullingDistance, AvgContact))
+							if (SampleObjectNoNormal(Object, ObjectTransform, SampleToObjectTM, SampleParticles.GetX(i), CullingDistance, AvgContact))
 							{
 								DeepestParticle = i;
 								if (UpdateType == ECollisionUpdateType::Any)
@@ -579,11 +579,11 @@ namespace Chaos
 				{
 					if (NormalAveraging && UpdateType != ECollisionUpdateType::Any)	//if we just want one don't bother with normal
 					{
-						const bool bInside = SampleObjectNormalAverageHelper(Object, ObjectTransform, SampleToObjectTM, SampleParticles.X(i), CullingDistance, WeightSum, AvgContact);
+						const bool bInside = SampleObjectNormalAverageHelper(Object, ObjectTransform, SampleToObjectTM, SampleParticles.GetX(i), CullingDistance, WeightSum, AvgContact);
 					}
 					else
 					{
-						if (SampleObjectNoNormal(Object, ObjectTransform, SampleToObjectTM, SampleParticles.X(i), CullingDistance, AvgContact))
+						if (SampleObjectNoNormal(Object, ObjectTransform, SampleToObjectTM, SampleParticles.GetX(i), CullingDistance, AvgContact))
 						{
 							DeepestParticle = i;
 							if (UpdateType == ECollisionUpdateType::Any)
@@ -619,10 +619,10 @@ namespace Chaos
 			else if (AvgContact.Phi < CullingDistance)
 			{
 				check(DeepestParticle >= 0);
-				FVec3 LocalPoint = SampleToObjectTM.TransformPositionNoScale(SampleParticles.X(DeepestParticle));
+				FVec3 LocalPoint = SampleToObjectTM.TransformPositionNoScale(SampleParticles.GetX(DeepestParticle));
 				FVec3 LocalNormal;
 				Contact.Phi = Object.PhiWithNormal(LocalPoint, LocalNormal);
-				Contact.ShapeContactPoints[0] = SampleParticles.X(DeepestParticle);
+				Contact.ShapeContactPoints[0] = SampleParticles.GetX(DeepestParticle);
 				Contact.ShapeContactPoints[1] = LocalPoint - Contact.Phi * LocalNormal;
 				Contact.ShapeContactNormal = LocalNormal;
 			}
diff --git a/Engine/Source/Runtime/Experimental/Chaos/Private/Chaos/ContactModification.cpp b/Engine/Source/Runtime/Experimental/Chaos/Private/Chaos/ContactModification.cpp
--- a/Engine/Source/Runtime/Experimental/Chaos/Private/Chaos/ContactModification.cpp
+++ b/Engine/Source/Runtime/Experimental/Chaos/Private/Chaos/ContactModification.cpp
@@ -251,7 +251,7 @@ namespace Chaos
 			EObjectStateType State(RigidHandle->ObjectState());
 			if (State == EObjectStateType::Dynamic || State == EObjectStateType::Sleeping)
 			{
-				RigidHandle->SetX(RigidHandle->P() - Velocity * Modifier->Dt);
+				RigidHandle->SetX(RigidHandle->GetP() - Velocity * Modifier->Dt);
 			}
 		}
 	}
@@ -301,10 +301,10 @@ namespace Chaos
 
 		if (RigidHandle)
 		{
-			return RigidHandle->P();
+			return RigidHandle->GetP();
 		}
 
-		return Particle->X();
+		return Particle->GetX();
 	}
 
 	void FContactPairModifier::UpdateConstraintShapeTransforms()
@@ -331,12 +331,12 @@ namespace Chaos
 
 				if (bMaintainVelocity)
 				{
-					RigidHandle->SetX(RigidHandle->P() - RigidHandle->GetV() * Modifier->Dt);
+					RigidHandle->SetX(RigidHandle->GetP() - RigidHandle->GetV() * Modifier->Dt);
 				}
 				else if(Modifier->Dt > 0.0f)
 				{
 					// Update V to new implicit velocity
-					RigidHandle->SetV((RigidHandle->P() - RigidHandle->X()) / Modifier->Dt);
+					RigidHandle->SetV((RigidHandle->GetP() - RigidHandle->GetX()) / Modifier->Dt);
 				}
 				UpdateConstraintShapeTransforms();
 				return;
@@ -392,16 +392,16 @@ namespace Chaos
 			EObjectStateType State(RigidHandle->ObjectState());
 			if (State == EObjectStateType::Dynamic || State == EObjectStateType::Sleeping)
 			{
-				RigidHandle->SetQ(Rotation);
+				RigidHandle->SetQf(Rotation);
 
 				if (bMaintainVelocity)
 				{
-					RigidHandle->SetR(FRotation3::IntegrateRotationWithAngularVelocity(RigidHandle->GetQ(), -RigidHandle->GetW(), Modifier->Dt));
+					RigidHandle->SetRf(FRotation3f::IntegrateRotationWithAngularVelocity(RigidHandle->GetQf(), -RigidHandle->GetWf(), FRealSingle(Modifier->Dt)));
 				}
 				else if (Modifier->Dt > 0.0f)
 				{
 					// Update W to new implicit velocity
-					RigidHandle->SetW(FRotation3::CalculateAngularVelocity(RigidHandle->GetR(), RigidHandle->GetQ(), Modifier->Dt));
+					RigidHandle->SetWf(FRotation3::CalculateAngularVelocity(RigidHandle->GetRf(), RigidHandle->GetQf(), FRealSingle(Modifier->Dt)));
 				}
 				UpdateConstraintShapeTransforms();
 				return;
@@ -409,8 +409,8 @@ namespace Chaos
 			else
 			{
 				// Kinematic must keep Q/R in sync
-				RigidHandle->SetR(Rotation);
-				RigidHandle->SetQ(Rotation);
+				RigidHandle->SetRf(Rotation);
+				RigidHandle->SetQf(Rotation);
 				UpdateConstraintShapeTransforms();
 				return;
 			}
@@ -419,7 +419,7 @@ namespace Chaos
 		FKinematicGeometryParticleHandle* KinematicHandle = Particle->CastToKinematicParticle();
 		if (KinematicHandle)
 		{
-			KinematicHandle->SetR(Rotation);
+			KinematicHandle->SetRf(Rotation);
 			UpdateConstraintShapeTransforms();
 			return;
 		}
