Chaos - Resubmit 31279946 with additional fix

- fix AABB support vertex index and match Box indices to AABB indices
- fix Box GetOpposingNormal to match new indexing scheme
- unit test to verify box support vertices for all cardinal and diagonal directions
- unit test to verify that AABB and Box vertex indices match

There was a mismatch between the vertex indices on TBox versus TAABB, but also the TAABB's vertex index calculation in the support functions was incorrect. Luckily this didnt matter, but some upcoming changes now require that the Box support functions return the correct vertex index, and that requires the above two fixes.

diff --git a/Engine/Source/Programs/HeadlessChaos/Private/HeadlessChaosTestConvex.cpp b/Engine/Source/Programs/HeadlessChaos/Private/HeadlessChaosTestConvex.cpp
--- a/Engine/Source/Programs/HeadlessChaos/Private/HeadlessChaosTestConvex.cpp
+++ b/Engine/Source/Programs/HeadlessChaos/Private/HeadlessChaosTestConvex.cpp
@@ -446,6 +446,215 @@ namespace ChaosTest
 		}
 	}
 
+	template<typename ConvexType>
+	void TestConvexPlaneVertices(const ConvexType& Convex)
+	{
+		const FReal NormalTolerance = UE_SMALL_NUMBER;
+		const FReal PositionTolerance = UE_KINDA_SMALL_NUMBER;
+
+		for (int32 PlaneIndex = 0; PlaneIndex < Convex.NumPlanes(); ++PlaneIndex)
+		{
+			const FVec3 PlaneN = Convex.GetPlane(PlaneIndex).Normal();
+			const FVec3 PlaneX = Convex.GetPlane(PlaneIndex).X();
+
+			const int NumPlaneVertices = Convex.NumPlaneVertices(PlaneIndex);
+			for (int32 PlaneVertexIndex0 = 0; PlaneVertexIndex0 < NumPlaneVertices; ++PlaneVertexIndex0)
+			{
+				const int32 VertexIndex0 = Convex.GetPlaneVertex(PlaneIndex, PlaneVertexIndex0);
+
+				// All vertices are actually on the plane
+				const FVec3 Vertex0 = Convex.GetVertex(VertexIndex0);
+				EXPECT_NEAR(FVec3::DotProduct(Vertex0, PlaneN), FVec3::DotProduct(PlaneX, PlaneN), PositionTolerance) << "PlaneIndex=" << PlaneIndex << " PlaneVertexIndex0=" << PlaneVertexIndex0;
+
+				// Winding is correct
+				int PlaneVertexIndex1 = (PlaneVertexIndex0 < NumPlaneVertices - 1) ? PlaneVertexIndex0 + 1 : 0;
+				int PlaneVertexIndex2 = (PlaneVertexIndex0 < NumPlaneVertices - 2) ? PlaneVertexIndex0 + 2 : PlaneVertexIndex0 - NumPlaneVertices + 2;
+				const int32 VertexIndex1 = Convex.GetPlaneVertex(PlaneIndex, PlaneVertexIndex1);
+				const int32 VertexIndex2 = Convex.GetPlaneVertex(PlaneIndex, PlaneVertexIndex2);
+				const FVec3 Vertex1 = Convex.GetVertex(VertexIndex1);
+				const FVec3 Vertex2 = Convex.GetVertex(VertexIndex2);
+
+				const FReal WindingMag = FVec3::DotProduct(FVec3::CrossProduct(Vertex1 - Vertex0, Vertex2 - Vertex1), PlaneN);
+				const FReal Winding = FMath::Sign(WindingMag);
+				const int32 ExpectedWinding = Convex.GetWindingOrder();
+				EXPECT_EQ(Winding, ExpectedWinding) << "PlaneIndex=" << PlaneIndex << " PlaneVertexIndex0=" << PlaneVertexIndex0;
+			}
+		}
+	}
+
+	template<typename ConvexType>
+	void TestConvexEdges(const ConvexType& Convex)
+	{
+		// Check the edges
+		for (int32 EdgeIndex = 0; EdgeIndex < Convex.NumEdges(); ++EdgeIndex)
+		{
+			const int PlaneIndex0 = Convex.GetEdgePlane(EdgeIndex, 0);
+			const int PlaneIndex1 = Convex.GetEdgePlane(EdgeIndex, 1);
+			const int32 VertexIndex0 = Convex.GetEdgeVertex(EdgeIndex, 0);
+			const int32 VertexIndex1 = Convex.GetEdgeVertex(EdgeIndex, 0);
+
+			// Plane0 contains the two vertices
+			bool bFoundVertex0 = false;
+			bool bFoundVertex1 = false;
+			for (int32 PlaneVertexIndex0 = 0; PlaneVertexIndex0 < Convex.NumPlaneVertices(PlaneIndex0); ++PlaneVertexIndex0)
+			{
+				const int32 ThisVertexIndex = Convex.GetPlaneVertex(PlaneIndex0, PlaneVertexIndex0);
+				if (ThisVertexIndex == VertexIndex0)
+				{
+					bFoundVertex0 = true;
+				}
+				if (ThisVertexIndex == VertexIndex1)
+				{
+					bFoundVertex1 = true;
+				}
+			}
+			EXPECT_TRUE(bFoundVertex0) << "EdgeIndex=" << EdgeIndex << "PlaneIndex=" << PlaneIndex0 << " VertexIndex=" << VertexIndex0;
+			EXPECT_TRUE(bFoundVertex1) << "EdgeIndex=" << EdgeIndex << "PlaneIndex=" << PlaneIndex0 << " VertexIndex=" << VertexIndex1;
+
+			// Plane1 contains the two vertices
+			bFoundVertex0 = false;
+			bFoundVertex1 = false;
+			for (int32 PlaneVertexIndex1 = 0; PlaneVertexIndex1 < Convex.NumPlaneVertices(PlaneIndex1); ++PlaneVertexIndex1)
+			{
+				const int32 ThisVertexIndex = Convex.GetPlaneVertex(PlaneIndex1, PlaneVertexIndex1);
+				if (ThisVertexIndex == VertexIndex0)
+				{
+					bFoundVertex0 = true;
+				}
+				if (ThisVertexIndex == VertexIndex1)
+				{
+					bFoundVertex1 = true;
+				}
+			}
+			EXPECT_TRUE(bFoundVertex0) << "EdgeIndex=" << EdgeIndex << "PlaneIndex=" << PlaneIndex1 << " VertexIndex=" << VertexIndex0;
+			EXPECT_TRUE(bFoundVertex1) << "EdgeIndex=" << EdgeIndex << "PlaneIndex=" << PlaneIndex1 << " VertexIndex=" << VertexIndex1;
+		}
+	}
+
+	// Verify that the box Plane Edge and Vertex APIs return the elements exactly as they are defined in Box.cpp
+	GTEST_TEST(ConvexStructureTests, TestBoxStructureDataDetails)
+	{
+		// These arrays are copied from Box.cpp - any changes there should trigger a failure here
+		// so we can be sure the change was expected.
+		const TArray<FVec3> PlaneNormals =
+		{
+			FVec3(-1,  0,  0),		// -X
+			FVec3(0, -1,  0),		// -Y
+			FVec3(0,  0, -1),		// -Z
+			FVec3(1,  0,  0),		//  X
+			FVec3(0,  1,  0),		//  Y
+			FVec3(0,  0,  1),		//  Z
+		};
+
+		const TArray<FVec3> UnitVertices =
+		{
+			FVec3(-1, -1, -1),		// 0 
+			FVec3(1, -1, -1),		// 1 
+			FVec3(-1,  1, -1),		// 2 
+			FVec3(1,  1, -1),		// 3 
+			FVec3(-1, -1,  1),		// 4 
+			FVec3(1, -1,  1),		// 5 
+			FVec3(-1,  1,  1),		// 6 
+			FVec3(1,  1,  1),		// 7 
+		};
+
+		TArray<TArray<int32>> PlaneVertices
+		{
+			{ 0, 4, 6, 2 },			// -X,
+			{ 0, 1, 5, 4 },			// -Y
+			{ 0, 2, 3, 1 },			// -Z
+			{ 1, 3, 7, 5 },			//  X
+			{ 2, 6, 7, 3 },			//  Y
+			{ 4, 5, 7, 6 },			//  Z
+		};
+
+		const FReal NormalTolerance = UE_SMALL_NUMBER;
+		const FReal PositionTolerance = UE_KINDA_SMALL_NUMBER;
+
+		const FVec3 Center = FVec3(0, 0, 0);
+		const FVec3 HalfExtent = FVec3(100, 200, 300);
+		const FReal Margin = FReal(0);
+
+		const FImplicitBox3 Box = FImplicitBox3(Center - HalfExtent, Center + HalfExtent, Margin);
+
+		EXPECT_EQ(Box.NumPlanes(), 6);
+		EXPECT_EQ(Box.NumEdges(), 12);
+		EXPECT_EQ(Box.NumVertices(), 8);
+
+		// Check that the vertices are in the expected order
+		for (int32 VertexIndex = 0; VertexIndex < UnitVertices.Num(); ++VertexIndex)
+		{
+			const FVec3 Vertex = Box.GetVertex(VertexIndex);
+			const FVec3 ExpectedVertex = UnitVertices[VertexIndex] * HalfExtent;
+			EXPECT_NEAR(Vertex.X, ExpectedVertex.X, PositionTolerance);
+			EXPECT_NEAR(Vertex.Y, ExpectedVertex.Y, PositionTolerance);
+			EXPECT_NEAR(Vertex.Z, ExpectedVertex.Z, PositionTolerance);
+		}
+
+		// Check that the planes have the correct normal and position
+		for (int32 PlaneIndex = 0; PlaneIndex < PlaneNormals.Num(); ++PlaneIndex)
+		{
+			TPlaneConcrete<FReal> Plane = Box.GetPlane(PlaneIndex);
+
+			// Normals are in the expected direction
+			EXPECT_NEAR(Plane.Normal().X, PlaneNormals[PlaneIndex].X, NormalTolerance) << "PlaneIndex=" << PlaneIndex;
+			EXPECT_NEAR(Plane.Normal().Y, PlaneNormals[PlaneIndex].Y, NormalTolerance) << "PlaneIndex=" << PlaneIndex;
+			EXPECT_NEAR(Plane.Normal().Z, PlaneNormals[PlaneIndex].Z, NormalTolerance) << "PlaneIndex=" << PlaneIndex;
+
+			// Positions are in the correct plane
+			const FReal PlaneDistance = FVec3::DotProduct(Plane.Normal(), Plane.X());
+			const FReal ExpectedPlaneDistance = FVec3::DotProduct(PlaneNormals[PlaneIndex], PlaneNormals[PlaneIndex] * HalfExtent);
+			EXPECT_NEAR(PlaneDistance, ExpectedPlaneDistance, PositionTolerance);
+		}
+
+		// Check that the planes have the correct vertices
+		for (int32 PlaneIndex = 0; PlaneIndex < PlaneNormals.Num(); ++PlaneIndex)
+		{
+			const int NumPlaneVertices = Box.NumPlaneVertices(PlaneIndex);
+			EXPECT_EQ(NumPlaneVertices, PlaneVertices[PlaneIndex].Num()) << "PlaneIndex=" << PlaneIndex;	// Always 4
+
+			for (int32 PlaneVertexIndex0 = 0; PlaneVertexIndex0 < NumPlaneVertices; ++PlaneVertexIndex0)
+			{
+				const int32 VertexIndex0 = Box.GetPlaneVertex(PlaneIndex, PlaneVertexIndex0);
+				EXPECT_EQ(VertexIndex0, PlaneVertices[PlaneIndex][PlaneVertexIndex0]) << "PlaneIndex=" << PlaneIndex << " PlaneVertexIndex0=" << PlaneVertexIndex0;
+			}
+		}
+
+		// Check the plane vertices are in the plane and have the correct winding order 
+		TestConvexPlaneVertices(Box);
+
+		// Check that the edges report planes that actually share vertices
+		TestConvexEdges(Box);
+	}
+
+	// Check that a Box implemented as a FImplicitConvex3 meets the same specs as ImplicitBox3
+	GTEST_TEST(ConvexStructureTests, TestConvexBoxStructureDataDetails)
+	{
+		const FVec3f Center = FVec3(0, 0, 0);
+		const FVec3f HalfExtent = FVec3(100, 200, 300);
+		const FRealSingle Margin = FReal(0);
+
+		const TArray<FVec3f> Vertices =
+		{
+			Center + HalfExtent * FVec3f(-1, -1, -1),		// 0 
+			Center + HalfExtent * FVec3f( 1, -1, -1),		// 1 
+			Center + HalfExtent * FVec3f(-1,  1, -1),		// 2 
+			Center + HalfExtent * FVec3f( 1,  1, -1),		// 3 
+			Center + HalfExtent * FVec3f(-1, -1,  1),		// 4 
+			Center + HalfExtent * FVec3f( 1, -1,  1),		// 5 
+			Center + HalfExtent * FVec3f(-1,  1,  1),		// 6 
+			Center + HalfExtent * FVec3f( 1,  1,  1),		// 7 
+		};
+
+		FImplicitConvex3 Convex = FImplicitConvex3(Vertices, Margin);
+
+		// Check the plane vertices are in the plane and have the correct winding order 
+		TestConvexPlaneVertices(Convex);
+
+		// Check that the edges report planes that actually share vertices
+		TestConvexEdges(Convex);
+	}
+
 	// The set of vertices generated from a unit box when creating a GeometryCollection from the default box in the editor.
 	// The default cube is tesselated. It has 26 vertices which include the 8 corners, plus mid-points along each edge and in the middle of each face.
 	// 
diff --git a/Engine/Source/Programs/HeadlessChaos/Private/HeadlessChaosTestSupport.cpp b/Engine/Source/Programs/HeadlessChaos/Private/HeadlessChaosTestSupport.cpp
new file mode 100644
--- /dev/null
+++ b/Engine/Source/Programs/HeadlessChaos/Private/HeadlessChaosTestSupport.cpp
@@ -0,0 +1,160 @@
+// Copyright Epic Games, Inc. All Rights Reserved.
+
+#include "HeadlessChaosTestGJK.h"
+
+#include "HeadlessChaos.h"
+#include "HeadlessChaosTestUtility.h"
+#include "Chaos/GJK.h"
+#include "Chaos/Capsule.h"
+#include "Chaos/Convex.h"
+#include "Chaos/GJK.h"
+#include "Chaos/ImplicitObjectScaled.h"
+#include "Chaos/Collision/PBDCollisionConstraint.h"
+#include "Chaos/Triangle.h"
+#include "Chaos/TriangleRegister.h"
+
+namespace ChaosTest
+{
+	using namespace Chaos;
+
+	// InSupportDirection need not be normalized;
+	// InExpectedSupport is any point on the support plane - we only check the distance along the support direction
+	template<typename BoxType>
+	void CheckBoxSupportDir(const BoxType& Box, const FVec3& InSupportDirection, const FVec3& InExpectedSupport)
+	{
+		const FReal DistanceTolerance = UE_KINDA_SMALL_NUMBER;
+
+		const FVec3 SupportDirection = InSupportDirection.GetSafeNormal();
+		const FReal ExpectedSupportDistance = FVec3::DotProduct(InExpectedSupport, SupportDirection);
+
+		// NOTE: Support functions should return a vertex position that is furthest along the support direct.
+		// When there are multiple options, any of the options should be ok
+		// Note: using non-normalized direction
+		const FReal Margin = 0;
+		FReal SupportDelta = 0;
+		int32 SupportVertexIndex = INDEX_NONE;
+		FVec3 Support = Box.SupportCore(InSupportDirection, Margin, &SupportDelta, SupportVertexIndex);
+
+		// The support point has the expected distance along the support direction
+		const FReal SupportDistance = FVec3::DotProduct(Support, SupportDirection);
+		EXPECT_NEAR(SupportDistance, ExpectedSupportDistance, DistanceTolerance);
+
+		// We have a vertex index
+		EXPECT_NE(SupportVertexIndex, INDEX_NONE);
+		if (SupportVertexIndex != INDEX_NONE)
+		{
+			// SupportVertexIndex is for a point with the correct support distance
+			const FVec3 SupportVertex = Box.GetVertex(SupportVertexIndex);
+			const FReal SupportVertexDistance = FVec3::DotProduct(SupportVertex, SupportDirection);
+			EXPECT_NEAR(SupportVertexDistance, ExpectedSupportDistance, DistanceTolerance);
+
+			// The internal Aabb vertex has the correct support distance
+			const FAABB3 Aabb = Box.BoundingBox();
+			const FVec3 AabbVertex = Aabb.GetVertex(SupportVertexIndex);
+			const FReal AabbVertexDistance = FVec3::DotProduct(AabbVertex, SupportDirection);
+			EXPECT_NEAR(AabbVertexDistance, ExpectedSupportDistance, DistanceTolerance);
+
+			// The Aabb and Box vertices are the same
+			EXPECT_NEAR(AabbVertex.X, SupportVertex.X, DistanceTolerance);
+			EXPECT_NEAR(AabbVertex.Y, SupportVertex.Y, DistanceTolerance);
+			EXPECT_NEAR(AabbVertex.Z, SupportVertex.Z, DistanceTolerance);
+
+			// Make sure that GetMostOpposingPlane produces a plane that points in the right direction and contains the support vertex
+			const int32 OpposingPlaneIndex = Box.GetMostOpposingPlane(-InSupportDirection);
+			EXPECT_NE(OpposingPlaneIndex, INDEX_NONE);
+			if (OpposingPlaneIndex != INDEX_NONE)
+			{
+				// Check the direction is actualy in the direction of support
+				const FVec3 OpposingPlaneNormal = Box.GetPlane(OpposingPlaneIndex).Normal();
+				const FReal OpposingNormalDotSupport = FVec3::DotProduct(OpposingPlaneNormal, SupportDirection);
+				EXPECT_GT(OpposingNormalDotSupport, FReal(0)) << "SupportDir=(" << InSupportDirection.X << "," << InSupportDirection.Y << "," << InSupportDirection.Z << ")";
+
+				// Make sure the plane contains the support vertex
+				bool bFoundSupportVertex = false;
+				for (int32 PlaneVertexIndex = 0; PlaneVertexIndex < Box.NumPlaneVertices(OpposingPlaneIndex); ++PlaneVertexIndex)
+				{
+					const int32 VertexIndex = Box.GetPlaneVertex(OpposingPlaneIndex, PlaneVertexIndex);
+					if (VertexIndex == SupportVertexIndex)
+					{
+						bFoundSupportVertex = true;
+						break;
+					}
+				}
+				EXPECT_TRUE(bFoundSupportVertex);
+			}
+		}
+	}
+
+	template<typename BoxType>
+	void CheckBoxSupportVertex(const BoxType& Box, int32 X, int32 Y, int32 Z)
+	{
+		EXPECT_TRUE((X != 0) || (Y != 0) || (Z != 0));
+
+		const FAABB3 Aabb = Box.BoundingBox();
+		const FVec3 UnitCubePoint = FVec3(FReal(X), FReal(Y), FReal(Z));
+		const FVec3 SupportDir = UnitCubePoint;	// No need to normalize
+		const FVec3 SupportVertex = Aabb.GetCenter() + FReal(0.5) * Aabb.Extents() * UnitCubePoint;
+		CheckBoxSupportDir(Box, SupportDir, SupportVertex);
+	}
+
+	template<typename BoxType>
+	void CheckBoxSupportAxes(const BoxType& Box)
+	{
+		// Check all cardinal and diagonal directions
+		for (int32 X = -1; X <= 1; ++X)
+		{
+			for (int32 Y = -1; Y <= 1; ++Y)
+			{
+				for (int32 Z = -1; Z <= 1; ++Z)
+				{
+					if ((X != 0) || (Y != 0) || (Z != 0))
+					{
+						CheckBoxSupportVertex(Box, X, Y, Z);
+					}
+				}
+			}
+		}
+	}
+
+	// Test SupportCore for all cardinal and diagonal axes
+	GTEST_TEST(SupportTests, TestBoxSupportCore)
+	{
+		const FVec3 Center = FVec3(3.5, -94.1, 15.6);
+		const FVec3 HalfExtent = FVec3(100, 200, 300);
+		const FReal Margin = FReal(0);
+
+		const FImplicitBox3 Box = FImplicitBox3(Center - HalfExtent, Center + HalfExtent, Margin);
+		CheckBoxSupportAxes(Box);
+	}
+
+	GTEST_TEST(SupportTests, TestScaledBoxSupportCore)
+	{
+		const FVec3 Center = FVec3(3.5, -94.1, 15.6);
+		const FVec3 HalfExtent = FVec3(100, 200, 300);
+		const FReal Margin = FReal(0);
+		const FVec3 Scale = FVec3(8.7, 3.2, 1.1);
+
+		TRefCountPtr<FImplicitBox3> Box = TRefCountPtr<FImplicitBox3>(new FImplicitBox3(Center - HalfExtent, Center + HalfExtent, Margin));
+		TRefCountPtr<TImplicitObjectScaled<FImplicitBox3>> BoxScaled = TRefCountPtr<TImplicitObjectScaled<FImplicitBox3>>(new TImplicitObjectScaled<FImplicitBox3>(Box, Scale));
+		CheckBoxSupportAxes(*BoxScaled.GetReference());
+	}
+
+	// Test capsule support functions when scaled
+	GTEST_TEST(SupportTests, TestScaledCapsuleSupportCore)
+	{
+		FVec3 PointA{ -1, 0, -100 };
+		FVec3 PointB{ 1, 0, 100 };
+		FVec3 ScaleX{ 100, 1, 1 };
+		FCapsulePtr Capsule(new FImplicitCapsule3(PointA, PointB, 10));
+		TImplicitObjectScaled<FImplicitCapsule3> CapsuleScaled(Capsule, ScaleX);
+
+		FVec3 SupportDir{ 0.1f, 0.0f, -1.0f }; // Pointing down and slightly to the right
+
+		int32 Vertex = INDEX_NONE;
+		FVec3 Support = Capsule->SupportCore(SupportDir, 0, nullptr, Vertex);
+		EXPECT_EQ(Support, PointA); // Expect bottom point
+
+		Support = CapsuleScaled.SupportCore(SupportDir, 0, nullptr, Vertex);
+		EXPECT_EQ(Support, ScaleX * PointA); // Still expect bottom point (But scaled)
+	}
+}
diff --git a/Engine/Source/Runtime/Experimental/Chaos/Private/Chaos/Box.cpp b/Engine/Source/Runtime/Experimental/Chaos/Private/Chaos/Box.cpp
--- a/Engine/Source/Runtime/Experimental/Chaos/Private/Chaos/Box.cpp
+++ b/Engine/Source/Runtime/Experimental/Chaos/Private/Chaos/Box.cpp
@@ -5,41 +5,69 @@
 
 namespace Chaos
 {
+	
+	// Indices below must match TAABB::GetVertex (ascii art below was copied from there)
+	// because TBox uses TAABB::SupportCore which returns a vertex index
+	// NOTE: This shows a right-handed coordinate system, but we are left-handed...
+	//
+	//	   6---------7
+	//	  /|        /|
+	//	 / |       / |
+	//	4---------5  |
+	//	|  |      |  |
+	//	|  2------|--3
+	//	| /       | /
+	//	|/        |/
+	//	0---------1
+	//
+	//  Z
+	//  |  Y
+	//  | /
+	//  |/
+	//  o-----X
+	// 
+	//
+
 	FConvexHalfEdgeStructureDataS16 MakeBoxStructureData()
 	{
+		// NOTE: planes vertices are counter-clockwise when looking at the face from outside 
+		// the box in the diagram above. Plane order defined by SNormals.
 		TArray<TArray<int32>> PlaneVertices
 		{
-			{7, 5, 4, 6},	// X
-			{2, 3, 7, 6},	// Y
-			{5, 7, 3, 1},	// Z,
-			{2, 0, 1, 3},	// -X
-			{0, 4, 5, 1},	// -Y
-			{0, 2, 6, 4},	// -Z,
+			{ 0, 4, 6, 2 },	// -X,
+			{ 0, 1, 5, 4 },	// -Y
+			{ 0, 2, 3, 1 },	// -Z
+			{ 1, 3, 7, 5 },	//  X
+			{ 2, 6, 7, 3 },	//  Y
+			{ 4, 5, 7, 6 },	//  Z
 		};
 
 		return FConvexHalfEdgeStructureDataS16::MakePlaneVertices(PlaneVertices, 8);
 	}
 
+	// NOTE: Plane order must match order implied by TAABB::GetIndex/GetVertex
+	// i.e., negative pointing normals first
 	template<typename T, int D> TArray<FVec3> TBox<T, D>::SNormals =
 	{
-		FVec3(1,0,0),	 // X
-		FVec3(0,1,0),	 // Y
-		FVec3(0,0,1),	 // Z
-		FVec3(-1,0,0), // -X
-		FVec3(0,-1,0), // -Y
-		FVec3(0,0,-1), // -Z
+		FVec3(-1, 0, 0),	// -X
+		FVec3( 0,-1, 0),	// -Y
+		FVec3( 0, 0,-1),	// -Z
+		FVec3( 1, 0, 0),	//  X
+		FVec3( 0, 1, 0),	//  Y
+		FVec3( 0, 0, 1),	//  Z
 	};
 	
+	// NOTE: Vertex order must match order implied by TAABB::GetIndex/GetVertex
 	template<typename T, int D> TArray<FVec3> TBox<T, D>::SVertices =
 	{
-		FVec3(-1,-1,-1),
-		FVec3(-1,-1,1),
-		FVec3(-1,1,-1),
-		FVec3(-1,1,1),
-		FVec3(1,-1,-1),
-		FVec3(1,-1,1),
-		FVec3(1,1,-1),
-		FVec3(1,1,1),
+		FVec3(-1,-1,-1),	// 0 
+		FVec3( 1,-1,-1),	// 1 
+		FVec3(-1, 1,-1),	// 2 
+		FVec3( 1, 1,-1),	// 3 
+		FVec3(-1,-1, 1),	// 4 
+		FVec3( 1,-1, 1),	// 5 
+		FVec3(-1, 1, 1),	// 6 
+		FVec3( 1, 1, 1),	// 7 
 	};
 
 	template<typename T, int D> FConvexHalfEdgeStructureDataS16 TBox<T, D>::SStructureData = MakeBoxStructureData();
diff --git a/Engine/Source/Runtime/Experimental/Chaos/Public/Chaos/AABB.h b/Engine/Source/Runtime/Experimental/Chaos/Public/Chaos/AABB.h
--- a/Engine/Source/Runtime/Experimental/Chaos/Public/Chaos/AABB.h
+++ b/Engine/Source/Runtime/Experimental/Chaos/Public/Chaos/AABB.h
@@ -311,7 +311,7 @@ namespace Chaos
 			return BestNormal;
 		}
 		
-		FORCEINLINE_DEBUGGABLE TVector<T, d> Support(const TVector<T, d>& Direction, const T Thickness, int32& VertexIndex) const
+		FORCEINLINE_DEBUGGABLE TVector<T, d> Support(const TVector<T, d>& Direction, const T Thickness, int32& OutVertexIndex) const
 		{
 			TVector<T, d> ChosenPt;
 			FIntVector ChosenAxis;
@@ -328,7 +328,8 @@ namespace Chaos
 					ChosenAxis[Axis] = 1;
 				}
 			}
-			VertexIndex =  ChosenAxis[0] * 4 + ChosenAxis[1] * 2 + ChosenAxis[2];
+
+			OutVertexIndex = GetIndex(ChosenAxis);
 
 			if (Thickness != (T)0)
 			{
@@ -349,7 +350,7 @@ namespace Chaos
 		}
 
 		// Support vertex in the specified direction, assuming each face has been moved inwards by InMargin
-		FORCEINLINE_DEBUGGABLE FVec3 SupportCore(const FVec3& Direction, const FReal InMargin, FReal* OutSupportDelta, int32& VertexIndex) const
+		FORCEINLINE_DEBUGGABLE FVec3 SupportCore(const FVec3& Direction, const FReal InMargin, FReal* OutSupportDelta, int32& OutVertexIndex) const
 		{
 			FVec3 ChosenPt;
 			FIntVector ChosenAxis;
@@ -366,7 +367,9 @@ namespace Chaos
 					ChosenAxis[Axis] = 1;
 				}
 			}
-			VertexIndex = ChosenAxis[0] * 4 + ChosenAxis[1] * 2 + ChosenAxis[2];
+
+			OutVertexIndex = GetIndex(ChosenAxis);
+
 			// Maximum distance between the Core+Margin position and the original outer vertex
 			constexpr FReal RootThreeMinusOne = FReal(1.7320508075688772935274463415059 - 1.0);
 			if (OutSupportDelta != nullptr)
@@ -386,7 +389,7 @@ namespace Chaos
 			return MakeVectorRegisterFloatFromDouble(MakeVectorRegister(SupportVert.X, SupportVert.Y, SupportVert.Z, 0.0));
 		}
 
-		FORCEINLINE_DEBUGGABLE TVector<T, d> SupportCoreScaled(const TVector<T, d>& Direction, const T InMargin, const TVector<T, d>& Scale, T* OutSupportDelta, int32& VertexIndex) const
+		FORCEINLINE_DEBUGGABLE TVector<T, d> SupportCoreScaled(const TVector<T, d>& Direction, const T InMargin, const TVector<T, d>& Scale, T* OutSupportDelta, int32& OutVertexIndex) const
 		{
 			const TVector<T, d> ScaledDirection = Direction * Scale; // Compensate for Negative scales, only the signs really matter here
 
@@ -405,8 +408,9 @@ namespace Chaos
 					ChosenAxis[Axis] = 1;
 				}
 			}
-			VertexIndex = ChosenAxis[0] * 4 + ChosenAxis[1] * 2 + ChosenAxis[2];
-			
+
+			OutVertexIndex = GetIndex(ChosenAxis);
+
 			constexpr T RootThreeMinusOne = T(1.7320508075688772935274463415059 - 1.0);
 			if (OutSupportDelta != nullptr)
 			{
@@ -495,12 +499,26 @@ namespace Chaos
 		FORCEINLINE TVector<T, d> GetVertex(const int32 Index) const
 		{
 			check(0 <= Index && Index < 8);
+
+			// See GetIndex() for reverse logic
 			return TVector<T, d>(
 				(Index & (1 << 0)) == 0 ? MMin.X : MMax.X,
 				(Index & (1 << 1)) == 0 ? MMin.Y : MMax.Y,
 				(Index & (1 << 2)) == 0 ? MMin.Z : MMax.Z);
 		}
 
+		/**
+		* Given a point on a unit cube expressed as an IntVector of 0s and 1s, get the vertex index of that point.
+		* 0 represents a negative axis and 1 represents a positive axis. All other values are invalid, but we do 
+		* not check this as it would add overhead to support functions. This is used by the support functions and 
+		* must invert the index logic in GetVertex().
+		*/
+		FORCEINLINE int32 GetIndex(const FIntVector& AxisSelector) const
+		{
+			// See GetVertex() for reverse logic
+			return AxisSelector[0] + AxisSelector[1] * 2 + AxisSelector[2] * 4;
+		}
+
 		/**
 		*	Get an array of two indices into the vertex list for one of the
 		*	twelve edges of the box. Edges are ordered by increasing vertex
diff --git a/Engine/Source/Runtime/Experimental/Chaos/Public/Chaos/Box.h b/Engine/Source/Runtime/Experimental/Chaos/Public/Chaos/Box.h
--- a/Engine/Source/Runtime/Experimental/Chaos/Public/Chaos/Box.h
+++ b/Engine/Source/Runtime/Experimental/Chaos/Public/Chaos/Box.h
@@ -201,8 +201,9 @@ namespace Chaos
 		// Get the index of the plane that most opposes the normal
 		int32 GetMostOpposingPlane(const TVector<T, d>& Normal) const
 		{
+			// NOTE: Index returned here must match indices as set up in SNormals (see Box.cpp)
 			int32 AxisIndex = FVec3(Normal.GetAbs()).MaxAxis();
-			if (Normal[AxisIndex] > 0.0f)
+			if (Normal[AxisIndex] < 0.0f)
 			{
 				AxisIndex += 3;
 			}
