Chaos - GJKDistance InitialV optimization for same-space shapes

diff --git a/Engine/Source/Programs/HeadlessChaos/Private/HeadlessChaosTestGJKDistance.cpp b/Engine/Source/Programs/HeadlessChaos/Private/HeadlessChaosTestGJKDistance.cpp
--- a/Engine/Source/Programs/HeadlessChaos/Private/HeadlessChaosTestGJKDistance.cpp
+++ b/Engine/Source/Programs/HeadlessChaos/Private/HeadlessChaosTestGJKDistance.cpp
@@ -28,7 +28,7 @@ namespace ChaosTest
 			EGJKDistanceResult Result = GJKDistance<FReal>(
 				TGJKCoreShape(A),
 				TGJKCoreShapeTransformed(B, BToATm),
-				GJKDistanceInitialV(A, B, BToATm),
+				GJKDistanceInitialVFromRelativeTransform(A, B, BToATm),
 				Distance, NearestA, NearestB, Normal);
 			EXPECT_NE(Result, EGJKDistanceResult::Separated);
 		}
@@ -41,7 +41,7 @@ namespace ChaosTest
 			EGJKDistanceResult Result = GJKDistance<FReal>(
 				TGJKCoreShape(A), 
 				TGJKCoreShapeTransformed(B, BToATm),
-				GJKDistanceInitialV(A, B, BToATm),
+				GJKDistanceInitialVFromRelativeTransform(A, B, BToATm),
 				Distance, NearestA, NearestB, Normal);
 			EXPECT_EQ(Result, EGJKDistanceResult::Separated);
 			EXPECT_NEAR(Distance, (FReal)1, Tolerance);
@@ -62,7 +62,7 @@ namespace ChaosTest
 			EGJKDistanceResult Result = GJKDistance<FReal>(
 				TGJKCoreShape(A),
 				TGJKCoreShapeTransformed(B, BToATm),
-				GJKDistanceInitialV(A, B, BToATm),
+				GJKDistanceInitialVFromRelativeTransform(A, B, BToATm),
 				Distance, NearestA, NearestB, Normal);
 			EXPECT_EQ(Result, EGJKDistanceResult::Separated);
 			FVec3 CenterDelta = (B.GetCenter() + BPos) - A.GetCenter();
@@ -85,7 +85,7 @@ namespace ChaosTest
 			EGJKDistanceResult Result = GJKDistance<FReal>(
 				TGJKCoreShape(A),
 				TGJKCoreShapeTransformed(B, BToATm),
-				GJKDistanceInitialV(A, B, BToATm),
+				GJKDistanceInitialVFromRelativeTransform(A, B, BToATm),
 				Distance, NearestA, NearestB, Normal);
 			EXPECT_EQ(Result, EGJKDistanceResult::Separated);
 			EXPECT_NEAR(Distance, (FReal)1 - BPos.X, Tolerance);
@@ -120,7 +120,7 @@ namespace ChaosTest
 			EGJKDistanceResult Result = GJKDistance<FReal>(
 				TGJKShape(A),
 				TGJKCoreShapeTransformed(B, BToATm),
-				GJKDistanceInitialV(A, B, BToATm),
+				GJKDistanceInitialVFromRelativeTransform(A, B, BToATm),
 				Distance, NearestA, NearestB, Normal);
 			EXPECT_NE(Result, EGJKDistanceResult::Separated);
 		}
@@ -133,7 +133,7 @@ namespace ChaosTest
 			EGJKDistanceResult Result = GJKDistance<FReal>(
 				TGJKShape(A),
 				TGJKCoreShapeTransformed(B, BToATm),
-				GJKDistanceInitialV(A, B, BToATm),
+				GJKDistanceInitialVFromRelativeTransform(A, B, BToATm),
 				Distance, NearestA, NearestB, Normal);
 			EXPECT_EQ(Result, EGJKDistanceResult::Separated);
 			EXPECT_NEAR(Distance, (FReal)1, Tolerance);
@@ -152,7 +152,7 @@ namespace ChaosTest
 			EGJKDistanceResult Result = GJKDistance<FReal>(
 				TGJKShape(A),
 				TGJKCoreShapeTransformed(B, BToATm),
-				GJKDistanceInitialV(A, B, BToATm),
+				GJKDistanceInitialVFromRelativeTransform(A, B, BToATm),
 				Distance, NearestA, NearestB, Normal);
 			EXPECT_EQ(Result, EGJKDistanceResult::Separated);
 			EXPECT_NEAR(Distance, (FReal)1, Tolerance);
@@ -172,7 +172,7 @@ namespace ChaosTest
 			EGJKDistanceResult Result = GJKDistance<FReal>(
 				TGJKShape(A),
 				TGJKCoreShapeTransformed(B, BToATm),
-				GJKDistanceInitialV(A, B, BToATm),
+				GJKDistanceInitialVFromRelativeTransform(A, B, BToATm),
 				Distance, NearestA, NearestB, Normal);
 			FVec3 NearPointOnA = A.Min();
 			FVec3 SphereNearPointDir = (NearPointOnA - B.GetCenter()).GetSafeNormal();
@@ -194,7 +194,7 @@ namespace ChaosTest
 			EGJKDistanceResult Result = GJKDistance<FReal>(
 				TGJKShape(A),
 				TGJKCoreShapeTransformed(B, BToATm),
-				GJKDistanceInitialV(A, B, BToATm),
+				GJKDistanceInitialVFromRelativeTransform(A, B, BToATm),
 				Distance, NearestA, NearestB, Normal);
 			FVec3 NearPointOnA = A.Min();
 			FVec3 SphereNearPointDir = (NearPointOnA - B.GetCenter()).GetSafeNormal();
@@ -217,7 +217,7 @@ namespace ChaosTest
 			EGJKDistanceResult Result = GJKDistance<FReal>(
 				TGJKShape(A),
 				TGJKCoreShapeTransformed(B, BToATm),
-				GJKDistanceInitialV(A, B, BToATm),
+				GJKDistanceInitialVFromRelativeTransform(A, B, BToATm),
 				Distance, NearestA, NearestB, Normal);
 			FVec3 NearPointOnA = FVec3(4, 0, 0);
 			FVec3 BPos = BToATm.TransformPositionNoScale(B.GetCenter());
@@ -241,7 +241,7 @@ namespace ChaosTest
 			EGJKDistanceResult Result = GJKDistance<FReal>(
 				TGJKCoreShape(B),
 				TGJKShapeTransformed(A, BToATm),
-				GJKDistanceInitialV(B, A, BToATm),
+				GJKDistanceInitialVFromRelativeTransform(B, A, BToATm),
 				Distance, NearestB, NearestA, Normal);
 			FVec3 NearPointOnA = FVec3(4, 0, 4);
 			FVec3 BPos = BToATm.InverseTransformPositionNoScale(B.GetCenter());
@@ -270,7 +270,7 @@ namespace ChaosTest
 			EGJKDistanceResult Result = GJKDistance<FReal>(
 				TGJKShape(A),
 				TGJKCoreShapeTransformed(B, BToATm),
-				GJKDistanceInitialV(A, B, BToATm),
+				GJKDistanceInitialVFromRelativeTransform(A, B, BToATm),
 				Distance, NearestA, NearestB, Normal);
 			FVec3 NearPointOnA = FVec3(5, 0, 2);
 			FVec3 NearPointDir = (NearPointOnA - B.GetCenter()).GetSafeNormal();
@@ -306,7 +306,7 @@ namespace ChaosTest
 			EGJKDistanceResult Result = GJKDistance<FReal>(
 				TGJKShape(A),
 				TGJKCoreShapeTransformed(B, BToATm),
-				GJKDistanceInitialV(A, B, BToATm),
+				GJKDistanceInitialVFromRelativeTransform(A, B, BToATm),
 				Distance, NearestA, NearestB, Normal);
 			EXPECT_NE(Result, EGJKDistanceResult::Separated);
 		}
@@ -319,7 +319,7 @@ namespace ChaosTest
 			EGJKDistanceResult Result = GJKDistance<FReal>(
 				TGJKShape(A),
 				TGJKCoreShapeTransformed(B, BToATm),
-				GJKDistanceInitialV(A, B, BToATm),
+				GJKDistanceInitialVFromRelativeTransform(A, B, BToATm),
 				Distance, NearestA, NearestB, Normal);
 			const FReal Tolerance = (FReal)2e-3;
 			EXPECT_EQ(Result, EGJKDistanceResult::Separated);
@@ -342,7 +342,7 @@ namespace ChaosTest
 			EGJKDistanceResult Result = GJKDistance<FReal>(
 				TGJKShape(A),
 				TGJKCoreShapeTransformed(B, BToATm),
-				GJKDistanceInitialV(A, B, BToATm),
+				GJKDistanceInitialVFromRelativeTransform(A, B, BToATm),
 				Distance, NearestA, NearestB, Normal);
 			FVec3 ExpectedNearestA = FVec3(5, 0, -2);
 			FVec3 ExpectedDir = (ExpectedNearestA - B.GetX2()).GetSafeNormal();
@@ -367,7 +367,7 @@ namespace ChaosTest
 			EGJKDistanceResult Result = GJKDistance<FReal>(
 				TGJKCoreShape(A),
 				TGJKShapeTransformed(B, BToATm),
-				GJKDistanceInitialV(A, B, BToATm),
+				GJKDistanceInitialVFromRelativeTransform(A, B, BToATm),
 				Distance, NearestA, NearestB, Normal);
 			FVec3 ExpectedNearestA = FVec3(3, 0, (FReal)-1.5);
 			FVec3 ExpectedNearestB = FVec3(2, 0, (FReal)-1.5);
@@ -393,7 +393,7 @@ namespace ChaosTest
 			EGJKDistanceResult Result = GJKDistance<FReal>(
 				TGJKCoreShape(A),
 				TGJKShapeTransformed(B, BToATm),
-				GJKDistanceInitialV(A, B, BToATm),
+				GJKDistanceInitialVFromRelativeTransform(A, B, BToATm),
 				Distance, NearestA, NearestB, Normal);
 			FVec3 ExpectedNearestA = FVec3(-1, 0, (FReal)2);
 			FVec3 ExpectedNearestB = FVec3(2, 0, (FReal)-2);
@@ -436,7 +436,7 @@ namespace ChaosTest
 			EGJKDistanceResult Result = GJKDistance<FReal>(
 				TGJKShape(A),
 				TGJKCoreShapeTransformed(B, BToATm),
-				GJKDistanceInitialV(A, B, BToATm),
+				GJKDistanceInitialVFromRelativeTransform(A, B, BToATm),
 				Distance, NearestA, NearestB, Normal);
 			FVec3 ExpectedNearestA = FVec3(5, 0, -2);
 			FVec3 ExpectedDir = (ExpectedNearestA - B.GetX2()).GetSafeNormal();
diff --git a/Engine/Source/Runtime/Experimental/Chaos/Private/GeometryCollection/GeometryCollectionProximityUtility.cpp b/Engine/Source/Runtime/Experimental/Chaos/Private/GeometryCollection/GeometryCollectionProximityUtility.cpp
--- a/Engine/Source/Runtime/Experimental/Chaos/Private/GeometryCollection/GeometryCollectionProximityUtility.cpp
+++ b/Engine/Source/Runtime/Experimental/Chaos/Private/GeometryCollection/GeometryCollectionProximityUtility.cpp
@@ -469,13 +469,12 @@ void BuildProximityFromConvexHulls(FGeometryCollection* Collection, const UE::Ge
 					//	const VectorRegister4Float InitialDirSimd = MakeVectorRegisterFloat(1.f, 0.f, 0.f, 0.f);
 					//	if (GJKIntersectionSameSpaceSimd(Hull, CandidateHull, DistanceThreshold, InitialDirSimd))
 
-					const Chaos::FRigidTransform3 IdentityTransform = Chaos::FRigidTransform3::Identity;
 					Chaos::FReal Distance;
 					Chaos::TVec3<Chaos::FReal> NearestA, NearestB, Normal; // All unused
 					Chaos::EGJKDistanceResult Result = Chaos::GJKDistance<Chaos::FReal>(
 						Chaos::TGJKShape(Hull),
 						Chaos::TGJKShape(CandidateHull),
-						Chaos::GJKDistanceInitialV(Hull, CandidateHull, IdentityTransform),
+						Chaos::GJKDistanceInitialVFromDirection(Hull, CandidateHull, CandidateHull.GetCenterOfMass() - Hull.GetCenterOfMass()),
 						Distance, NearestA, NearestB, Normal);
 					if (Result == Chaos::EGJKDistanceResult::Contact || Result == Chaos::EGJKDistanceResult::DeepContact
 						|| (Result == Chaos::EGJKDistanceResult::Separated && Distance <= DistanceThreshold))
diff --git a/Engine/Source/Runtime/Experimental/Chaos/Public/Chaos/GJK.h b/Engine/Source/Runtime/Experimental/Chaos/Public/Chaos/GJK.h
--- a/Engine/Source/Runtime/Experimental/Chaos/Public/Chaos/GJK.h
+++ b/Engine/Source/Runtime/Experimental/Chaos/Public/Chaos/GJK.h
@@ -1939,27 +1939,52 @@ namespace Chaos
 		return GJKRaycast2Impl(A, B, StartTM, RayDir, RayLength, OutTime, OutPosition, OutNormal, GivenThicknessA, bComputeMTD, InitialDir, GivenThicknessB);
 	}
 
+	template <typename T, typename TGeometryA, typename TGeometryB>
+	UE_DEPRECATED(5.5, "Use GJKDistanceInitialVFromDirection if possible, or GJKDistanceInitialVFromRelativeTransform if original behaviour was required")
+	TVector<T, 3> GJKDistanceInitialV(const TGeometryA& A, const TGeometryB& B, const TRigidTransform<T, 3>& BToATM)
+	{
+		return GJKDistanceInitialVFromRelativeTransform(A, B, BToATM);
+	}
 
-	/**
-	 * Can be used to generate the initial support direction for use with GjkDistance. Returns a point on the Minkowski Sum
-	 * surface opposite to the direction of the supplied transform. This is usually a good guess for the initial direction
-	 * but it calls SupportCore on both shapes, and there are often faster alternatives if you know the type of shapes
-	 * you are dealing with (e.g., return the vectors between the centers of the two convex shapes).
-	 * 
+	// Avoid this function - use GJKDistanceInitialVFromDirection that takes a vector direction and assumes the geometries are in the 
+	// same space. If your geometries are not in the same space, they can be wrapped in TGJKShapeTransformed or TGJKCoreShapeTransformed
+	// and you would already have these available because you need them for GJKDistance
+	template <typename T, typename TGeometryA, typename TGeometryB>
+	TVector<T, 3> GJKDistanceInitialVFromRelativeTransform(const TGeometryA& A, const TGeometryB& B, const TRigidTransform<T, 3>& BToATM)
+	{
+		FVec3 Direction = BToATM.GetTranslation();
+		if (Direction.IsZero())
+		{
+			Direction = TVec3<T>(1, 0, 0);
+		}
+		const TVector<T, 3> DirectionInB = BToATM.GetRotation().Inverse() * Direction;
+
+		int32 UnusedVertexIndex = INDEX_NONE;
+		const TVector<T, 3> SupportA = A.SupportCore(Direction, A.GetMargin(), nullptr, UnusedVertexIndex);
+		const TVector<T, 3> SupportBLocal = B.SupportCore(-DirectionInB, B.GetMargin(), nullptr, UnusedVertexIndex);
+		const TVector<T, 3> SupportB = BToATM.TransformPositionNoScale(SupportBLocal);
+		return SupportA - SupportB;
+	}
+
+	/*
+	 * Can be used to generate the initial support direction for use with GJKDistance. Returns a point on the Minkowski Sum
+	 * surface opposite to the direction of the supplied Direction. This is usually a good guess for the initial direction
+	 * but it calls SupportCore on both shapes which in O(N) in the number of vertices, and there are often faster alternatives 
+	 * if you know the type of shapes you are dealing with (e.g., return the vectors between the centers of the two convex shapes).
+	 *
 	 * If you do roll your own function, make sure that the vector returned is in or on the Minkowski sum (and don't just use a unit
 	 * vector along some direction for example) or GJKDistance may early-exit with an inaccurate result.
 	 */
 	template <typename T, typename TGeometryA, typename TGeometryB>
-	TVector<T, 3> GJKDistanceInitialV(const TGeometryA& A, const TGeometryB& B, const TRigidTransform<T, 3>& BToATM)
+	TVec3<T> GJKDistanceInitialVFromDirection(const TGeometryA& A, const TGeometryB& B, TVec3<T> Direction)
 	{
-		const T MarginA = A.GetMargin();
-		const T MarginB = B.GetMargin();
-		int32 VertexIndexA = INDEX_NONE, VertexIndexB = INDEX_NONE;
-		const TVec3<T> V = -BToATM.GetTranslation();
-		const TVector<T, 3> SupportA = A.SupportCore(-V, MarginA, nullptr, VertexIndexA);
-		const TVector<T, 3> VInB = BToATM.GetRotation().Inverse() * V;
-		const TVector<T, 3> SupportBLocal = B.SupportCore(VInB, MarginB, nullptr, VertexIndexB);
-		const TVector<T, 3> SupportB = BToATM.TransformPositionNoScale(SupportBLocal);
+		if (Direction.IsZero())
+		{
+			Direction = TVec3<T>(1, 0, 0);
+		}
+		int32 UnusedVertexIndex = INDEX_NONE;
+		const TVec3<T> SupportA = A.SupportCore(Direction, A.GetMargin(), nullptr, UnusedVertexIndex);
+		const TVec3<T> SupportB = B.SupportCore(-Direction, B.GetMargin(), nullptr, UnusedVertexIndex);
 		return SupportA - SupportB;
 	}
 
@@ -1993,11 +2018,10 @@ namespace Chaos
 	 *
 	 * @param A The first object (usually TGJKShape or TGJKCoreShape)
 	 * @param B The second object (usually TGJKShapeTransformed or TGJKCoreShapeTransformed)
-	 * @param B The second object.
-	 * @param InitialV  Starting support direction that must be in the Minkowski Sum. Use GJKDistanceInitialV() if unsure.
-	 * @param OutDistance if returns true, the minimum distance between A and B, otherwise not modified.
-	 * @param OutNearestA if returns true, the near point on A in local-space, otherwise not modified.
-	 * @param OutNearestB if returns true, the near point on B in local-space, otherwise not modified.
+	 * @param InitialV  Starting support direction that must be in or on the Minkowski Sum. See GJKDistanceInitialVFromDirection().
+	 * @param OutDistance If result is not DeepContact, the minimum distance between A and B, otherwise not modified.
+	 * @param OutNearestA If result is not DeepContact, the near point on A in local-space, otherwise not modified.
+	 * @param OutNearestB If result is not DeepContact, the near point on B in local-space, otherwise not modified.
 	 * @param Epsilon The algorithm terminates when the iterative distance reduction gets below this threshold.
 	 * @param MaxIts A limit on the number of iterations. Results may be approximate if this is too low.
 	 * @return EGJKDistanceResult - see comments on the enum
