CIS - Fixup HeadlessChaos after deprecations to rest collection storage and lifetime management within dynamic collections

diff --git a/Engine/Source/Programs/HeadlessChaos/Private/GeometryCollection/GeometryCollectionTestFramework.cpp b/Engine/Source/Programs/HeadlessChaos/Private/GeometryCollection/GeometryCollectionTestFramework.cpp
--- a/Engine/Source/Programs/HeadlessChaos/Private/GeometryCollection/GeometryCollectionTestFramework.cpp
+++ b/Engine/Source/Programs/HeadlessChaos/Private/GeometryCollection/GeometryCollectionTestFramework.cpp
@@ -148,11 +148,11 @@ namespace GeometryCollectionTest
 		InitMaterialToZero(NewHandle.Get());
 		Chaos::FPhysicalMaterialManager::Get().UpdateMaterial(NewHandle);
 
-		TSharedPtr<FGeometryDynamicCollection> DynamicCollection = GeometryCollectionToGeometryDynamicCollection(RestCollection.Get(), Params.DynamicState);
+		TSharedPtr<FGeometryDynamicCollection> DynamicCollection = GeometryCollectionToGeometryDynamicCollection(RestCollection, Params.DynamicState);
 
 		FSimulationParameters SimulationParams;
 		{
-			SimulationParams.RestCollection = RestCollection.Get();
+			SimulationParams.RestCollectionShared = RestCollection;
 			SimulationParams.PhysicalMaterialHandle = NewHandle;
 			SimulationParams.Shared.Mass = Params.Mass;
 			SimulationParams.Shared.bMassAsDensity = Params.bMassAsDensity;
diff --git a/Engine/Source/Programs/HeadlessChaos/Private/GeometryCollection/GeometryCollectionTestUtility.cpp b/Engine/Source/Programs/HeadlessChaos/Private/GeometryCollection/GeometryCollectionTestUtility.cpp
--- a/Engine/Source/Programs/HeadlessChaos/Private/GeometryCollection/GeometryCollectionTestUtility.cpp
+++ b/Engine/Source/Programs/HeadlessChaos/Private/GeometryCollection/GeometryCollectionTestUtility.cpp
@@ -19,7 +19,7 @@
 
 namespace GeometryCollectionTest {
 
-	TSharedPtr<FGeometryDynamicCollection> GeometryCollectionToGeometryDynamicCollection(const FGeometryCollection* InputCollection, EObjectStateTypeEnum DynamicStateDefault)
+	TSharedPtr<FGeometryDynamicCollection> GeometryCollectionToGeometryDynamicCollection(TSharedPtr<const FGeometryCollection> InputCollection, EObjectStateTypeEnum DynamicStateDefault)
 	{
 		TSharedPtr<FGeometryDynamicCollection> NewCollection(new FGeometryDynamicCollection(InputCollection));
 		// TransformAttribute will be initialized by a lazy initializer pattern
diff --git a/Engine/Source/Programs/HeadlessChaos/Public/GeometryCollection/GeometryCollectionTestUtility.h b/Engine/Source/Programs/HeadlessChaos/Public/GeometryCollection/GeometryCollectionTestUtility.h
--- a/Engine/Source/Programs/HeadlessChaos/Public/GeometryCollection/GeometryCollectionTestUtility.h
+++ b/Engine/Source/Programs/HeadlessChaos/Public/GeometryCollection/GeometryCollectionTestUtility.h
@@ -15,7 +15,7 @@ struct FSimulationParameters;
 
 namespace GeometryCollectionTest {
 
-	TSharedPtr<FGeometryDynamicCollection> GeometryCollectionToGeometryDynamicCollection(const FGeometryCollection* InputCollection, EObjectStateTypeEnum DynamicStateDefault = EObjectStateTypeEnum::Chaos_Object_Dynamic);
+	TSharedPtr<FGeometryDynamicCollection> GeometryCollectionToGeometryDynamicCollection(TSharedPtr<const FGeometryCollection> InputCollection, EObjectStateTypeEnum DynamicStateDefault = EObjectStateTypeEnum::Chaos_Object_Dynamic);
 
 	TSharedPtr<FGeometryCollection>	CreateClusteredBody(FVector Position);
 	TSharedPtr<FGeometryCollection>	CreateClusteredBody_OnePartent_FourBodies(FVector Position);
