Fix Rotation drive constraint
Acceleration mode were not considered
Add a unit test

diff --git a/Engine/Source/Programs/HeadlessChaos/Private/HeadlessChaosTestJointForces.cpp b/Engine/Source/Programs/HeadlessChaos/Private/HeadlessChaosTestJointForces.cpp
--- a/Engine/Source/Programs/HeadlessChaos/Private/HeadlessChaosTestJointForces.cpp
+++ b/Engine/Source/Programs/HeadlessChaos/Private/HeadlessChaosTestJointForces.cpp
@@ -291,60 +291,66 @@ namespace ChaosTest {
 
 	GTEST_TEST(JointForceTests, TestAngularDriveForceMode_Damping)
 	{
-		const int32 NumSolverIterations = 20;
-		const FReal Gravity = 0;
-		const FReal Dt = 0.01;
-		const int32 NumBodies = 2;
+		for (int32 IndexSimd = 0; IndexSimd < 2; IndexSimd++)
+		{
+			bool bUseSimd = IndexSimd == 0;
+			const int32 NumSolverIterations = 20;
+			const FReal Gravity = 0;
+			const FReal Dt = 0.01;
+			const int32 NumBodies = 2;
 
-		const FReal AngularVelocity = 3;
-		const FReal Stiffness = 0;
-		const FReal Damping = 200;
+			const FReal AngularVelocity = 3;
+			const FReal Stiffness = 0;
+			const FReal Damping = 200;
 
-		FJointChainTest<FPBDRigidsEvolutionGBF> Test(NumSolverIterations, Gravity);
-		Test.InitChain(NumBodies, FVec3(0, 0, -1));
+			FJointChainTest<FPBDRigidsEvolutionGBF> Test(NumSolverIterations, Gravity);
+			Test.InitChain(NumBodies, FVec3(0, 0, -1));
 
-		if (!Test.Evolution.GetJointConstraints().GetSettings().bUsePositionBasedDrives)
-		{
-			Test.Evolution.SetNumVelocityIterations(NumSolverIterations);
-		}
+			Test.Evolution.GetJointConstraints().SetUseSimd(bUseSimd);
 
-		// Disable all limits
-		Test.JointSettings[0].LinearMotionTypes = { EJointMotionType::Free, EJointMotionType::Free, EJointMotionType::Free };
-		Test.JointSettings[0].AngularMotionTypes = { EJointMotionType::Free, EJointMotionType::Free, EJointMotionType::Free };
+			if (!Test.Evolution.GetJointConstraints().GetSettings().bUsePositionBasedDrives)
+			{
+				Test.Evolution.SetNumVelocityIterations(NumSolverIterations);
+			}
 
-		// Set up the drive in force mode
-		Test.JointSettings[0].bAngularSLerpVelocityDriveEnabled = true;
-		Test.JointSettings[0].AngularDriveForceMode = EJointForceMode::Force;
-		Test.JointSettings[0].AngularDriveStiffness = FVec3(Stiffness, Stiffness, Stiffness);
-		Test.JointSettings[0].AngularDriveDamping = FVec3(Damping, Damping, Damping);
+			// Disable all limits
+			Test.JointSettings[0].LinearMotionTypes = { EJointMotionType::Free, EJointMotionType::Free, EJointMotionType::Free };
+			Test.JointSettings[0].AngularMotionTypes = { EJointMotionType::Free, EJointMotionType::Free, EJointMotionType::Free };
 
-		Test.Create();
+			// Set up the drive in force mode
+			Test.JointSettings[0].bAngularSLerpVelocityDriveEnabled = true;
+			Test.JointSettings[0].AngularDriveForceMode = EJointForceMode::Force;
+			Test.JointSettings[0].AngularDriveStiffness = FVec3(Stiffness, Stiffness, Stiffness);
+			Test.JointSettings[0].AngularDriveDamping = FVec3(Damping, Damping, Damping);
 
-		FGenericParticleHandle P1 = Test.GetParticle(1);
+			Test.Create();
 
-		// Give the particle angular velocity so that joint damping has some work to do
-		P1->SetW(FVec3(0, 0, -AngularVelocity));
+			FGenericParticleHandle P1 = Test.GetParticle(1);
 
-		// Run the sim
-		Test.Evolution.AdvanceOneTimeStep(Dt);
-		Test.Evolution.EndFrame(Dt);
+			// Give the particle angular velocity so that joint damping has some work to do
+			P1->SetW(FVec3(0, 0, -AngularVelocity));
 
-		// Calculate expected Trque from T = -D.W with implicit integration
-		const FReal I = FConstGenericParticleHandle(Test.GetParticle(1))->I().Z;
-		FReal ExpectedTorqueZ = 0;
-		FReal DQ = 0;
-		for (int32 It = 0; It < NumSolverIterations; ++It)
-		{
-			const FReal W = -AngularVelocity + DQ / Dt;
-			const FReal T = -Damping * W;
-			const FReal DW = ((T - ExpectedTorqueZ) / I) * Dt;
-			DQ += DW * Dt;
-			ExpectedTorqueZ = T;
-		}
+			// Run the sim
+			Test.Evolution.AdvanceOneTimeStep(Dt);
+			Test.Evolution.EndFrame(Dt);
 
-		// Check the joint forces agree
-		const FReal TorqueZ = -Test.Evolution.GetJointConstraints().GetConstraintAngularImpulse(0).Z / Dt;
-		EXPECT_NEAR(TorqueZ, ExpectedTorqueZ, 0.1);
+			// Calculate expected Trque from T = -D.W with implicit integration
+			const FReal I = FConstGenericParticleHandle(Test.GetParticle(1))->I().Z;
+			FReal ExpectedTorqueZ = 0;
+			FReal DQ = 0;
+			for (int32 It = 0; It < NumSolverIterations; ++It)
+			{
+				const FReal W = -AngularVelocity + DQ / Dt;
+				const FReal T = -Damping * W;
+				const FReal DW = ((T - ExpectedTorqueZ) / I) * Dt;
+				DQ += DW * Dt;
+				ExpectedTorqueZ = T;
+			}
+
+			// Check the joint forces agree
+			const FReal TorqueZ = -Test.Evolution.GetJointConstraints().GetConstraintAngularImpulse(0).Z / Dt;
+			EXPECT_NEAR(TorqueZ, ExpectedTorqueZ, 0.1);
+		}
 	}
 
 	// Check that a hanging mass on a joint drive reaches the correct extension with the correct spring force when using Force mode.
diff --git a/Engine/Source/Runtime/Experimental/Chaos/Private/Chaos/Joint/PBDJointCachedSolverGaussSeidel.cpp b/Engine/Source/Runtime/Experimental/Chaos/Private/Chaos/Joint/PBDJointCachedSolverGaussSeidel.cpp
--- a/Engine/Source/Runtime/Experimental/Chaos/Private/Chaos/Joint/PBDJointCachedSolverGaussSeidel.cpp
+++ b/Engine/Source/Runtime/Experimental/Chaos/Private/Chaos/Joint/PBDJointCachedSolverGaussSeidel.cpp
@@ -1328,7 +1328,7 @@ void FPBDJointCachedSolver::InitRotationConstraintsSimd(
 		RotationConstraints.Simd.ConstraintDRAxis[ConstraintIndex][1] = VectorNegate(IA1);
 
 		ConstraintHardIM[ConstraintIndex] = II0 + II1;
-		const FReal SpringMassScale = (bUsePositionBasedDrives) ? (FReal)1 / (ConstraintHardIM[ConstraintIndex]) : (FReal)1;
+		const FReal SpringMassScale = (RotationConstraints.bAccelerationMode) ? (FReal)1 / (ConstraintHardIM[ConstraintIndex]) : (FReal)1;
 		ConstraintSoftStiffness[ConstraintIndex] = FRealSingle(SpringMassScale * RotationConstraints.SettingsSoftStiffness[ConstraintIndex] * Dt * Dt);
 		ConstraintSoftDamping[ConstraintIndex] = (bUsePositionBasedDrives) ? FRealSingle(SpringMassScale * RotationConstraints.SettingsSoftDamping[ConstraintIndex] * Dt) : 0.0f;
 		ConstraintSoftIM[ConstraintIndex] = FRealSingle((ConstraintSoftStiffness[ConstraintIndex] + ConstraintSoftDamping[ConstraintIndex]) * ConstraintHardIM[ConstraintIndex] + (FReal)1);
@@ -2600,7 +2600,6 @@ void FPBDJointCachedSolver::InitRotationConstraintDriveSimd(
 		RotationDrives.Simd.ConstraintArms[ConstraintIndex][1] = VectorZeroFloat();
 		RotationDrives.Simd.ConstraintAxis[ConstraintIndex] = MakeVectorRegisterFloatFromDouble(MakeVectorRegister(ConstraintAxes[ConstraintIndex][0], ConstraintAxes[ConstraintIndex][1], ConstraintAxes[ConstraintIndex][2], 0.0f));
 
-		// InitRotationDatasMass(RotationDrives, ConstraintIndex, Dt);
 		const VectorRegister4Float& Axis = RotationDrives.Simd.ConstraintAxis[ConstraintIndex];
 		const VectorRegister4Float AxisX = VectorReplicate(Axis, 0);
 		const VectorRegister4Float AxisY = VectorReplicate(Axis, 1);
@@ -2627,7 +2626,7 @@ void FPBDJointCachedSolver::InitRotationConstraintDriveSimd(
 
 
 		check(RotationDrives.bSoftLimit[ConstraintIndex])
-			const FReal SpringMassScale = (bUsePositionBasedDrives) ? (FReal)1 / (ConstraintHardIM[ConstraintIndex]) : (FReal)1;
+		const FReal SpringMassScale = (RotationDrives.bAccelerationMode) ? (FReal)1 / (ConstraintHardIM[ConstraintIndex]) : (FReal)1;
 		ConstraintSoftStiffness[ConstraintIndex] = SpringMassScale * RotationDrives.SettingsSoftStiffness[ConstraintIndex] * Dt * Dt;
 		ConstraintSoftDamping[ConstraintIndex] = (bUsePositionBasedDrives) ? SpringMassScale * RotationDrives.SettingsSoftDamping[ConstraintIndex] * Dt : 0;
 		ConstraintSoftIM[ConstraintIndex] = (ConstraintSoftStiffness[ConstraintIndex] + ConstraintSoftDamping[ConstraintIndex]) * ConstraintHardIM[ConstraintIndex] + (FReal)1;
