Fix velocity constraint on the vectorized solver

diff --git a/Engine/Source/Programs/HeadlessChaos/Private/HeadlessChaosTestJointBreak.cpp b/Engine/Source/Programs/HeadlessChaos/Private/HeadlessChaosTestJointBreak.cpp
--- a/Engine/Source/Programs/HeadlessChaos/Private/HeadlessChaosTestJointBreak.cpp
+++ b/Engine/Source/Programs/HeadlessChaos/Private/HeadlessChaosTestJointBreak.cpp
@@ -194,7 +194,7 @@ namespace ChaosTest {
 
 	// 1 Kinematic Body with 2 Dynamic bodies hanging from it by a breakable constraint.
 	// Constraint break force is less than M x G, so joint should not break.
-	template <typename TEvolution>
+	template <typename TEvolution, bool bUseSimd = false>
 	void JointBreak_UnderLinearThreshold3()
 	{
 		const int32 NumIterations = 1;
@@ -205,6 +205,8 @@ namespace ChaosTest {
 		FJointChainTest<TEvolution> Test(NumIterations, Gravity);
 		Test.InitChain(3, FVec3(0, 0, -1));
 
+		Test.Evolution.GetJointConstraints().SetUseSimd(bUseSimd);
+
 		// Joint should break only if Threshold < MG
 		// So no in this test
 		Test.JointSettings[0].LinearBreakForce = 1.2f * (Test.ParticleMasses[1] + Test.ParticleMasses[2]) * Gravity;
@@ -234,6 +236,8 @@ namespace ChaosTest {
 	GTEST_TEST(AllEvolutions, JointBreakTests_TestUnderLinearThreshold3)
 	{
 		JointBreak_UnderLinearThreshold3<FPBDRigidsEvolutionGBF>();
+		constexpr bool bUseSimd = true;
+		JointBreak_UnderLinearThreshold3<FPBDRigidsEvolutionGBF, bUseSimd>();
 	}
 
 	// 1 Kinematic Body with 1 Dynamic body held vertically by a breakable angular constraint.
diff --git a/Engine/Source/Runtime/Experimental/Chaos/Private/Chaos/Joint/PBDJointCachedSolverGaussSeidel.cpp b/Engine/Source/Runtime/Experimental/Chaos/Private/Chaos/Joint/PBDJointCachedSolverGaussSeidel.cpp
--- a/Engine/Source/Runtime/Experimental/Chaos/Private/Chaos/Joint/PBDJointCachedSolverGaussSeidel.cpp
+++ b/Engine/Source/Runtime/Experimental/Chaos/Private/Chaos/Joint/PBDJointCachedSolverGaussSeidel.cpp
@@ -1064,7 +1064,7 @@ void FPBDJointCachedSolver::ApplyVelocityConstraintSimd()
 			VectorStoreFloat3(DV0, &DV0f[0]);
 			FVec3f DW0f;
 			VectorStoreFloat3(DW0, &DW0f[0]);
-			Body(0).ApplyVelocityDelta(FVec3(DW0f), FVec3(DW0f));
+			Body(0).ApplyVelocityDelta(FVec3(DV0f), FVec3(DW0f));
 		}
 		if (Body(1).IsDynamic())
 		{
@@ -1081,7 +1081,7 @@ void FPBDJointCachedSolver::ApplyVelocityConstraintSimd()
 			VectorStoreFloat3(DV1, &DV1f[0]);
 			FVec3f DW1f;
 			VectorStoreFloat3(DW1, &DW1f[0]);
-			Body(1).ApplyVelocityDelta(FVec3(DW1f), FVec3(DW1f));
+			Body(1).ApplyVelocityDelta(FVec3(DV1f), FVec3(DW1f));
 		}
 	}
 }
