Add unit tests for vectorized solver

diff --git a/Engine/Source/Programs/HeadlessChaos/Private/GeometryCollection/GeometryCollectionTestClustering.cpp b/Engine/Source/Programs/HeadlessChaos/Private/GeometryCollection/GeometryCollectionTestClustering.cpp
--- a/Engine/Source/Programs/HeadlessChaos/Private/GeometryCollection/GeometryCollectionTestClustering.cpp
+++ b/Engine/Source/Programs/HeadlessChaos/Private/GeometryCollection/GeometryCollectionTestClustering.cpp
@@ -1000,9 +1000,9 @@ namespace GeometryCollectionTest
 		EXPECT_TRUE(ClusterMapContains(ClusterMap, ParticleHandles[8], { ParticleHandles[7],ParticleHandles[0] }));
 
 		// Storage for positions for particles 0, 1, 6 for testing assumptions
-		FVector Ref0;
-		FVector Ref1;
-		FVector Ref6;
+		FVector Ref0(0.0);
+		FVector Ref1(0.0);
+		FVector Ref6(0.0);
 
 		for (int Frame = 1; Frame < 10; Frame++)
 		{
@@ -1770,7 +1770,6 @@ namespace GeometryCollectionTest
 		const auto& ParticleHandles = UnitTest.Solver->GetParticles().GetParticleHandles();
 		const TArrayCollectionArray<FRigidTransform3>& ChildToParent = Clustering.GetChildToParentMap();
 
-		FVector TestOffset;
 		FVector InitialRootPosition;
 		FVector RelativeChildOffsets[4];
 
diff --git a/Engine/Source/Programs/HeadlessChaos/Private/HeadlessChaosTestConstraints.cpp b/Engine/Source/Programs/HeadlessChaos/Private/HeadlessChaosTestConstraints.cpp
--- a/Engine/Source/Programs/HeadlessChaos/Private/HeadlessChaosTestConstraints.cpp
+++ b/Engine/Source/Programs/HeadlessChaos/Private/HeadlessChaosTestConstraints.cpp
@@ -24,7 +24,7 @@ namespace ChaosTest {
 	/**
 	 * Position constraint test
 	 */
-	template<typename TEvolution>
+	template<typename TEvolution, bool bUseSimd>
 	void Position()
 	{
 		{
@@ -38,6 +38,7 @@ namespace ChaosTest {
 			TArray<FVec3> Positions = { FVec3(0) };
 			FPBDPositionConstraints PositionConstraints(MoveTemp(Positions), MoveTemp(Dynamics), 1.f);
 			InitEvolutionSettings(Evolution);
+			Evolution.GetJointConstraints().SetUseSimd(bUseSimd);
 
 			Evolution.AddConstraintContainer(PositionConstraints);
 			Evolution.AdvanceOneTimeStep(0.1);
@@ -57,6 +58,7 @@ namespace ChaosTest {
 			TArray<FVec3> Positions = { FVec3(1) };
 			FPBDPositionConstraints PositionConstraints(MoveTemp(Positions), MoveTemp(Dynamics), 0.5f);
 			Evolution.AddConstraintContainer(PositionConstraints);
+			Evolution.GetJointConstraints().SetUseSimd(bUseSimd);
 
 			// The effect of stiffness parameter (which is set to 0.5 above) is iteration depeendent
 			Evolution.SetNumPositionIterations(1);
@@ -330,10 +332,10 @@ namespace ChaosTest {
 
 	GTEST_TEST(AllEvolutions, Constraints)
 	{
-		ChaosTest::Position<FPBDRigidsEvolutionGBF>();
-		ChaosTest::PositionAndJoint<FPBDRigidsEvolutionGBF>();
-		constexpr bool bTestSimdSolver = true;
-		ChaosTest::PositionAndJoint<FPBDRigidsEvolutionGBF, bTestSimdSolver>();
+		ChaosTest::Position<FPBDRigidsEvolutionGBF, false>();
+		ChaosTest::Position<FPBDRigidsEvolutionGBF, true>();
+		ChaosTest::PositionAndJoint<FPBDRigidsEvolutionGBF, false>();
+		ChaosTest::PositionAndJoint<FPBDRigidsEvolutionGBF, true>();
 		ChaosTest::SuspensionConstraintHardstop<FPBDRigidsEvolutionGBF>();
 		ChaosTest::SuspensionConstraintSpring<FPBDRigidsEvolutionGBF>();
 
diff --git a/Engine/Source/Programs/HeadlessChaos/Private/HeadlessChaosTestHeightfield.cpp b/Engine/Source/Programs/HeadlessChaos/Private/HeadlessChaosTestHeightfield.cpp
--- a/Engine/Source/Programs/HeadlessChaos/Private/HeadlessChaosTestHeightfield.cpp
+++ b/Engine/Source/Programs/HeadlessChaos/Private/HeadlessChaosTestHeightfield.cpp
@@ -219,7 +219,6 @@ namespace ChaosTest {
 						}
 
 						// reverse ray to test double sided
-						FVec3 ReversePosition;
 						FReal ReverseTOI;
 						const bool bReverseResult = Heightfield.Raycast(RayEnd, -RayDir, RayLength, 0, ReverseTOI, Position, Normal, FaceIdx);
 						if (Row + 1 == Rows)
diff --git a/Engine/Source/Programs/HeadlessChaos/Private/HeadlessChaosTestJointBreak.cpp b/Engine/Source/Programs/HeadlessChaos/Private/HeadlessChaosTestJointBreak.cpp
--- a/Engine/Source/Programs/HeadlessChaos/Private/HeadlessChaosTestJointBreak.cpp
+++ b/Engine/Source/Programs/HeadlessChaos/Private/HeadlessChaosTestJointBreak.cpp
@@ -10,7 +10,7 @@ namespace ChaosTest {
 
 	// Set up a test with a non-breakable joint, then manually break it.
 	// Verify that the break callback is called and the joint is disabled.
-	template <typename TEvolution>
+	template <typename TEvolution, bool bUseSimd = false>
 	void JointBreak_ManualBreak()
 	{
 		const int32 NumIterations = 1;
@@ -21,6 +21,7 @@ namespace ChaosTest {
 		FJointChainTest<TEvolution> Test(NumIterations, Gravity);
 		Test.InitChain(2, FVec3(0,0,-1));
 		Test.Create();
+		Test.Evolution.GetJointConstraints().SetUseSimd(bUseSimd);
 
 		bool bBrokenCallbackCalled = false;
 		Test.Evolution.GetJointConstraints().SetBreakCallback([&bBrokenCallbackCalled](FPBDJointConstraintHandle* Constraint)
@@ -59,12 +60,13 @@ namespace ChaosTest {
 
 	GTEST_TEST(AllEvolutions, JointBreakTests_TestManualBreak)
 	{
-		JointBreak_ManualBreak<FPBDRigidsEvolutionGBF>();
+		JointBreak_ManualBreak<FPBDRigidsEvolutionGBF, false>();
+		JointBreak_ManualBreak<FPBDRigidsEvolutionGBF, true>();
 	}
 
 	// 1 Kinematic Body with 1 Dynamic body hanging from it by a breakable constraint.
 	// Constraint break force is larger than M x G, so joint should not break.
-	template <typename TEvolution>
+	template <typename TEvolution, bool bUseSimd = false>
 	void JointBreak_UnderLinearThreshold()
 	{
 		const int32 NumIterations = 1;
@@ -74,7 +76,7 @@ namespace ChaosTest {
 
 		FJointChainTest<TEvolution> Test(NumIterations, Gravity);
 		Test.InitChain(2, FVec3(0, 0, -1));
-
+		Test.Evolution.GetJointConstraints().SetUseSimd(bUseSimd);
 		// Joint should break only if Threshold < MG
 		// So not in this test
 		Test.JointSettings[0].LinearBreakForce = 1.1f * Test.ParticleMasses[1] * Gravity;
@@ -101,12 +103,13 @@ namespace ChaosTest {
 
 	GTEST_TEST(AllEvolutions, JointBreakTests_TestUnderLinearThreshold)
 	{
-		JointBreak_UnderLinearThreshold<FPBDRigidsEvolutionGBF>();
+		JointBreak_UnderLinearThreshold<FPBDRigidsEvolutionGBF, false>();
+		JointBreak_UnderLinearThreshold<FPBDRigidsEvolutionGBF, true>();
 	}
 
 	// 1 Kinematic Body with 2 Dynamic bodies hanging from it by a breakable constraint.
 	// Constraint break forces are larger than M x G, so joint should not break.
-	template <typename TEvolution>
+	template <typename TEvolution, bool bUseSimd>
 	void JointBreak_UnderLinearThreshold2()
 	{
 		const int32 NumIterations = 1;
@@ -116,6 +119,7 @@ namespace ChaosTest {
 
 		FJointChainTest<TEvolution> Test(NumIterations, Gravity);
 		Test.InitChain(3, FVec3(0, 0, -1));
+		Test.Evolution.GetJointConstraints().SetUseSimd(bUseSimd);
 
 		// Joint should break only if Threshold < MG
 		// So not in this test
@@ -146,12 +150,13 @@ namespace ChaosTest {
 
 	GTEST_TEST(AllEvolutions, JointBreakTests_TestUnderLinearThreshold2)
 	{
-		JointBreak_UnderLinearThreshold2<FPBDRigidsEvolutionGBF>();
+		JointBreak_UnderLinearThreshold2<FPBDRigidsEvolutionGBF, false>();
+		JointBreak_UnderLinearThreshold2<FPBDRigidsEvolutionGBF, true>();
 	}
 
 	// 1 Kinematic Body with 1 Dynamic body hanging from it by a breakable constraint.
 	// Constraint break force is less than M x G, so joint should break.
-	template <typename TEvolution>
+	template <typename TEvolution, bool bUseSimd>
 	void JointBreak_OverLinearThreshold()
 	{
 		const int32 NumIterations = 1;
@@ -161,7 +166,7 @@ namespace ChaosTest {
 
 		FJointChainTest<TEvolution> Test(NumIterations, Gravity);
 		Test.InitChain(2, FVec3(0,0,-1));
-
+		Test.Evolution.GetJointConstraints().SetUseSimd(bUseSimd);
 		// Joint should break only if Threshold < MG
 		// So yes in this test
 		Test.JointSettings[0].LinearBreakForce = 0.9f * Test.ParticleMasses[1] * Gravity;
@@ -188,7 +193,8 @@ namespace ChaosTest {
 
 	GTEST_TEST(AllEvolutions, JointBreakTests_TestOverLinearThreshold)
 	{
-		JointBreak_OverLinearThreshold<FPBDRigidsEvolutionGBF>();
+		JointBreak_OverLinearThreshold<FPBDRigidsEvolutionGBF, false>();
+		JointBreak_OverLinearThreshold<FPBDRigidsEvolutionGBF, true>();
 	}
 
 
@@ -204,7 +210,6 @@ namespace ChaosTest {
 
 		FJointChainTest<TEvolution> Test(NumIterations, Gravity);
 		Test.InitChain(3, FVec3(0, 0, -1));
-
 		Test.Evolution.GetJointConstraints().SetUseSimd(bUseSimd);
 
 		// Joint should break only if Threshold < MG
@@ -242,7 +247,7 @@ namespace ChaosTest {
 
 	// 1 Kinematic Body with 1 Dynamic body held vertically by a breakable angular constraint.
 	// Constraint break torque is larger than input torque so constraint will not break.
-	template <typename TEvolution>
+	template <typename TEvolution, bool bUseSimd>
 	void JointBreak_UnderAngularThreshold()
 	{
 		const int32 NumIterations = 1;
@@ -253,7 +258,7 @@ namespace ChaosTest {
 
 		FJointChainTest<TEvolution> Test(NumIterations, Gravity);
 		Test.InitChain(2, FVec3(0, 0, -1));
-
+		Test.Evolution.GetJointConstraints().SetUseSimd(bUseSimd);
 		// Joint should break only if Threshold < MGL
 		// So not in this test
 		Test.JointSettings[0].AngularBreakTorque = 1.1f * Torque.X;
@@ -283,12 +288,13 @@ namespace ChaosTest {
 
 	GTEST_TEST(AllEvolutions, JointBreakTests_TestUnderAngularThreshold)
 	{
-		JointBreak_UnderAngularThreshold<FPBDRigidsEvolutionGBF>();
+		JointBreak_UnderAngularThreshold<FPBDRigidsEvolutionGBF, false>();
+		JointBreak_UnderAngularThreshold<FPBDRigidsEvolutionGBF, true>();
 	}
 
 	// 1 Kinematic Body with 1 Dynamic body held vertically by a breakable angular constraint.
 	// Constraint break torque is less than input torque so constraint will break.
-	template <typename TEvolution>
+	template <typename TEvolution, bool bUseSimd>
 	void JointBreak_OverAngularThreshold()
 	{
 		const int32 NumSolverIterations = 10;
@@ -299,7 +305,7 @@ namespace ChaosTest {
 
 		FJointChainTest<TEvolution> Test(NumSolverIterations, Gravity);
 		Test.InitChain(2, FVec3(0, 0, -1));
-
+		Test.Evolution.GetJointConstraints().SetUseSimd(bUseSimd);
 		// Joint should break only if Threshold < MGL
 		// So not in this test
 		Test.JointSettings[0].AngularBreakTorque = 0.9f * Torque.X;
@@ -329,7 +335,8 @@ namespace ChaosTest {
 
 	GTEST_TEST(AllEvolutions, JointBreakTests_TestOverAngularThreshold)
 	{
-		JointBreak_OverAngularThreshold<FPBDRigidsEvolutionGBF>();
+		JointBreak_OverAngularThreshold<FPBDRigidsEvolutionGBF, false>();
+		JointBreak_OverAngularThreshold<FPBDRigidsEvolutionGBF, true>();
 	}
 
 }
\ No newline at end of file
diff --git a/Engine/Source/Programs/HeadlessChaos/Private/HeadlessChaosTestJointForces.cpp b/Engine/Source/Programs/HeadlessChaos/Private/HeadlessChaosTestJointForces.cpp
--- a/Engine/Source/Programs/HeadlessChaos/Private/HeadlessChaosTestJointForces.cpp
+++ b/Engine/Source/Programs/HeadlessChaos/Private/HeadlessChaosTestJointForces.cpp
@@ -14,37 +14,41 @@ namespace ChaosTest {
 	template <typename TEvolution>
 	void JointForces_Linear()
 	{
-		const int32 NumSolverIterations = 10;
-		const FReal Gravity = 980;
-		const FReal Dt = 0.01f;
-		const int32 NumSteps = 100;
-		const int32 NumBodies = 5;
-
-		FJointChainTest<TEvolution> Test(NumSolverIterations, Gravity);
-		Test.InitChain(NumBodies, FVec3(0, 0, -1));
-
-		Test.Create();
-
-		// Run the sim
-		for (int32 i = 0; i < NumSteps; ++i)
+		for (int32 IndexSimd = 0; IndexSimd < 2; IndexSimd++)
 		{
-			Test.Evolution.AdvanceOneTimeStep(Dt);
-			Test.Evolution.EndFrame(Dt);
-		}
+			const bool bUseSimd = IndexSimd == 0;
+			const int32 NumSolverIterations = 10;
+			const FReal Gravity = 980;
+			const FReal Dt = 0.01f;
+			const int32 NumSteps = 100;
+			const int32 NumBodies = 5;
+
+			FJointChainTest<TEvolution> Test(NumSolverIterations, Gravity);
+			Test.InitChain(NumBodies, FVec3(0, 0, -1));
+			Test.Evolution.GetJointConstraints().SetUseSimd(bUseSimd);
+			Test.Create();
 
-		for (int32 JointIndex = 0; JointIndex < Test.Evolution.GetJointConstraints().NumConstraints(); ++JointIndex)
-		{
-			FReal ChildMass = 0.0f;
-			for (int32 ChildBodyIndex = JointIndex + 1; ChildBodyIndex < Test.ParticleMasses.Num(); ++ChildBodyIndex)
+			// Run the sim
+			for (int32 i = 0; i < NumSteps; ++i)
 			{
-				ChildMass += Test.ParticleMasses[ChildBodyIndex];
+				Test.Evolution.AdvanceOneTimeStep(Dt);
+				Test.Evolution.EndFrame(Dt);
 			}
 
-			FVec3 ExpectedLinearImpulse = FVec3(0.0f, 0.0f, -ChildMass * Gravity * Dt);
-			FVec3 LinearImpulse = Test.Evolution.GetJointConstraints().GetConstraintLinearImpulse(JointIndex);
-			EXPECT_NEAR(LinearImpulse.X, ExpectedLinearImpulse.X, ExpectedLinearImpulse.Size() / 100.0f);
-			EXPECT_NEAR(LinearImpulse.Y, ExpectedLinearImpulse.Y, ExpectedLinearImpulse.Size() / 100.0f);
-			EXPECT_NEAR(LinearImpulse.Z, ExpectedLinearImpulse.Z, ExpectedLinearImpulse.Size() / 100.0f);
+			for (int32 JointIndex = 0; JointIndex < Test.Evolution.GetJointConstraints().NumConstraints(); ++JointIndex)
+			{
+				FReal ChildMass = 0.0f;
+				for (int32 ChildBodyIndex = JointIndex + 1; ChildBodyIndex < Test.ParticleMasses.Num(); ++ChildBodyIndex)
+				{
+					ChildMass += Test.ParticleMasses[ChildBodyIndex];
+				}
+
+				FVec3 ExpectedLinearImpulse = FVec3(0.0f, 0.0f, -ChildMass * Gravity * Dt);
+				FVec3 LinearImpulse = Test.Evolution.GetJointConstraints().GetConstraintLinearImpulse(JointIndex);
+				EXPECT_NEAR(LinearImpulse.X, ExpectedLinearImpulse.X, ExpectedLinearImpulse.Size() / 100.0f);
+				EXPECT_NEAR(LinearImpulse.Y, ExpectedLinearImpulse.Y, ExpectedLinearImpulse.Size() / 100.0f);
+				EXPECT_NEAR(LinearImpulse.Z, ExpectedLinearImpulse.Z, ExpectedLinearImpulse.Size() / 100.0f);
+			}
 		}
 	}
 
@@ -58,55 +62,60 @@ namespace ChaosTest {
 	template <typename TEvolution>
 	void JointForces_Linear2()
 	{
-		const int32 NumSolverIterations = 40;
-		const FReal Gravity = 980;
-		const FReal Dt = 0.01f;
-		const int32 NumSteps = 100;
-		const int32 NumBodies = 3;
-
-		FJointChainTest<TEvolution> Test(NumSolverIterations, Gravity);
-		Test.InitChain(NumBodies, FVec3(1, 0, 0));
-
-		// Set all joints to fixed angular
-		for (int32 JointIndex = 0; JointIndex < Test.JointSettings.Num(); ++JointIndex)
+		for (int32 IndexSimd = 0; IndexSimd < 2; IndexSimd++)
 		{
-			Test.JointSettings[JointIndex].AngularMotionTypes = { EJointMotionType::Locked, EJointMotionType::Locked, EJointMotionType::Locked };
-			Test.JointSettings[JointIndex].bProjectionEnabled = false;
-		}
-
-		Test.Create();
+			const bool bUseSimd = IndexSimd == 0;
+			const int32 NumSolverIterations = 40;
+			const FReal Gravity = 980;
+			const FReal Dt = 0.01f;
+			const int32 NumSteps = 100;
+			const int32 NumBodies = 3;
+
+			FJointChainTest<TEvolution> Test(NumSolverIterations, Gravity);
+			Test.InitChain(NumBodies, FVec3(1, 0, 0));
+			Test.Evolution.GetJointConstraints().SetUseSimd(bUseSimd);
 
-		// Run the sim
-		for (int32 i = 0; i < NumSteps; ++i)
-		{
-			Test.Evolution.AdvanceOneTimeStep(Dt);
-			Test.Evolution.EndFrame(Dt);
-		}
+			// Set all joints to fixed angular
+			for (int32 JointIndex = 0; JointIndex < Test.JointSettings.Num(); ++JointIndex)
+			{
+				Test.JointSettings[JointIndex].AngularMotionTypes = { EJointMotionType::Locked, EJointMotionType::Locked, EJointMotionType::Locked };
+				Test.JointSettings[JointIndex].bProjectionEnabled = false;
+			}
 
-		FReal L = Test.ParticlePositions[1].X - Test.ParticlePositions[0].X;
+			Test.Create();
 
-		for (int32 JointIndex = 0; JointIndex < Test.Evolution.GetJointConstraints().NumConstraints(); ++JointIndex)
-		{
-			FReal ChildMass = 0.0f;
-			FReal ChildMoment = 0.0f;
-			for (int32 ChildBodyIndex = JointIndex + 1; ChildBodyIndex < Test.ParticleMasses.Num(); ++ChildBodyIndex)
+			// Run the sim
+			for (int32 i = 0; i < NumSteps; ++i)
 			{
-				ChildMass += Test.ParticleMasses[ChildBodyIndex];
-				FReal ChildL = (ChildBodyIndex - JointIndex) * L;
-				ChildMoment += Test.ParticleMasses[ChildBodyIndex] * ChildL;
+				Test.Evolution.AdvanceOneTimeStep(Dt);
+				Test.Evolution.EndFrame(Dt);
 			}
 
-			FVec3 ExpectedLinearImpulse = FVec3(0.0f, 0.0f, -ChildMass * Gravity * Dt);
-			FVec3 LinearImpulse = Test.Evolution.GetJointConstraints().GetConstraintLinearImpulse(JointIndex);
-			EXPECT_NEAR(LinearImpulse.X, ExpectedLinearImpulse.X, ExpectedLinearImpulse.Size() / 100.0f);
-			EXPECT_NEAR(LinearImpulse.Y, ExpectedLinearImpulse.Y, ExpectedLinearImpulse.Size() / 100.0f);
-			EXPECT_NEAR(LinearImpulse.Z, ExpectedLinearImpulse.Z, ExpectedLinearImpulse.Size() / 100.0f);
-
-			FVec3 ExpectedAngularImpulse = FVec3(0.0f, ChildMoment * Gravity * Dt, 0.0f);
-			FVec3 AngularImpulse = Test.Evolution.GetJointConstraints().GetConstraintAngularImpulse(JointIndex);
-			EXPECT_NEAR(AngularImpulse.X, ExpectedAngularImpulse.X, ExpectedAngularImpulse.Size() / 100.0f);
-			EXPECT_NEAR(AngularImpulse.Y, ExpectedAngularImpulse.Y, ExpectedAngularImpulse.Size() / 100.0f);
-			EXPECT_NEAR(AngularImpulse.Z, ExpectedAngularImpulse.Z, ExpectedAngularImpulse.Size() / 100.0f);
+			FReal L = Test.ParticlePositions[1].X - Test.ParticlePositions[0].X;
+
+			for (int32 JointIndex = 0; JointIndex < Test.Evolution.GetJointConstraints().NumConstraints(); ++JointIndex)
+			{
+				FReal ChildMass = 0.0f;
+				FReal ChildMoment = 0.0f;
+				for (int32 ChildBodyIndex = JointIndex + 1; ChildBodyIndex < Test.ParticleMasses.Num(); ++ChildBodyIndex)
+				{
+					ChildMass += Test.ParticleMasses[ChildBodyIndex];
+					FReal ChildL = (ChildBodyIndex - JointIndex) * L;
+					ChildMoment += Test.ParticleMasses[ChildBodyIndex] * ChildL;
+				}
+
+				FVec3 ExpectedLinearImpulse = FVec3(0.0f, 0.0f, -ChildMass * Gravity * Dt);
+				FVec3 LinearImpulse = Test.Evolution.GetJointConstraints().GetConstraintLinearImpulse(JointIndex);
+				EXPECT_NEAR(LinearImpulse.X, ExpectedLinearImpulse.X, ExpectedLinearImpulse.Size() / 100.0f);
+				EXPECT_NEAR(LinearImpulse.Y, ExpectedLinearImpulse.Y, ExpectedLinearImpulse.Size() / 100.0f);
+				EXPECT_NEAR(LinearImpulse.Z, ExpectedLinearImpulse.Z, ExpectedLinearImpulse.Size() / 100.0f);
+
+				FVec3 ExpectedAngularImpulse = FVec3(0.0f, ChildMoment * Gravity * Dt, 0.0f);
+				FVec3 AngularImpulse = Test.Evolution.GetJointConstraints().GetConstraintAngularImpulse(JointIndex);
+				EXPECT_NEAR(AngularImpulse.X, ExpectedAngularImpulse.X, ExpectedAngularImpulse.Size() / 100.0f);
+				EXPECT_NEAR(AngularImpulse.Y, ExpectedAngularImpulse.Y, ExpectedAngularImpulse.Size() / 100.0f);
+				EXPECT_NEAR(AngularImpulse.Z, ExpectedAngularImpulse.Z, ExpectedAngularImpulse.Size() / 100.0f);
+			}
 		}
 	}
 
@@ -121,47 +130,52 @@ namespace ChaosTest {
 	template <typename TEvolution>
 	void JointForces_Angular()
 	{
-		const int32 NumSolverIterations = 20;
-		const FReal Gravity = 0;
-		const FReal Dt = 0.01f;
-		const int32 NumSteps = 100;
-		const int32 NumBodies = 2;
-		const FVec3 Torque = FVec3(10000, 0, 0);
-
-		FJointChainTest<TEvolution> Test(NumSolverIterations, Gravity);
-		Test.InitChain(NumBodies, FVec3(1, 0, 0));
-
-		// Set all joints to fixed angular and disable projection
-		for (int32 JointIndex = 0; JointIndex < Test.JointSettings.Num(); ++JointIndex)
+		for (int32 IndexSimd = 0; IndexSimd < 2; IndexSimd++)
 		{
-			Test.JointSettings[JointIndex].AngularMotionTypes = { EJointMotionType::Locked, EJointMotionType::Locked, EJointMotionType::Locked };
-			Test.JointSettings[JointIndex].bProjectionEnabled = false;
-		}
+			bool bUseSimd = IndexSimd == 0;
+			const int32 NumSolverIterations = 20;
+			const FReal Gravity = 0;
+			const FReal Dt = 0.01f;
+			const int32 NumSteps = 100;
+			const int32 NumBodies = 2;
+			const FVec3 Torque = FVec3(10000, 0, 0);
 
-		Test.Create();
+			FJointChainTest<TEvolution> Test(NumSolverIterations, Gravity);
+			Test.InitChain(NumBodies, FVec3(1, 0, 0));
+			Test.Evolution.GetJointConstraints().SetUseSimd(bUseSimd);
 
-		// Run the sim
-		for (int32 i = 0; i < NumSteps; ++i)
-		{
-			Test.GetParticle(1)->CastToRigidParticle()->SetTorque(Torque);
-			
-			Test.Evolution.AdvanceOneTimeStep(Dt);
-			Test.Evolution.EndFrame(Dt);
-		}
+			// Set all joints to fixed angular and disable projection
+			for (int32 JointIndex = 0; JointIndex < Test.JointSettings.Num(); ++JointIndex)
+			{
+				Test.JointSettings[JointIndex].AngularMotionTypes = { EJointMotionType::Locked, EJointMotionType::Locked, EJointMotionType::Locked };
+				Test.JointSettings[JointIndex].bProjectionEnabled = false;
+			}
 
-		for (int32 JointIndex = 0; JointIndex < Test.Evolution.GetJointConstraints().NumConstraints(); ++JointIndex)
-		{
-			FReal ChildMass = 0.0f;
-			for (int32 ChildBodyIndex = JointIndex + 1; ChildBodyIndex < Test.ParticleMasses.Num(); ++ChildBodyIndex)
+			Test.Create();
+
+			// Run the sim
+			for (int32 i = 0; i < NumSteps; ++i)
 			{
-				ChildMass += Test.ParticleMasses[ChildBodyIndex];
+				Test.GetParticle(1)->CastToRigidParticle()->SetTorque(Torque);
+
+				Test.Evolution.AdvanceOneTimeStep(Dt);
+				Test.Evolution.EndFrame(Dt);
 			}
 
-			FVec3 ExpectedAngularImpulse = Torque * Dt;
-			FVec3 AngularImpulse = Test.Evolution.GetJointConstraints().GetConstraintAngularImpulse(JointIndex);
-			EXPECT_NEAR(AngularImpulse.X, ExpectedAngularImpulse.X, ExpectedAngularImpulse.Size() / 100.0f);
-			EXPECT_NEAR(AngularImpulse.Y, ExpectedAngularImpulse.Y, ExpectedAngularImpulse.Size() / 100.0f);
-			EXPECT_NEAR(AngularImpulse.Z, ExpectedAngularImpulse.Z, ExpectedAngularImpulse.Size() / 100.0f);
+			for (int32 JointIndex = 0; JointIndex < Test.Evolution.GetJointConstraints().NumConstraints(); ++JointIndex)
+			{
+				FReal ChildMass = 0.0f;
+				for (int32 ChildBodyIndex = JointIndex + 1; ChildBodyIndex < Test.ParticleMasses.Num(); ++ChildBodyIndex)
+				{
+					ChildMass += Test.ParticleMasses[ChildBodyIndex];
+				}
+
+				FVec3 ExpectedAngularImpulse = Torque * Dt;
+				FVec3 AngularImpulse = Test.Evolution.GetJointConstraints().GetConstraintAngularImpulse(JointIndex);
+				EXPECT_NEAR(AngularImpulse.X, ExpectedAngularImpulse.X, ExpectedAngularImpulse.Size() / 100.0f);
+				EXPECT_NEAR(AngularImpulse.Y, ExpectedAngularImpulse.Y, ExpectedAngularImpulse.Size() / 100.0f);
+				EXPECT_NEAR(AngularImpulse.Z, ExpectedAngularImpulse.Z, ExpectedAngularImpulse.Size() / 100.0f);
+			}
 		}
 	}
 
@@ -175,56 +189,61 @@ namespace ChaosTest {
 	//
 	GTEST_TEST(JointForceTests, TestLinearDriveForceMode_Force)
 	{
-		const int32 NumSolverIterations = 20;
-		const FReal Gravity = 0;
-		const FReal Dt = 0.01;
-		const int32 NumBodies = 2;
-
-		const FReal Extension = 10;
-		const FReal Stiffness = 10000;
-		const FReal Damping = 0;
-
-		FJointChainTest<FPBDRigidsEvolutionGBF> Test(NumSolverIterations, Gravity);
-		Test.InitChain(NumBodies, FVec3(0, 0, -1));
-
-		// Disable all limits
-		Test.JointSettings[0].LinearMotionTypes = { EJointMotionType::Free, EJointMotionType::Free, EJointMotionType::Free };
-		Test.JointSettings[0].AngularMotionTypes = { EJointMotionType::Free, EJointMotionType::Free, EJointMotionType::Free };
-
-		// Set up the drive in force mode
-		Test.JointSettings[0].bLinearPositionDriveEnabled = { false, false, true };
-		Test.JointSettings[0].bLinearVelocityDriveEnabled = { false, false, false };
-		Test.JointSettings[0].LinearDriveForceMode = EJointForceMode::Force;
-		Test.JointSettings[0].LinearDriveStiffness = FVec3(0, 0, Stiffness);
-		Test.JointSettings[0].LinearDriveDamping = FVec3(0, 0, Damping);
-
-		Test.Create();
-
-		FGenericParticleHandle P1 = Test.GetParticle(1);
-
-		// Reposition the particle to have some extension in the spring
-		P1->InitTransform(P1->GetX() + FVec3(0,0,-Extension), P1->GetR());
-
-		// Run the sim
-		Test.Evolution.AdvanceOneTimeStep(Dt);
-		Test.Evolution.EndFrame(Dt);
-		
-		// Calculate expected force from F = -K.X with implicit integration
-		const FReal M = Test.ParticleMasses[1];
-		FReal ExpectedForceZ = 0;
-		FReal DP = 0;
-		for (int32 It = 0; It < 10; ++It)
+		for (int32 IndexSimd = 0; IndexSimd < 2; IndexSimd++)
 		{
-			const FReal X = -Extension + DP;
-			const FReal F = -Stiffness * X;
-			const FReal DV = ((F - ExpectedForceZ) / M) * Dt;
-			DP += DV * Dt;
-			ExpectedForceZ = F;
-		}
+			bool bUseSimd = IndexSimd == 0;
+			const int32 NumSolverIterations = 20;
+			const FReal Gravity = 0;
+			const FReal Dt = 0.01;
+			const int32 NumBodies = 2;
+
+			const FReal Extension = 10;
+			const FReal Stiffness = 10000;
+			const FReal Damping = 0;
+
+			FJointChainTest<FPBDRigidsEvolutionGBF> Test(NumSolverIterations, Gravity);
+			Test.InitChain(NumBodies, FVec3(0, 0, -1));
+			Test.Evolution.GetJointConstraints().SetUseSimd(bUseSimd);
+
+			// Disable all limits
+			Test.JointSettings[0].LinearMotionTypes = { EJointMotionType::Free, EJointMotionType::Free, EJointMotionType::Free };
+			Test.JointSettings[0].AngularMotionTypes = { EJointMotionType::Free, EJointMotionType::Free, EJointMotionType::Free };
+
+			// Set up the drive in force mode
+			Test.JointSettings[0].bLinearPositionDriveEnabled = { false, false, true };
+			Test.JointSettings[0].bLinearVelocityDriveEnabled = { false, false, false };
+			Test.JointSettings[0].LinearDriveForceMode = EJointForceMode::Force;
+			Test.JointSettings[0].LinearDriveStiffness = FVec3(0, 0, Stiffness);
+			Test.JointSettings[0].LinearDriveDamping = FVec3(0, 0, Damping);
 
-		// Check the joint forces agree
-		const FReal ForceZ = -Test.Evolution.GetJointConstraints().GetConstraintLinearImpulse(0).Z / Dt;
-		EXPECT_NEAR(ForceZ, ExpectedForceZ, 0.01);
+			Test.Create();
+
+			FGenericParticleHandle P1 = Test.GetParticle(1);
+
+			// Reposition the particle to have some extension in the spring
+			P1->InitTransform(P1->GetX() + FVec3(0, 0, -Extension), P1->GetR());
+
+			// Run the sim
+			Test.Evolution.AdvanceOneTimeStep(Dt);
+			Test.Evolution.EndFrame(Dt);
+
+			// Calculate expected force from F = -K.X with implicit integration
+			const FReal M = Test.ParticleMasses[1];
+			FReal ExpectedForceZ = 0;
+			FReal DP = 0;
+			for (int32 It = 0; It < 10; ++It)
+			{
+				const FReal X = -Extension + DP;
+				const FReal F = -Stiffness * X;
+				const FReal DV = ((F - ExpectedForceZ) / M) * Dt;
+				DP += DV * Dt;
+				ExpectedForceZ = F;
+			}
+
+			// Check the joint forces agree
+			const FReal ForceZ = -Test.Evolution.GetJointConstraints().GetConstraintLinearImpulse(0).Z / Dt;
+			EXPECT_NEAR(ForceZ, ExpectedForceZ, 0.01);
+		}
 	}
 
 	// Check that a joint drive linear damping calculates the correct force F=-D.V assuming implicit integration
@@ -232,61 +251,66 @@ namespace ChaosTest {
 	//
 	GTEST_TEST(JointForceTests, TestLinearDriveForceMode_Damping)
 	{
-		const int32 NumSolverIterations = 20;
-		const FReal Gravity = 0;
-		const FReal Dt = 0.01;
-		const int32 NumBodies = 2;
+		for (int32 IndexSimd = 0; IndexSimd < 2; IndexSimd++)
+		{
+			bool bUseSimd = IndexSimd == 0;
+			const int32 NumSolverIterations = 20;
+			const FReal Gravity = 0;
+			const FReal Dt = 0.01;
+			const int32 NumBodies = 2;
 
-		const FReal Velocity = 100;
-		const FReal Stiffness = 0;
-		const FReal Damping = 2000;
+			const FReal Velocity = 100;
+			const FReal Stiffness = 0;
+			const FReal Damping = 2000;
 
-		FJointChainTest<FPBDRigidsEvolutionGBF> Test(NumSolverIterations, Gravity);
-		Test.InitChain(NumBodies, FVec3(0, 0, -1));
+			FJointChainTest<FPBDRigidsEvolutionGBF> Test(NumSolverIterations, Gravity);
+			Test.InitChain(NumBodies, FVec3(0, 0, -1));
+			Test.Evolution.GetJointConstraints().SetUseSimd(bUseSimd);
 
-		if (!Test.Evolution.GetJointConstraints().GetSettings().bUsePositionBasedDrives)
-		{
-			Test.Evolution.SetNumVelocityIterations(NumSolverIterations);
-		}
+			if (!Test.Evolution.GetJointConstraints().GetSettings().bUsePositionBasedDrives)
+			{
+				Test.Evolution.SetNumVelocityIterations(NumSolverIterations);
+			}
 
-		// Disable all limits
-		Test.JointSettings[0].LinearMotionTypes = { EJointMotionType::Free, EJointMotionType::Free, EJointMotionType::Free };
-		Test.JointSettings[0].AngularMotionTypes = { EJointMotionType::Free, EJointMotionType::Free, EJointMotionType::Free };
+			// Disable all limits
+			Test.JointSettings[0].LinearMotionTypes = { EJointMotionType::Free, EJointMotionType::Free, EJointMotionType::Free };
+			Test.JointSettings[0].AngularMotionTypes = { EJointMotionType::Free, EJointMotionType::Free, EJointMotionType::Free };
 
-		// Set up the drive in force mode
-		Test.JointSettings[0].bLinearPositionDriveEnabled = { false, false, false };
-		Test.JointSettings[0].bLinearVelocityDriveEnabled = { false, false, true };
-		Test.JointSettings[0].LinearDriveForceMode = EJointForceMode::Force;
-		Test.JointSettings[0].LinearDriveStiffness = FVec3(0, 0, Stiffness);
-		Test.JointSettings[0].LinearDriveDamping = FVec3(0, 0, Damping);
+			// Set up the drive in force mode
+			Test.JointSettings[0].bLinearPositionDriveEnabled = { false, false, false };
+			Test.JointSettings[0].bLinearVelocityDriveEnabled = { false, false, true };
+			Test.JointSettings[0].LinearDriveForceMode = EJointForceMode::Force;
+			Test.JointSettings[0].LinearDriveStiffness = FVec3(0, 0, Stiffness);
+			Test.JointSettings[0].LinearDriveDamping = FVec3(0, 0, Damping);
+
+			Test.Create();
 
-		Test.Create();
+			FGenericParticleHandle P1 = Test.GetParticle(1);
 
-		FGenericParticleHandle P1 = Test.GetParticle(1);
+			// Give the particle velocity so that joint damping has some work to do
+			P1->SetV(FVec3(0, 0, -Velocity));
 
-		// Give the particle velocity so that joint damping has some work to do
-		P1->SetV(FVec3(0,0,-Velocity));
+			// Run the sim
+			Test.Evolution.AdvanceOneTimeStep(Dt);
+			Test.Evolution.EndFrame(Dt);
 
-		// Run the sim
-		Test.Evolution.AdvanceOneTimeStep(Dt);
-		Test.Evolution.EndFrame(Dt);
+			// Calculate expected force from F = -D.V with implicit integration
+			const FReal M = Test.ParticleMasses[1];
+			FReal ExpectedForceZ = 0;
+			FReal DP = 0;
+			for (int32 It = 0; It < NumSolverIterations; ++It)
+			{
+				const FReal V = -Velocity + DP / Dt;
+				const FReal F = -Damping * V;
+				const FReal DV = ((F - ExpectedForceZ) / M) * Dt;
+				DP += DV * Dt;
+				ExpectedForceZ = F;
+			}
 
-		// Calculate expected force from F = -D.V with implicit integration
-		const FReal M = Test.ParticleMasses[1];
-		FReal ExpectedForceZ = 0;
-		FReal DP = 0;
-		for (int32 It = 0; It < NumSolverIterations; ++It)
-		{
-			const FReal V = -Velocity + DP / Dt;
-			const FReal F = -Damping * V;
-			const FReal DV = ((F - ExpectedForceZ) / M) * Dt;
-			DP += DV * Dt;
-			ExpectedForceZ = F;
+			// Check the joint forces agree
+			const FReal ForceZ = -Test.Evolution.GetJointConstraints().GetConstraintLinearImpulse(0).Z / Dt;
+			EXPECT_NEAR(ForceZ, ExpectedForceZ, 1);
 		}
-
-		// Check the joint forces agree
-		const FReal ForceZ = -Test.Evolution.GetJointConstraints().GetConstraintLinearImpulse(0).Z / Dt;
-		EXPECT_NEAR(ForceZ, ExpectedForceZ, 1);
 	}
 
 	GTEST_TEST(JointForceTests, TestAngularDriveForceMode_Damping)
@@ -305,7 +329,6 @@ namespace ChaosTest {
 
 			FJointChainTest<FPBDRigidsEvolutionGBF> Test(NumSolverIterations, Gravity);
 			Test.InitChain(NumBodies, FVec3(0, 0, -1));
-
 			Test.Evolution.GetJointConstraints().SetUseSimd(bUseSimd);
 
 			if (!Test.Evolution.GetJointConstraints().GetSettings().bUsePositionBasedDrives)
@@ -358,50 +381,55 @@ namespace ChaosTest {
 	//
 	GTEST_TEST(JointForceTests, TestLinearDriveForceMode_MaxForcePreTest)
 	{
-		const int32 NumSolverIterations = 20;
-		const FReal Gravity = 1000;
-		const FReal Dt = 0.01;
-		const int32 NumSteps = 100;
-		const int32 NumBodies = 2;
-		const FReal Extension = 10;
+		for (int32 IndexSimd = 0; IndexSimd < 2; IndexSimd++)
+		{
+			bool bUseSimd = IndexSimd == 0;
+			const int32 NumSolverIterations = 20;
+			const FReal Gravity = 1000;
+			const FReal Dt = 0.01;
+			const int32 NumSteps = 100;
+			const int32 NumBodies = 2;
+			const FReal Extension = 10;
 
-		FJointChainTest<FPBDRigidsEvolutionGBF> Test(NumSolverIterations, Gravity);
-		Test.InitChain(NumBodies, FVec3(0, 0, -1));
+			FJointChainTest<FPBDRigidsEvolutionGBF> Test(NumSolverIterations, Gravity);
+			Test.InitChain(NumBodies, FVec3(0, 0, -1));
+			Test.Evolution.GetJointConstraints().SetUseSimd(bUseSimd);
 
-		if (!Test.Evolution.GetJointConstraints().GetSettings().bUsePositionBasedDrives)
-		{
-			Test.Evolution.SetNumVelocityIterations(NumSolverIterations);
-		}
+			if (!Test.Evolution.GetJointConstraints().GetSettings().bUsePositionBasedDrives)
+			{
+				Test.Evolution.SetNumVelocityIterations(NumSolverIterations);
+			}
+
+			// Disable all limits
+			Test.JointSettings[0].LinearMotionTypes = { EJointMotionType::Free, EJointMotionType::Free, EJointMotionType::Free };
+			Test.JointSettings[0].AngularMotionTypes = { EJointMotionType::Free, EJointMotionType::Free, EJointMotionType::Free };
 
-		// Disable all limits
-		Test.JointSettings[0].LinearMotionTypes = { EJointMotionType::Free, EJointMotionType::Free, EJointMotionType::Free };
-		Test.JointSettings[0].AngularMotionTypes = { EJointMotionType::Free, EJointMotionType::Free, EJointMotionType::Free };
+			// Set up the drive without a force limit
+			const FReal Stiffness = Test.ParticleMasses[1] * Gravity / Extension;
+			const FReal Damping = 2 * FMath::Sqrt(Stiffness * Test.ParticleMasses[1]);
+			Test.JointSettings[0].bLinearPositionDriveEnabled = { false, false, true };
+			Test.JointSettings[0].bLinearVelocityDriveEnabled = { false, false, true };
+			Test.JointSettings[0].LinearDriveForceMode = EJointForceMode::Force;
+			Test.JointSettings[0].LinearDriveStiffness = FVec3(0, 0, Stiffness);
+			Test.JointSettings[0].LinearDriveDamping = FVec3(0, 0, Damping);
 
-		// Set up the drive without a force limit
-		const FReal Stiffness = Test.ParticleMasses[1] * Gravity / Extension;
-		const FReal Damping = 2 * FMath::Sqrt(Stiffness * Test.ParticleMasses[1]);
-		Test.JointSettings[0].bLinearPositionDriveEnabled = { false, false, true };
-		Test.JointSettings[0].bLinearVelocityDriveEnabled = { false, false, true };
-		Test.JointSettings[0].LinearDriveForceMode = EJointForceMode::Force;
-		Test.JointSettings[0].LinearDriveStiffness = FVec3(0, 0, Stiffness);
-		Test.JointSettings[0].LinearDriveDamping = FVec3(0, 0, Damping);
+			Test.Create();
 
-		Test.Create();
+			FConstGenericParticleHandle P1 = Test.GetParticle(1);
 
-		FConstGenericParticleHandle P1 = Test.GetParticle(1);
+			// Run the sim - the dangling box should reach a steady state
+			for (int32 i = 0; i < NumSteps; ++i)
+			{
+				Test.Evolution.AdvanceOneTimeStep(Dt);
+				Test.Evolution.EndFrame(Dt);
+			}
 
-		// Run the sim - the dangling box should reach a steady state
-		for (int32 i = 0; i < NumSteps; ++i)
-		{
-			Test.Evolution.AdvanceOneTimeStep(Dt);
-			Test.Evolution.EndFrame(Dt);
+			// We should be stationary at the desired extension
+			const FReal ExpectedVZ = 0;
+			const FReal ExpectedZ = Test.ParticlePositions[1].Z - Extension;
+			EXPECT_NEAR(P1->V().Z, ExpectedVZ, 1);
+			EXPECT_NEAR(P1->GetX().Z, ExpectedZ, 1);
 		}
-
-		// We should be stationary at the desired extension
-		const FReal ExpectedVZ = 0;
-		const FReal ExpectedZ = Test.ParticlePositions[1].Z - Extension;
-		EXPECT_NEAR(P1->V().Z, ExpectedVZ, 1);
-		EXPECT_NEAR(P1->GetX().Z, ExpectedZ, 1);
 	}
 
 	// Check that the maximum drive force setting honored for linear drives.
@@ -409,47 +437,52 @@ namespace ChaosTest {
 	//
 	GTEST_TEST(JointForceTests, TestLinearDriveForceMode_MaxForce)
 	{
-		const int32 NumSolverIterations = 20;
-		const FReal Gravity = 1000;
-		const FReal Dt = 0.01;
-		const int32 NumSteps = 100;
-		const int32 NumBodies = 2;
-		const FReal Extension = 10;
+		for (int32 IndexSimd = 0; IndexSimd < 2; IndexSimd++)
+		{
+			bool bUseSimd = IndexSimd == 0;
+			const int32 NumSolverIterations = 20;
+			const FReal Gravity = 1000;
+			const FReal Dt = 0.01;
+			const int32 NumSteps = 100;
+			const int32 NumBodies = 2;
+			const FReal Extension = 10;
 
-		FJointChainTest<FPBDRigidsEvolutionGBF> Test(NumSolverIterations, Gravity);
-		Test.InitChain(NumBodies, FVec3(0, 0, -1));
+			FJointChainTest<FPBDRigidsEvolutionGBF> Test(NumSolverIterations, Gravity);
+			Test.InitChain(NumBodies, FVec3(0, 0, -1));
+			Test.Evolution.GetJointConstraints().SetUseSimd(bUseSimd);
 
-		if (!Test.Evolution.GetJointConstraints().GetSettings().bUsePositionBasedDrives)
-		{
-			Test.Evolution.SetNumVelocityIterations(NumSolverIterations);
-		}
+			if (!Test.Evolution.GetJointConstraints().GetSettings().bUsePositionBasedDrives)
+			{
+				Test.Evolution.SetNumVelocityIterations(NumSolverIterations);
+			}
 
-		// Disable all limits
-		Test.JointSettings[0].LinearMotionTypes = { EJointMotionType::Free, EJointMotionType::Free, EJointMotionType::Free };
-		Test.JointSettings[0].AngularMotionTypes = { EJointMotionType::Free, EJointMotionType::Free, EJointMotionType::Free };
+			// Disable all limits
+			Test.JointSettings[0].LinearMotionTypes = { EJointMotionType::Free, EJointMotionType::Free, EJointMotionType::Free };
+			Test.JointSettings[0].AngularMotionTypes = { EJointMotionType::Free, EJointMotionType::Free, EJointMotionType::Free };
 
-		// Set up the drive with a force limit
-		const FReal Stiffness = Test.ParticleMasses[1] * Gravity / Extension;
-		const FReal Damping = 2 * FMath::Sqrt(Stiffness * Test.ParticleMasses[1]);
-		Test.JointSettings[0].bLinearPositionDriveEnabled = { false, false, true };
-		Test.JointSettings[0].bLinearVelocityDriveEnabled = { false, false, true };
-		Test.JointSettings[0].LinearDriveForceMode = EJointForceMode::Force;
-		Test.JointSettings[0].LinearDriveStiffness = FVec3(0, 0, Stiffness);
-		Test.JointSettings[0].LinearDriveDamping = FVec3(0, 0, Damping);
-		Test.JointSettings[0].LinearDriveMaxForce = FVec3(0,0, 0.5 * Stiffness * Extension);
+			// Set up the drive with a force limit
+			const FReal Stiffness = Test.ParticleMasses[1] * Gravity / Extension;
+			const FReal Damping = 2 * FMath::Sqrt(Stiffness * Test.ParticleMasses[1]);
+			Test.JointSettings[0].bLinearPositionDriveEnabled = { false, false, true };
+			Test.JointSettings[0].bLinearVelocityDriveEnabled = { false, false, true };
+			Test.JointSettings[0].LinearDriveForceMode = EJointForceMode::Force;
+			Test.JointSettings[0].LinearDriveStiffness = FVec3(0, 0, Stiffness);
+			Test.JointSettings[0].LinearDriveDamping = FVec3(0, 0, Damping);
+			Test.JointSettings[0].LinearDriveMaxForce = FVec3(0, 0, 0.5 * Stiffness * Extension);
 
-		Test.Create();
+			Test.Create();
 
-		FConstGenericParticleHandle P1 = Test.GetParticle(1);
+			FConstGenericParticleHandle P1 = Test.GetParticle(1);
 
-		// Run the sim - the dangling box should reach a steady state
-		for (int32 i = 0; i < NumSteps; ++i)
-		{
-			Test.Evolution.AdvanceOneTimeStep(Dt);
-			Test.Evolution.EndFrame(Dt);
+			// Run the sim - the dangling box should reach a steady state
+			for (int32 i = 0; i < NumSteps; ++i)
+			{
+				Test.Evolution.AdvanceOneTimeStep(Dt);
+				Test.Evolution.EndFrame(Dt);
 
-			const FReal ForceZ = -Test.Evolution.GetJointConstraints().GetConstraintLinearImpulse(0).Z / Dt;
-			EXPECT_LT(ForceZ, Test.JointSettings[0].LinearDriveMaxForce.Z + 0.1);
+				const FReal ForceZ = -Test.Evolution.GetJointConstraints().GetConstraintLinearImpulse(0).Z / Dt;
+				EXPECT_LT(ForceZ, Test.JointSettings[0].LinearDriveMaxForce.Z + 0.1);
+			}
 		}
 	}
 
@@ -460,51 +493,56 @@ namespace ChaosTest {
 	//
 	GTEST_TEST(JointForceTests, TestLinearDriveAccMode_MaxForcePreTest)
 	{
-		const int32 NumSolverIterations = 20;
-		const FReal Gravity = 1000;
-		const FReal Dt = 0.01;
-		const int32 NumSteps = 100;
-		const int32 NumBodies = 2;
-		const FReal Extension = 10;
+		for (int32 IndexSimd = 0; IndexSimd < 2; IndexSimd++)
+		{
+			bool bUseSimd = IndexSimd == 0;
+			const int32 NumSolverIterations = 20;
+			const FReal Gravity = 1000;
+			const FReal Dt = 0.01;
+			const int32 NumSteps = 100;
+			const int32 NumBodies = 2;
+			const FReal Extension = 10;
 
-		FJointChainTest<FPBDRigidsEvolutionGBF> Test(NumSolverIterations, Gravity);
-		Test.InitChain(NumBodies, FVec3(0, 0, -1));
+			FJointChainTest<FPBDRigidsEvolutionGBF> Test(NumSolverIterations, Gravity);
+			Test.InitChain(NumBodies, FVec3(0, 0, -1));
+			Test.Evolution.GetJointConstraints().SetUseSimd(bUseSimd);
 
-		if (!Test.Evolution.GetJointConstraints().GetSettings().bUsePositionBasedDrives)
-		{
-			Test.Evolution.SetNumVelocityIterations(NumSolverIterations);
-		}
+			if (!Test.Evolution.GetJointConstraints().GetSettings().bUsePositionBasedDrives)
+			{
+				Test.Evolution.SetNumVelocityIterations(NumSolverIterations);
+			}
 
-		// Disable all limits
-		Test.JointSettings[0].LinearMotionTypes = { EJointMotionType::Free, EJointMotionType::Free, EJointMotionType::Free };
-		Test.JointSettings[0].AngularMotionTypes = { EJointMotionType::Free, EJointMotionType::Free, EJointMotionType::Free };
+			// Disable all limits
+			Test.JointSettings[0].LinearMotionTypes = { EJointMotionType::Free, EJointMotionType::Free, EJointMotionType::Free };
+			Test.JointSettings[0].AngularMotionTypes = { EJointMotionType::Free, EJointMotionType::Free, EJointMotionType::Free };
 
-		// Set up the drive without a force limit
-		// NOTE: Acceleration mode - no masses in expressions
-		const FReal Stiffness = Gravity / Extension;
-		const FReal Damping = 2 * FMath::Sqrt(Stiffness);
-		Test.JointSettings[0].bLinearPositionDriveEnabled = { false, false, true };
-		Test.JointSettings[0].bLinearVelocityDriveEnabled = { false, false, true };
-		Test.JointSettings[0].LinearDriveForceMode = EJointForceMode::Acceleration;
-		Test.JointSettings[0].LinearDriveStiffness = FVec3(0, 0, Stiffness);
-		Test.JointSettings[0].LinearDriveDamping = FVec3(0, 0, Damping);
+			// Set up the drive without a force limit
+			// NOTE: Acceleration mode - no masses in expressions
+			const FReal Stiffness = Gravity / Extension;
+			const FReal Damping = 2 * FMath::Sqrt(Stiffness);
+			Test.JointSettings[0].bLinearPositionDriveEnabled = { false, false, true };
+			Test.JointSettings[0].bLinearVelocityDriveEnabled = { false, false, true };
+			Test.JointSettings[0].LinearDriveForceMode = EJointForceMode::Acceleration;
+			Test.JointSettings[0].LinearDriveStiffness = FVec3(0, 0, Stiffness);
+			Test.JointSettings[0].LinearDriveDamping = FVec3(0, 0, Damping);
 
-		Test.Create();
+			Test.Create();
 
-		FConstGenericParticleHandle P1 = Test.GetParticle(1);
+			FConstGenericParticleHandle P1 = Test.GetParticle(1);
 
-		// Run the sim - the dangling box should reach a steady state
-		for (int32 i = 0; i < NumSteps; ++i)
-		{
-			Test.Evolution.AdvanceOneTimeStep(Dt);
-			Test.Evolution.EndFrame(Dt);
-		}
+			// Run the sim - the dangling box should reach a steady state
+			for (int32 i = 0; i < NumSteps; ++i)
+			{
+				Test.Evolution.AdvanceOneTimeStep(Dt);
+				Test.Evolution.EndFrame(Dt);
+			}
 
-		// We should be stationary at the desired extension
-		const FReal ExpectedVZ = 0;
-		const FReal ExpectedZ = Test.ParticlePositions[1].Z - Extension;
-		EXPECT_NEAR(P1->V().Z, ExpectedVZ, 1);
-		EXPECT_NEAR(P1->GetX().Z, ExpectedZ, 1);
+			// We should be stationary at the desired extension
+			const FReal ExpectedVZ = 0;
+			const FReal ExpectedZ = Test.ParticlePositions[1].Z - Extension;
+			EXPECT_NEAR(P1->V().Z, ExpectedVZ, 1);
+			EXPECT_NEAR(P1->GetX().Z, ExpectedZ, 1);
+		}
 	}
 
 
@@ -514,48 +552,53 @@ namespace ChaosTest {
 	//
 	GTEST_TEST(JointForceTests, TestLinearDriveAccMode_MaxForce)
 	{
-		const int32 NumSolverIterations = 20;
-		const FReal Gravity = 1000;
-		const FReal Dt = 0.01;
-		const int32 NumSteps = 100;
-		const int32 NumBodies = 2;
-		const FReal Extension = 10;
+		for (int32 IndexSimd = 0; IndexSimd < 2; IndexSimd++)
+		{
+			bool bUseSimd = IndexSimd == 0;
+			const int32 NumSolverIterations = 20;
+			const FReal Gravity = 1000;
+			const FReal Dt = 0.01;
+			const int32 NumSteps = 100;
+			const int32 NumBodies = 2;
+			const FReal Extension = 10;
 
-		FJointChainTest<FPBDRigidsEvolutionGBF> Test(NumSolverIterations, Gravity);
-		Test.InitChain(NumBodies, FVec3(0, 0, -1));
+			FJointChainTest<FPBDRigidsEvolutionGBF> Test(NumSolverIterations, Gravity);
+			Test.InitChain(NumBodies, FVec3(0, 0, -1));
+			Test.Evolution.GetJointConstraints().SetUseSimd(bUseSimd);
 
-		if (!Test.Evolution.GetJointConstraints().GetSettings().bUsePositionBasedDrives)
-		{
-			Test.Evolution.SetNumVelocityIterations(NumSolverIterations);
-		}
+			if (!Test.Evolution.GetJointConstraints().GetSettings().bUsePositionBasedDrives)
+			{
+				Test.Evolution.SetNumVelocityIterations(NumSolverIterations);
+			}
 
-		// Disable all limits
-		Test.JointSettings[0].LinearMotionTypes = { EJointMotionType::Free, EJointMotionType::Free, EJointMotionType::Free };
-		Test.JointSettings[0].AngularMotionTypes = { EJointMotionType::Free, EJointMotionType::Free, EJointMotionType::Free };
+			// Disable all limits
+			Test.JointSettings[0].LinearMotionTypes = { EJointMotionType::Free, EJointMotionType::Free, EJointMotionType::Free };
+			Test.JointSettings[0].AngularMotionTypes = { EJointMotionType::Free, EJointMotionType::Free, EJointMotionType::Free };
 
-		// Set up the drive with a force limit
-		// NOTE: Acceleration mode - no masses in expressions
-		const FReal Stiffness = Gravity / Extension;
-		const FReal Damping = 2 * FMath::Sqrt(Stiffness);
-		Test.JointSettings[0].bLinearPositionDriveEnabled = { false, false, true };
-		Test.JointSettings[0].bLinearVelocityDriveEnabled = { false, false, true };
-		Test.JointSettings[0].LinearDriveForceMode = EJointForceMode::Acceleration;
-		Test.JointSettings[0].LinearDriveStiffness = FVec3(0, 0, Stiffness);
-		Test.JointSettings[0].LinearDriveDamping = FVec3(0, 0, Damping);
-		Test.JointSettings[0].LinearDriveMaxForce = FVec3(0, 0, 0.5 * Stiffness * Extension);
+			// Set up the drive with a force limit
+			// NOTE: Acceleration mode - no masses in expressions
+			const FReal Stiffness = Gravity / Extension;
+			const FReal Damping = 2 * FMath::Sqrt(Stiffness);
+			Test.JointSettings[0].bLinearPositionDriveEnabled = { false, false, true };
+			Test.JointSettings[0].bLinearVelocityDriveEnabled = { false, false, true };
+			Test.JointSettings[0].LinearDriveForceMode = EJointForceMode::Acceleration;
+			Test.JointSettings[0].LinearDriveStiffness = FVec3(0, 0, Stiffness);
+			Test.JointSettings[0].LinearDriveDamping = FVec3(0, 0, Damping);
+			Test.JointSettings[0].LinearDriveMaxForce = FVec3(0, 0, 0.5 * Stiffness * Extension);
 
-		Test.Create();
+			Test.Create();
 
-		FConstGenericParticleHandle P1 = Test.GetParticle(1);
+			FConstGenericParticleHandle P1 = Test.GetParticle(1);
 
-		// Run the sim - the dangling box should reach a steady state
-		for (int32 i = 0; i < NumSteps; ++i)
-		{
-			Test.Evolution.AdvanceOneTimeStep(Dt);
-			Test.Evolution.EndFrame(Dt);
+			// Run the sim - the dangling box should reach a steady state
+			for (int32 i = 0; i < NumSteps; ++i)
+			{
+				Test.Evolution.AdvanceOneTimeStep(Dt);
+				Test.Evolution.EndFrame(Dt);
 
-			const FReal ForceZ = -Test.Evolution.GetJointConstraints().GetConstraintLinearImpulse(0).Z / Dt;
-			EXPECT_LT(ForceZ, Test.JointSettings[0].LinearDriveMaxForce.Z * Test.ParticleMasses[1] + 0.1);
+				const FReal ForceZ = -Test.Evolution.GetJointConstraints().GetConstraintLinearImpulse(0).Z / Dt;
+				EXPECT_LT(ForceZ, Test.JointSettings[0].LinearDriveMaxForce.Z * Test.ParticleMasses[1] + 0.1);
+			}
 		}
 	}
 }
\ No newline at end of file
diff --git a/Engine/Source/Runtime/Experimental/Chaos/Public/Chaos/Capsule.h b/Engine/Source/Runtime/Experimental/Chaos/Public/Chaos/Capsule.h
--- a/Engine/Source/Runtime/Experimental/Chaos/Public/Chaos/Capsule.h
+++ b/Engine/Source/Runtime/Experimental/Chaos/Public/Chaos/Capsule.h
@@ -515,7 +515,7 @@ namespace Chaos
 		// Used for manifold generation
 		const FVec3 GetVertex(int32 VertexIndex) const
 		{
-			FVec3 Result;
+			FVec3 Result(0);
 
 			switch (VertexIndex)
 			{
