ChaosFlesh: Fix HeadlessChaos, ComputeFiberField parameters

diff --git a/Engine/Source/Programs/HeadlessChaos/Private/HeadlessChaosTestMath.cpp b/Engine/Source/Programs/HeadlessChaos/Private/HeadlessChaosTestMath.cpp
--- a/Engine/Source/Programs/HeadlessChaos/Private/HeadlessChaosTestMath.cpp
+++ b/Engine/Source/Programs/HeadlessChaos/Private/HeadlessChaosTestMath.cpp
@@ -335,7 +335,8 @@ namespace ChaosTest
 		}
 
 		TArray<Chaos::TVector<double, 3>> Directions;
-		Chaos::ComputeFiberField<double>(Mesh, X, IncidentElements, IncidentElementsLocalIndex, Origins, Insertions, Directions);
+		TArray<double> ScalarField;
+		Chaos::ComputeFiberField<double>(Mesh, X, IncidentElements, IncidentElementsLocalIndex, Origins, Insertions, Directions, ScalarField);
 
 		for(int32 e=0; e < Mesh.Num(); e++)
 		{
