Fix Event Manager in Single Buffer Mode

diff --git a/Engine/Source/Programs/HeadlessChaos/Private/GeometryCollection/GeometryCollectionTestEvents.cpp b/Engine/Source/Programs/HeadlessChaos/Private/GeometryCollection/GeometryCollectionTestEvents.cpp
--- a/Engine/Source/Programs/HeadlessChaos/Private/GeometryCollection/GeometryCollectionTestEvents.cpp
+++ b/Engine/Source/Programs/HeadlessChaos/Private/GeometryCollection/GeometryCollectionTestEvents.cpp
@@ -6,6 +6,7 @@
 #include "PBDRigidsSolver.h"
 #include "ChaosSolversModule.h"
 #include "HeadlessChaosTestUtility.h"
+#include "Chaos/ChaosScene.h"
 
 namespace GeometryCollectionTest
 {
@@ -19,7 +20,7 @@ namespace GeometryCollectionTest
 	struct EventTestData
 	{
 		EventTestData(int InData1, FVector InData2) : Data1(InData1), Data2(InData2) {}
-		EventTestData() {}
+		EventTestData() : Data1(0xBADF000D), Data2(FVector(0xBADF000D, 0xBADF000D, 0xBADF000D)) {}
 
 		bool operator==(const EventTestData& Other) const
 		{
@@ -28,6 +29,8 @@ namespace GeometryCollectionTest
 
 		void Reset()
 		{
+			Data1 = 0xBADF000D;
+			Data2 = FVector(0xBADF000D, 0xBADF000D, 0xBADF000D);
 		}
 
 		int Data1;
@@ -185,6 +188,94 @@ namespace GeometryCollectionTest
 		EXPECT_EQ(HandlerTest.ResultFromHandler2, TestArrayData);
 		EXPECT_EQ(AnotherHandlerTest.ResultFromHandler2, TestArrayData);
 	}
+
+
+	struct SimpleEventHandler
+	{	
+		void HandleEvent(const EventTestData& EventData)
+		{
+			ResultFromHandler = EventData;
+		}
+		EventTestData ResultFromHandler;
+	};
+
+	GTEST_TEST(AllTraits, EventTestBuffer)
+	{
+		for (EMultiBufferMode BufferMode : {EMultiBufferMode::Single, EMultiBufferMode::Double, EMultiBufferMode::Triple/*, EMultiBufferMode::TripleGuarded*/})
+		{
+			const FReal FixedDT = 1;
+			FChaosScene Scene(nullptr, FixedDT);
+			Scene.GetSolver()->SetThreadingMode_External(EThreadingModeTemp::SingleThread);
+
+			Chaos::FEventManager* EventManager = Scene.GetSolver()->GetEventManager();
+			EventManager->SetBufferMode(BufferMode);
+
+			Chaos::FPBDRigidsSolver* Solver = Scene.GetSolver();
+
+			EventTestData TestData;
+			TestData.Data1 = 123;
+			TestData.Data2 = FVector(1, 2, 3);
+			EventTestData* TestDataPtr = &TestData;
+			EventManager->template RegisterEvent<EventTestData>(CustomEvent1, [TestDataPtr]
+			(const auto* Solver, EventTestData& MyData, bool ResetData)
+				{
+					MyData = *TestDataPtr;
+				});
+
+			SimpleEventHandler EventHandler;
+			EventManager->template RegisterHandler<EventTestData>(CustomEvent1, &EventHandler, &SimpleEventHandler::HandleEvent);
+
+			FVec3 Grav(0, 0, -1);
+			Scene.SetUpForFrame(&Grav, 1, 0, 99999, 99999, 10, false);
+			Scene.StartFrame();
+			Scene.EndFrame();
+
+			EXPECT_EQ(EventHandler.ResultFromHandler, TestData);
+		}
+	}
+
+	GTEST_TEST(AllTraits, EventTestBuffer_SeveralTicks)
+	{
+		for (EMultiBufferMode BufferMode : {EMultiBufferMode::Single, EMultiBufferMode::Double, EMultiBufferMode::Triple})
+		{
+			const FReal FixedDT = 1;
+			FChaosScene Scene(nullptr, FixedDT);
+			Scene.GetSolver()->SetThreadingMode_External(EThreadingModeTemp::SingleThread);
+
+			Chaos::FEventManager* EventManager = Scene.GetSolver()->GetEventManager();
+			EventManager->SetBufferMode(BufferMode);
+
+			Chaos::FPBDRigidsSolver* Solver = Scene.GetSolver();
+
+			EventTestData TestData;
+			TestData.Data1 = 0;
+			TestData.Data2 = FVector(0, 0, 0);
+			EventTestData* TestDataPtr = &TestData;
+			EventManager->template RegisterEvent<EventTestData>(CustomEvent1, [TestDataPtr]
+			(const auto* Solver, EventTestData& MyData, bool ResetData)
+				{
+					MyData = *TestDataPtr;
+				});
+
+
+			SimpleEventHandler EventHandler;
+			EventHandler.ResultFromHandler = TestData;
+			EventManager->template RegisterHandler<EventTestData>(CustomEvent1, &EventHandler, &SimpleEventHandler::HandleEvent);
+
+			FVec3 Grav(0, 0, -1);
+			Scene.SetUpForFrame(&Grav, 1, 0, 99999, 99999, 10, false);
+
+
+			for (int32 Index = 0; Index < 3; Index++)
+			{
+				TestData.Data1 += 1;
+				TestData.Data2 += FVec3(1);
+				Scene.StartFrame();
+				Scene.EndFrame();
+				EXPECT_EQ(EventHandler.ResultFromHandler, TestData);
+			}
+		}
+	}
 }
 
 
diff --git a/Engine/Source/Runtime/Experimental/Chaos/Private/EventManager.cpp b/Engine/Source/Runtime/Experimental/Chaos/Private/EventManager.cpp
--- a/Engine/Source/Runtime/Experimental/Chaos/Private/EventManager.cpp
+++ b/Engine/Source/Runtime/Experimental/Chaos/Private/EventManager.cpp
@@ -60,7 +60,11 @@ namespace Chaos
 
 	void FEventManager::FlipBuffersIfRequired()
 	{
-		if (BufferMode == EMultiBufferMode::Double)
+		if (BufferMode == EMultiBufferMode::Single)
+		{
+			return;
+		}
+		else if (BufferMode == EMultiBufferMode::Double)
 		{
 			ResourceLock.WriteLock();
 		}
@@ -102,6 +106,17 @@ namespace Chaos
 		if (BufferMode == EMultiBufferMode::Double)
 		{
 			ResourceLock.ReadUnlock();
+		} 
+		else if (BufferMode == EMultiBufferMode::Single)
+		{
+			for (FEventContainerBasePtr EventContainer : EventContainers)
+			{
+				if (EventContainer)
+				{
+					EventContainer->ResetConsumerBuffer();
+					EventContainer->FlipBufferIfRequired();
+				}
+			}
 		}
 	}
 
