Fix CIS in Headless Chaos Test project. This file was missing from a previous CL33855917

diff --git a/Engine/Source/Programs/HeadlessChaos/Private/HeadlessChaosTestModularVehicles.cpp b/Engine/Source/Programs/HeadlessChaos/Private/HeadlessChaosTestModularVehicles.cpp
--- a/Engine/Source/Programs/HeadlessChaos/Private/HeadlessChaosTestModularVehicles.cpp
+++ b/Engine/Source/Programs/HeadlessChaos/Private/HeadlessChaosTestModularVehicles.cpp
@@ -10,7 +10,6 @@
 #include "SimModule/ClutchModule.h"
 #include "SimModule/TransmissionModule.h"
 #include "SimModule/WheelModule.h"
-#include "SimModule/SuspensionModule.h"
 #include "SimModule/SimModuleTree.h"
 #include "SimModule/ModuleInput.h"
 
