Fixed HeadlessChaos after GC memory optimizations

diff --git a/Engine/Source/Programs/HeadlessChaos/Private/GeometryCollection/GeometryCollectionTestSimulationField.cpp b/Engine/Source/Programs/HeadlessChaos/Private/GeometryCollection/GeometryCollectionTestSimulationField.cpp
--- a/Engine/Source/Programs/HeadlessChaos/Private/GeometryCollection/GeometryCollectionTestSimulationField.cpp
+++ b/Engine/Source/Programs/HeadlessChaos/Private/GeometryCollection/GeometryCollectionTestSimulationField.cpp
@@ -168,7 +168,7 @@ namespace GeometryCollectionTest
 		Params.RootTransform.SetLocation(FVector(0.0, 0.0, 0.0));
 
 		Params.InitialVelocityType = EInitialVelocityTypeEnum::Chaos_Initial_Velocity_User_Defined;
-		Params.InitialLinearVelocity = FVector(0.f, 100.f, 0.f);
+		Params.InitialLinearVelocity = FVector3f(0.f, 100.f, 0.f);
 		FGeometryCollectionWrapper* Collection = TNewSimulationObject<GeometryType::GeometryCollectionWithSingleRigid>::Init(Params)->template As<FGeometryCollectionWrapper>();
 		UnitTest.AddSimulationObject(Collection);
 
@@ -1054,7 +1054,7 @@ namespace GeometryCollectionTest
 		CreationParameters Params;
 		Params.DynamicState = EObjectStateTypeEnum::Chaos_Object_Dynamic;
 		Params.InitialVelocityType = EInitialVelocityTypeEnum::Chaos_Initial_Velocity_User_Defined;
-		Params.InitialLinearVelocity = FVector(100.0, 0.0, 0.0);
+		Params.InitialLinearVelocity = FVector3f(100.0, 0.0, 0.0);
 		Params.RootTransform.SetLocation(ExpectedLocation);
 		FGeometryCollectionWrapper* Collection = TNewSimulationObject<GeometryType::GeometryCollectionWithSingleRigid>::Init(Params)->template As<FGeometryCollectionWrapper>();
 		UnitTest.AddSimulationObject(Collection);
diff --git a/Engine/Source/Programs/HeadlessChaos/Public/GeometryCollection/GeometryCollectionTestFramework.h b/Engine/Source/Programs/HeadlessChaos/Public/GeometryCollection/GeometryCollectionTestFramework.h
--- a/Engine/Source/Programs/HeadlessChaos/Public/GeometryCollection/GeometryCollectionTestFramework.h
+++ b/Engine/Source/Programs/HeadlessChaos/Public/GeometryCollection/GeometryCollectionTestFramework.h
@@ -76,7 +76,7 @@ namespace GeometryCollectionTest
 		 * Implicit box uses Scale X, Y, Z for dimensions.
 		 * Implicit sphere uses Scale X for radius.
 		 */
-		FVector InitialLinearVelocity = FVector::ZeroVector;
+		FVector3f InitialLinearVelocity = FVector3f::ZeroVector;
 		EObjectStateTypeEnum DynamicState = EObjectStateTypeEnum::Chaos_Object_Dynamic;
 		bool Simulating = true;
 		FReal Mass = 1.0;
