[HeadlessChaos] ConvertToNonProbe test and trivial change to ConvertToProbe
- Add in ContactModification_NonProbe test case
- Modified existing ContactModification_Probe testcase to check for CubeB (CCD)

Disabling both tests at the moment as they're failing due to:
- UE-218719 (Probe)
- UE-216793 (Non-Probe)

#rnx

diff --git a/Engine/Source/Programs/HeadlessChaos/Private/HeadlessChaosTestContactModification.cpp b/Engine/Source/Programs/HeadlessChaos/Private/HeadlessChaosTestContactModification.cpp
--- a/Engine/Source/Programs/HeadlessChaos/Private/HeadlessChaosTestContactModification.cpp
+++ b/Engine/Source/Programs/HeadlessChaos/Private/HeadlessChaosTestContactModification.cpp
@@ -121,7 +121,8 @@ namespace ChaosTest
 		Module->DestroySolver(Solver);
 	}
 
-	GTEST_TEST(AllTraits, ContactModification_Probe)
+	// Disabling due to: UE-218719; CCD is being done before Contact Modification which is resulting in the object B colliding
+	GTEST_TEST(AllTraits, DISABLED_ContactModification_Probe)
 	{
 		FChaosSolversModule* Module = FChaosSolversModule::GetModule();
 		auto* Solver = Module->CreateSolver(nullptr, /*AsyncDt=*/-1);
@@ -189,7 +190,7 @@ namespace ChaosTest
 
 		// Both cubes should be below the floor
 		EXPECT_LT(CubeParticleA.X().Z, FloorParticle.X().Z);
-		EXPECT_LT(CubeParticleA.X().Z, FloorParticle.X().Z);
+		EXPECT_LT(CubeParticleB.X().Z, FloorParticle.X().Z);
 
 		// Floor should be at origin.
 		EXPECT_EQ(FloorParticle.X().Z, 0);
@@ -201,6 +202,88 @@ namespace ChaosTest
 		Module->DestroySolver(Solver);
 	}
 
+	// Disabling due to: UE-216793; Objects fall through the floor with ConvertToNonProbe modification.
+	GTEST_TEST(AllTraits, DISABLED_ContactModification_NonProbe)
+	{
+		FChaosSolversModule* Module = FChaosSolversModule::GetModule();
+		auto* Solver = Module->CreateSolver(nullptr, /*AsyncDt=*/-1);
+		InitSolverSettings(Solver);
+		Solver->SetThreadingMode_External(EThreadingModeTemp::SingleThread);
+
+		// Create a static floor and two cubes falling onto it.
+		// Both cubes start off as probes and should apply a modification to turn all contacts to non-probes, 
+		// CubeB has CCD enabled and CubeA does not.
+
+		auto CubeGeom = Chaos::FImplicitObjectPtr(new TBox<FReal, 3>(FVec3(-100), FVec3(100)));
+
+		// CubeA - Probe & No CCD
+		FSingleParticlePhysicsProxy* CubeProxyA = FSingleParticlePhysicsProxy::Create(Chaos::FPBDRigidParticle::CreateParticle());
+		auto& CubeParticleA = CubeProxyA->GetGameThreadAPI();
+		CubeParticleA.SetGeometry(CubeGeom);
+		Solver->RegisterObject(CubeProxyA);
+		CubeParticleA.SetGravityEnabled(false);
+		CubeParticleA.SetCCDEnabled(false);
+		CubeParticleA.SetV(FVec3(0, 0, -100));
+		CubeParticleA.SetX(FVec3(200, 0, 500));
+		CubeParticleA.ShapesArray()[0]->SetIsProbe(true);
+		SetCubeInertiaTensor(CubeParticleA, /*Dimension=*/200, /*Mass=*/1);
+		ChaosTest::SetParticleSimDataToCollide({ CubeProxyA->GetParticle_LowLevel() });
+
+		
+		// CubeB - Probe & With CCD
+		FSingleParticlePhysicsProxy* CubeProxyB = FSingleParticlePhysicsProxy::Create(Chaos::FPBDRigidParticle::CreateParticle());
+		auto& CubeParticleB = CubeProxyB->GetGameThreadAPI();
+		CubeParticleB.SetGeometry(CubeGeom);
+		Solver->RegisterObject(CubeProxyB);
+		CubeParticleB.SetGravityEnabled(false);
+		CubeParticleB.SetCCDEnabled(true);
+		CubeParticleB.SetV(FVec3(0, 0, -1000));
+		CubeParticleB.SetX(FVec3(-200, 0, 500));
+		CubeParticleB.ShapesArray()[0]->SetIsProbe(true);
+		SetCubeInertiaTensor(CubeParticleB, /*Dimension=*/200, /*Mass=*/1);
+		ChaosTest::SetParticleSimDataToCollide({ CubeProxyB->GetParticle_LowLevel() });
+		
+
+		// Static floor at origin, occupying Z = [-100,0], XY = [-500, 500]
+		FSingleParticlePhysicsProxy* FloorProxy = FSingleParticlePhysicsProxy::Create(Chaos::FGeometryParticle::CreateParticle());
+		auto& FloorParticle = FloorProxy->GetGameThreadAPI();
+		auto FloorGeom = Chaos::FImplicitObjectPtr(new TBox<FReal, 3>(FVec3(-500, -500, -100), FVec3(500, 500, 0)));
+		FloorParticle.SetGeometry(FloorGeom);
+		Solver->RegisterObject(FloorProxy);
+		FloorParticle.SetX(FVec3(0, 0, 0));
+		ChaosTest::SetParticleSimDataToCollide({ FloorProxy->GetParticle_LowLevel() });
+
+		FContactModificationTestCallback* Callback = Solver->CreateAndRegisterSimCallbackObject_External<FContactModificationTestCallback>();
+		Callback->TestLambda = [](Chaos::FCollisionContactModifier& Modifier)
+			{
+				for (FContactPairModifier& PairModifier : Modifier)
+				{
+					PairModifier.ConvertToNonProbe();
+				}
+			};
+
+		const float Dt = 1.0f;
+		const int32 Steps = 10;
+		for (int Step = 0; Step < Steps; ++Step)
+		{
+			Solver->AdvanceAndDispatch_External(Dt);
+			Solver->UpdateGameThreadStructures();
+		}
+
+		// Both cubes should be above the floor
+		EXPECT_GT(CubeParticleA.X().Z, FloorParticle.X().Z);
+		EXPECT_GT(CubeParticleB.X().Z, FloorParticle.X().Z);
+
+		// Floor should still be at origin.
+		EXPECT_EQ(FloorParticle.X().Z, 0);
+
+		Solver->UnregisterAndFreeSimCallbackObject_External(Callback);
+		Solver->UnregisterObject(CubeProxyA);
+		Solver->UnregisterObject(CubeProxyB);
+		Solver->UnregisterObject(FloorProxy);
+		Module->DestroySolver(Solver);
+	}
+
 	GTEST_TEST(AllTraits, ContactModification_ModifySeparation)
 	{
 		// The amount to pad the saparation by in the collision callback. Currently this must
