Fix crash in Material.
Physics Thread was accessing the material list from the Game Thread

diff --git a/Engine/Source/Programs/HeadlessChaos/Private/HeadlessChaosTestShapePadding.cpp b/Engine/Source/Programs/HeadlessChaos/Private/HeadlessChaosTestShapePadding.cpp
--- a/Engine/Source/Programs/HeadlessChaos/Private/HeadlessChaosTestShapePadding.cpp
+++ b/Engine/Source/Programs/HeadlessChaos/Private/HeadlessChaosTestShapePadding.cpp
@@ -12,6 +12,7 @@
 #include "Chaos/Sphere.h"
 #include "Chaos/GJK.h"
 #include "Chaos/Pair.h"
+#include "Chaos/PhysicalMaterials.h"
 #include "Chaos/PBDCollisionConstraints.h"
 #include "Chaos/PBDRigidParticles.h"
 #include "Chaos/PBDRigidsEvolution.h"
@@ -354,4 +355,40 @@ namespace ChaosTest {
 		TestBoxRayCastsMargin(70, FVec3(100, 100, 100), FVec3(-200, 0, 0), FVec3(1, 0, 0), 500.0f, true, 150.0f, FVec3(-50, 0, 0), FVec3(-1, 0, 0));	// Too much margin
 	}
 
+
+	TEST(ShapeInstanceTests, TestSingleMaterial)
+	{
+		{
+		FImplicitObjectPtr CollidingCubeGeom = MakeImplicitObjectPtr<TBox<FReal, 3>>(-FVec3(100.0), FVec3(100.0));
+		TUniquePtr<FShapeInstance> ShapeInstance1 = FShapeInstance::Make(0, CollidingCubeGeom);
+		EXPECT_EQ(ShapeInstance1->NumMaterials(), 0);
+
+		FMaterialHandle MaterialHandle = FPhysicalMaterialManager::Get().Create();
+		FMaterialData MaterialData;
+		MaterialData.Materials.Add(MaterialHandle);
+		ShapeInstance1->SetMaterialData(MaterialData); 
+
+		EXPECT_EQ(ShapeInstance1->NumMaterials(), 1);
+		EXPECT_TRUE(ShapeInstance1->GetMaterial(0).InnerHandle.IsValid());
+	}
+		{
+			THandleArray<FChaosPhysicsMaterial> SimMaterials;
+			FChaosMaterialHandle MaterialInnerHandle = SimMaterials.Create();
+			FMaterialHandle MaterialHandle = FPhysicalMaterialManager::Get().Create();
+			MaterialHandle.InnerHandle = MaterialInnerHandle;
+
+			FImplicitObjectPtr CollidingCubeGeom = MakeImplicitObjectPtr<TBox<FReal, 3>>(-FVec3(100.0), FVec3(100.0));
+			TUniquePtr<FShapeInstance> ShapeInstance1 = FShapeInstance::Make(0, CollidingCubeGeom);
+			EXPECT_EQ(ShapeInstance1->NumMaterialsInternal(&SimMaterials), 0);
+
+			
+			FMaterialData MaterialData;
+			MaterialData.Materials.Add(MaterialHandle);
+			ShapeInstance1->SetMaterialData(MaterialData);
+
+			EXPECT_EQ(ShapeInstance1->NumMaterialsInternal(&SimMaterials), 1);
+			EXPECT_TRUE(ShapeInstance1->GetMaterialInternal(0, &SimMaterials).InnerHandle.IsValid());
+		}
+	}
+
 }
\ No newline at end of file
diff --git a/Engine/Source/Runtime/Experimental/Chaos/Private/Chaos/PhysicalMaterials.cpp b/Engine/Source/Runtime/Experimental/Chaos/Private/Chaos/PhysicalMaterials.cpp
--- a/Engine/Source/Runtime/Experimental/Chaos/Private/Chaos/PhysicalMaterials.cpp
+++ b/Engine/Source/Runtime/Experimental/Chaos/Private/Chaos/PhysicalMaterials.cpp
@@ -24,6 +24,16 @@ namespace Chaos
 		return nullptr;
 	}
 
+	FChaosPhysicsMaterial* FMaterialHandle::GetInternal(const THandleArray<FChaosPhysicsMaterial>* const SimMaterials) const
+	{
+		if (SimMaterials != nullptr && InnerHandle.IsValid())
+		{
+			return SimMaterials->Get(InnerHandle);
+		}
+		return nullptr;
+	}
+
+
 	FChaosPhysicsMaterialMask* FMaterialMaskHandle::Get() const
 	{
 		if (InnerHandle.IsValid())
diff --git a/Engine/Source/Runtime/Experimental/Chaos/Private/Chaos/PhysicsMaterialUtilities.cpp b/Engine/Source/Runtime/Experimental/Chaos/Private/Chaos/PhysicsMaterialUtilities.cpp
--- a/Engine/Source/Runtime/Experimental/Chaos/Private/Chaos/PhysicsMaterialUtilities.cpp
+++ b/Engine/Source/Runtime/Experimental/Chaos/Private/Chaos/PhysicsMaterialUtilities.cpp
@@ -38,20 +38,20 @@ namespace Chaos
 			}
 
 			// If no particle material, see if the shape has one
-			if ((Shape != nullptr) && (SimMaterials != nullptr) && (Shape->NumMaterials() > 0))
+			if ((Shape != nullptr) && (SimMaterials != nullptr) && (Shape->NumMaterialsInternal(SimMaterials) > 0))
 			{
 				// By default we take the first material (all non-mesh shapes will use material 0)
 				int32 ShapeMaterialIndex = 0;
 
-				if ((Shape->NumMaterials() > 1) && (ShapeFaceIndex != INDEX_NONE))
+				if ((Shape->NumMaterialsInternal(SimMaterials) > 1) && (ShapeFaceIndex != INDEX_NONE))
 				{
 					// Multi-shape object (tri mesh, heightfield, etc) - use the face index
 					ShapeMaterialIndex = Shape->GetGeometry()->GetMaterialIndex(ShapeFaceIndex);
 				}
 
-				if (ShapeMaterialIndex < Shape->NumMaterials())
+				if (ShapeMaterialIndex < Shape->NumMaterialsInternal(SimMaterials))
 				{
-					return SimMaterials->Get(Shape->GetMaterial(ShapeMaterialIndex).InnerHandle);
+					return SimMaterials->Get(Shape->GetMaterialInternal(ShapeMaterialIndex, SimMaterials).InnerHandle);
 				}
 			}
 
diff --git a/Engine/Source/Runtime/Experimental/Chaos/Public/Chaos/PhysicalMaterials.h b/Engine/Source/Runtime/Experimental/Chaos/Public/Chaos/PhysicalMaterials.h
--- a/Engine/Source/Runtime/Experimental/Chaos/Public/Chaos/PhysicalMaterials.h
+++ b/Engine/Source/Runtime/Experimental/Chaos/Public/Chaos/PhysicalMaterials.h
@@ -22,11 +22,14 @@ namespace Chaos
 	 */
 	struct FMaterialHandle
 	{
+		// Get should be only called on the Game Thread, instead use GetInternal on Physics Thread
 		CHAOS_API FChaosPhysicsMaterial* Get() const;
+		CHAOS_API FChaosPhysicsMaterial* GetInternal(const THandleArray<FChaosPhysicsMaterial>* const SimMaterials) const;
 		FChaosMaterialHandle InnerHandle;
 
+		// IsValid should be only called on the Game Thread, instead use IsValidInternal on Physics Thread
 		bool IsValid() const { return Get() != nullptr; }
-
+		bool IsValidInternal(const THandleArray<FChaosPhysicsMaterial>* const SimMaterials) const { return GetInternal(SimMaterials) != nullptr; }
 		void Reset() { InnerHandle = FChaosMaterialHandle(); }
 
 		friend bool operator ==(const FMaterialHandle& A, const FMaterialHandle& B) { return A.InnerHandle == B.InnerHandle; }
diff --git a/Engine/Source/Runtime/Experimental/Chaos/Public/Chaos/ShapeInstance.h b/Engine/Source/Runtime/Experimental/Chaos/Public/Chaos/ShapeInstance.h
--- a/Engine/Source/Runtime/Experimental/Chaos/Public/Chaos/ShapeInstance.h
+++ b/Engine/Source/Runtime/Experimental/Chaos/Public/Chaos/ShapeInstance.h
@@ -600,6 +600,33 @@ namespace Chaos
 			}
 		}
 
+		int32 NumMaterialsInternal(const THandleArray<FChaosPhysicsMaterial>*const SimMaterials) const
+		{
+			if (bIsSingleMaterial)
+			{
+				return Material.MaterialHandle.IsValidInternal(SimMaterials) ? 1 : 0;
+			}
+			else
+			{
+				return (Material.MaterialData != nullptr) ? Material.MaterialData->Materials.Num() : 0;
+			}
+		}
+
+		const FMaterialHandle& GetMaterialInternal(const int32 Index, const THandleArray<FChaosPhysicsMaterial>*const SimMaterials) const
+		{
+			if (bIsSingleMaterial)
+			{
+				check(Index == 0);
+				check(Material.MaterialHandle.IsValidInternal(SimMaterials));
+				return Material.MaterialHandle;
+			}
+			else
+			{
+				return GetMaterialDataImpl().Materials[Index];
+			}
+		}
+
+
 		const TArray<FMaterialHandle>& GetMaterials() const { return GetMaterialDataImpl().Materials; }
 		const TArray<FMaterialMaskHandle>& GetMaterialMasks() const { return GetMaterialDataImpl().MaterialMasks; }
 		const TArray<uint32>& GetMaterialMaskMaps() const { return GetMaterialDataImpl().MaterialMaskMaps; }
