HeadlessChaos:: Add unit tests for per-body iteration settings and per-joint linear/nonlinear solver options.

diff --git a/Engine/Source/Programs/HeadlessChaos/Private/HeadlessChaosTestJoints.cpp b/Engine/Source/Programs/HeadlessChaos/Private/HeadlessChaosTestJoints.cpp
--- a/Engine/Source/Programs/HeadlessChaos/Private/HeadlessChaosTestJoints.cpp
+++ b/Engine/Source/Programs/HeadlessChaos/Private/HeadlessChaosTestJoints.cpp
@@ -50,6 +50,7 @@ namespace ChaosTest {
 
 			FPBDJointSettings Settings = Constraint->GetSettings();
 			Settings.LinearProjection = 1.0;
+			SetTestJointSetting(Settings);
 			Constraint->SetSettings(Settings);
 
 			// @todo(chaos): this should be automatic, but it's handled by the proxy. Fix this
@@ -74,6 +75,8 @@ namespace ChaosTest {
 			}
 		}
 
+		TFunction<void(FPBDJointSettings&)>& SetSetTestJointSetting() { return SetTestJointSetting; }
+
 		// Initial particles setup
 		TArray<FVec3> ParticlePositions;
 		TArray<FVec3> ParticleSizes;
@@ -82,6 +85,8 @@ namespace ChaosTest {
 		// Initial joints setup
 		TArray<FVec3> JointPositions;
 		TArray<TVec2<int32>> JointParticleIndices;
+
+		TFunction<void(FPBDJointSettings&)> SetTestJointSetting = [](FPBDJointSettings& InputSetting) {};
 	};
 
 	/**
@@ -121,6 +126,12 @@ namespace ChaosTest {
 			{ 0, 1 },
 		};
 
+		Test.SetSetTestJointSetting() = [](FPBDJointSettings& JointSettings)
+			{
+				JointSettings.bUseLinearSolver = false;
+			};
+
+
 		Test.Create();
 
 		const int32 Box1Id = 0;
@@ -472,6 +483,219 @@ namespace ChaosTest {
 		EXPECT_LT(MaxDistanceError, AcceptableDistanceError) << "On frame " << MaxDistanceErrorFrameIndex;
 	}
 
+	/**
+	 * Pendulum with animated root with varying number of iterations.
+	 */
+	template <typename TEvolution>
+	void JointConstraint_LongChainVaryingIterations()
+	{
+		const int NumParticles = 10;
+		const int32 NumIterations = 20;
+		const FReal Gravity = 980;
+		const FReal BoxSize = 100;
+		const FReal BoxMass = 1000;
+		const FReal Dt = (FReal)1 / 20;
+		const FReal AcceptableDistanceError = 5;
+		const FReal Separation = 2 * BoxSize;
+		const FVec3 Begin = FVec3(0, 0, (NumParticles + 10) * Separation);
+		const FVec3 Dir = FVec3(0, 1, 0);
+
+		// Create a chain of connected particles, with particle 0 fixed
+		FJointConstraintsTest<TEvolution> Test(NumIterations, Gravity);
+		for (int32 ParticleIndex = 0; ParticleIndex < NumParticles; ++ParticleIndex)
+		{
+			Test.ParticlePositions.Add(Begin + ParticleIndex * Separation * Dir);
+			Test.ParticleSizes.Add({ BoxSize, BoxSize, BoxSize });
+			Test.ParticleMasses.Add((ParticleIndex == 0) ? (FReal)0 : BoxMass);
+		}
+		for (int32 JointIndex = 0; JointIndex < NumParticles - 1; ++JointIndex)
+		{
+			Test.JointPositions.Add(Test.ParticlePositions[JointIndex]);
+			Test.JointParticleIndices.Add({ JointIndex, JointIndex + 1 });
+		}
+
+		Test.Create();
+
+		const FVec3 Box2LocalSpaceJointPosition = Test.JointPositions[0] - Test.ParticlePositions[1];
+
+		FReal MaxDistanceErrorBigIterations = 0.0f;
+		FReal MaxDistanceErrorSmallIterations = 0.0f;
+
+		const int32 SmallPositionIterations = 1;
+		const int32 SmallVelocityIterations = 1;
+		const int32 SmallProjectionIterations = 1;
+		const int32 BigPositionIterations = 100;
+		const int32 BigVelocityIterations = 100;
+		const int32 BigProjectionIterations = 100;
+
+		Test.Evolution.SetNumPositionIterations(BigPositionIterations);
+		Test.Evolution.SetNumVelocityIterations(BigVelocityIterations);
+		Test.Evolution.SetNumProjectionIterations(BigProjectionIterations);
+
+		const FVec3 RootPosition = Test.ParticlePositions[0];
+		Test.Evolution.SetParticleKinematicTarget(Test.GetParticle(0)->CastToKinematicParticle(), FKinematicTarget::MakePositionTarget(RootPosition, Test.GetParticle(0)->GetR()));
+
+
+		for (int32 FrameIndex = 0; FrameIndex < 10; ++FrameIndex)
+		{
+			const FReal Time = FrameIndex * Dt;
+			Test.Evolution.AdvanceOneTimeStep(Dt);
+			Test.Evolution.EndFrame(Dt);
+
+			Chaos::Private::FIterationSettings GroupIterations = Test.Evolution.GetIslandGroupManager().GetIslandGroupIterations(0);
+
+			EXPECT_EQ(GroupIterations.GetNumPositionIterations(), 100);
+			EXPECT_EQ(GroupIterations.GetNumVelocityIterations(), 100);
+			EXPECT_EQ(GroupIterations.GetNumProjectionIterations(), 100);
+		}
+
+		// Particles should remain fixed distance apart
+		for (int JointIndex = 0; JointIndex < Test.JointPositions.Num(); ++JointIndex)
+		{
+			const int32 ParticleIndex1 = Test.JointParticleIndices[JointIndex][0];
+			const int32 ParticleIndex2 = Test.JointParticleIndices[JointIndex][1];
+			const FVec3 Delta = Test.GetParticle(ParticleIndex2)->CastToRigidParticle()->GetP() - Test.GetParticle(ParticleIndex1)->GetX();
+			const FReal Distance = Delta.Size();
+			const FReal ExpectedDistance = (Test.ParticlePositions[ParticleIndex2] - Test.ParticlePositions[ParticleIndex1]).Size();
+			EXPECT_NEAR(Distance, ExpectedDistance, AcceptableDistanceError) << "Joint " << JointIndex;
+
+			// Track largest error that exceeds threshold
+			const FReal DistanceError = FMath::Abs(Distance - ExpectedDistance);
+			if (DistanceError > MaxDistanceErrorBigIterations)
+			{
+				MaxDistanceErrorBigIterations = DistanceError;
+			}
+		}
+
+		// Report the largest error and when it occurred if it exceeded the threshold
+		EXPECT_LT(MaxDistanceErrorBigIterations, AcceptableDistanceError);
+
+		Test.Evolution.SetNumPositionIterations(SmallPositionIterations);
+		Test.Evolution.SetNumVelocityIterations(SmallVelocityIterations);
+		Test.Evolution.SetNumProjectionIterations(SmallProjectionIterations);
+
+		for (int32 FrameIndex = 0; FrameIndex < 10; ++FrameIndex)
+		{
+			Test.Evolution.AdvanceOneTimeStep(Dt);
+			Test.Evolution.EndFrame(Dt);
+
+			Chaos::Private::FIterationSettings GroupIterations = Test.Evolution.GetIslandGroupManager().GetIslandGroupIterations(0);
+
+			EXPECT_EQ(GroupIterations.GetNumPositionIterations(), 1);
+			EXPECT_EQ(GroupIterations.GetNumVelocityIterations(), 1);
+			EXPECT_EQ(GroupIterations.GetNumProjectionIterations(), 1);
+
+		}
+		// Check max distance error between particles
+		for (int JointIndex = 0; JointIndex < Test.JointPositions.Num(); ++JointIndex)
+		{
+			const int32 ParticleIndex1 = Test.JointParticleIndices[JointIndex][0];
+			const int32 ParticleIndex2 = Test.JointParticleIndices[JointIndex][1];
+			const FVec3 Delta = Test.GetParticle(ParticleIndex2)->CastToRigidParticle()->GetP() - Test.GetParticle(ParticleIndex1)->GetX();
+			const FReal Distance = Delta.Size();
+			const FReal ExpectedDistance = (Test.ParticlePositions[ParticleIndex2] - Test.ParticlePositions[ParticleIndex1]).Size();
+			const FReal DistanceError = FMath::Abs(Distance - ExpectedDistance);
+			if (DistanceError > MaxDistanceErrorSmallIterations)
+			{
+				MaxDistanceErrorSmallIterations = DistanceError;
+			}
+		}
+		// Report if distance error from many iterations is bigger than that from small iterations. 
+		EXPECT_LT(MaxDistanceErrorBigIterations, MaxDistanceErrorSmallIterations);
+	}
+
+	/**
+	 * Two pendulums, one solved using the linear solver, the other using nonlinear solver
+	 */
+	template <typename TEvolution>
+	void JointConstraint_LongChainLinearJoints()
+	{
+		const int NumParticles = 10;
+		const int32 NumIterations = 5;
+		const FReal Gravity = 980;
+		const FReal BoxLength = 0.01;
+		const FReal BoxWidth = 0.005;
+		const FReal BoxMass = 10;
+		const FReal Dt = (FReal)1 / 20;
+		const FReal Separation = 2 * BoxLength;
+		const FVec3 Begin = FVec3(0, 0, (NumParticles + 10) * Separation);
+		const FVec3 Dir = FVec3(0, 1, 0);
+
+		// Create two chains of connected particles, with particle 0 fixed
+		FJointConstraintsTest<TEvolution> LinearJointTest(NumIterations, Gravity), NonLinearJointTest(NumIterations, Gravity);
+		for (int32 ParticleIndex = 0; ParticleIndex < NumParticles; ++ParticleIndex)
+		{
+			LinearJointTest.ParticlePositions.Add(Begin + ParticleIndex * Separation * Dir);
+			LinearJointTest.ParticleSizes.Add({ BoxWidth, BoxLength, BoxWidth });
+			LinearJointTest.ParticleMasses.Add((ParticleIndex == 0) ? (FReal)0 : BoxMass);
+			NonLinearJointTest.ParticlePositions.Add(Begin + ParticleIndex * Separation * Dir);
+			NonLinearJointTest.ParticleSizes.Add({ BoxWidth, BoxLength, BoxWidth });
+			NonLinearJointTest.ParticleMasses.Add((ParticleIndex == 0) ? (FReal)0 : BoxMass);
+		}
+		for (int32 JointIndex = 0; JointIndex < NumParticles - 1; ++JointIndex)
+		{
+			LinearJointTest.JointPositions.Add(LinearJointTest.ParticlePositions[JointIndex]);
+			LinearJointTest.JointParticleIndices.Add({ JointIndex, JointIndex + 1 });
+			NonLinearJointTest.JointPositions.Add(LinearJointTest.ParticlePositions[JointIndex]);
+			NonLinearJointTest.JointParticleIndices.Add({ JointIndex, JointIndex + 1 });
+		}
+
+		LinearJointTest.SetSetTestJointSetting() = [](FPBDJointSettings& JointSettings)
+			{
+				JointSettings.bUseLinearSolver = true;
+				JointSettings.AngularMotionTypes = TVector<EJointMotionType, 3>(EJointMotionType::Locked);
+			};
+
+		NonLinearJointTest.SetSetTestJointSetting() = [](FPBDJointSettings& JointSettings)
+			{
+				JointSettings.bUseLinearSolver = false;
+				JointSettings.AngularMotionTypes = TVector<EJointMotionType, 3>(EJointMotionType::Locked);
+			};
+
+		LinearJointTest.Create();
+		NonLinearJointTest.Create();
+		for (int32 ParticleIndex = 0; ParticleIndex < NumParticles; ++ParticleIndex)
+		{
+			LinearJointTest.GetParticle(ParticleIndex)->CastToRigidParticle()->SetInertiaConditioningEnabled(false);
+			NonLinearJointTest.GetParticle(ParticleIndex)->CastToRigidParticle()->SetInertiaConditioningEnabled(false);
+		}
+
+		const FVec3 Box2LocalSpaceJointPosition = LinearJointTest.JointPositions[0] - LinearJointTest.ParticlePositions[1];
+
+		FReal MaxDistanceError = 0.0f;
+		int32 MaxDistanceErrorFrameIndex = INDEX_NONE;
+		for (int32 FrameIndex = 0; FrameIndex < 100; ++FrameIndex)
+		{
+			const FReal Time = FrameIndex * Dt;
+
+			LinearJointTest.Evolution.SetParticleKinematicTarget(LinearJointTest.GetParticle(0)->CastToKinematicParticle(), FKinematicTarget::MakePositionTarget(LinearJointTest.ParticlePositions[0], LinearJointTest.GetParticle(0)->GetR()));
+
+			LinearJointTest.Evolution.AdvanceOneTimeStep(Dt);
+			LinearJointTest.Evolution.EndFrame(Dt);
+
+			NonLinearJointTest.Evolution.SetParticleKinematicTarget(NonLinearJointTest.GetParticle(0)->CastToKinematicParticle(), FKinematicTarget::MakePositionTarget(NonLinearJointTest.ParticlePositions[0], NonLinearJointTest.GetParticle(0)->GetR()));
+
+			NonLinearJointTest.Evolution.AdvanceOneTimeStep(Dt);
+			NonLinearJointTest.Evolution.EndFrame(Dt);
+		}
+
+		// The chain solved using the nonlinear solver should have less error than the one using linear solver. 
+		for (int JointIndex = 0; JointIndex < LinearJointTest.JointPositions.Num(); ++JointIndex)
+		{
+			const int32 ParticleIndex1 = LinearJointTest.JointParticleIndices[JointIndex][0];
+			const int32 ParticleIndex2 = LinearJointTest.JointParticleIndices[JointIndex][1];
+			const FVec3 LinearDelta = LinearJointTest.GetParticle(ParticleIndex2)->CastToRigidParticle()->GetP() - LinearJointTest.GetParticle(ParticleIndex1)->GetX();
+			const FReal LinearDistance = LinearDelta.Size();
+			const FVec3 NonLinearDelta = NonLinearJointTest.GetParticle(ParticleIndex2)->CastToRigidParticle()->GetP() - NonLinearJointTest.GetParticle(ParticleIndex1)->GetX();
+			const FReal NonLinearDistance = NonLinearDelta.Size();
+			const FReal ExpectedDistance = (LinearJointTest.ParticlePositions[ParticleIndex2] - LinearJointTest.ParticlePositions[ParticleIndex1]).Size();
+			const FReal LinearDistanceError = FMath::Abs(ExpectedDistance - LinearDistance);
+			const FReal NonLinearDistanceError = FMath::Abs(ExpectedDistance - NonLinearDistance);
+			EXPECT_LT(NonLinearDistanceError * 10, LinearDistanceError);
+		}
+	}
+
+
 	template <typename TEvolution>
 	void SpringConstraint()
 	{
@@ -679,6 +903,14 @@ namespace ChaosTest {
 		JointConstraint_LongChainAnimated<FPBDRigidsEvolutionGBF>();
 	}
 
+	GTEST_TEST(AllEvolutions, JointTests_TestLongJointChainWithVaryingIterations) {
+		JointConstraint_LongChainVaryingIterations<FPBDRigidsEvolutionGBF>();
+	}
+
+	GTEST_TEST(AllEvolutions, JointTests_TestLongJointChainWithNonLinearSolvers) {
+		JointConstraint_LongChainLinearJoints<FPBDRigidsEvolutionGBF>();
+	}
+
 	GTEST_TEST(AllEvolutions, JointTests_TestSingleSpringConstraint) {
 		SpringConstraint<FPBDRigidsEvolutionGBF>();
 	}
diff --git a/Engine/Source/Programs/HeadlessChaos/Private/HeadlessChaosTestPhysicsScene.cpp b/Engine/Source/Programs/HeadlessChaos/Private/HeadlessChaosTestPhysicsScene.cpp
--- a/Engine/Source/Programs/HeadlessChaos/Private/HeadlessChaosTestPhysicsScene.cpp
+++ b/Engine/Source/Programs/HeadlessChaos/Private/HeadlessChaosTestPhysicsScene.cpp
@@ -1673,6 +1673,66 @@ namespace ChaosTest {
 		EXPECT_EQ(Particle.V()[1], YVelAfterChange);
 	}
 
+	GTEST_TEST(EngineInterface, IterationSetOnGT)
+	{
+		//Need to test:
+		//Before set iteration on game thread the iterations are default (8, 2, 1)
+		//After set iteration on game thread the data goes through.
+		FPBDPositionConstraints SinglePositionConstraint;
+		const FReal FixedDT = 1;
+		FChaosScene Scene(nullptr, FixedDT);
+		Scene.GetSolver()->SetThreadingMode_External(EThreadingModeTemp::SingleThread);
+
+		FActorCreationParams Params;
+		Params.Scene = &Scene;
+
+		FPhysicsActorHandle Proxy = nullptr;
+
+		FChaosEngineInterface::CreateActor(Params, Proxy);
+		FRigidBodyHandle_External& Particle = Proxy->GetGameThreadAPI();
+		{
+			auto Sphere = MakeImplicitObjectPtr<TSphere<FReal, 3>>(FVec3(0), 3);
+			Particle.SetGeometry(Sphere);
+		}
+
+		TArray<FPhysicsActorHandle> Proxys = { Proxy };
+		Scene.AddActorsToScene_AssumesLocked(Proxys);
+		Particle.SetObjectState(EObjectStateType::Dynamic);
+		Particle.SetGravityEnabled(true);
+
+		FReal Time = 0;
+		const FReal GTDt = FixedDT * 4;
+		const int32 BigIteration = 100;
+
+		FVec3 Grav(0, 0, -1);
+		Scene.SetUpForFrame(&Grav, GTDt, 0, 99999, 99999, 10, false);
+		Scene.StartFrame();
+		Scene.EndFrame();
+		Chaos::Private::FIterationSettings GroupIterations = Scene.GetSolver()->GetEvolution()->GetIslandGroupManager().GetIslandGroupIterations(0);
+
+		EXPECT_EQ(GroupIterations.GetNumPositionIterations(), 8);
+		EXPECT_EQ(GroupIterations.GetNumVelocityIterations(), 2);
+		EXPECT_EQ(GroupIterations.GetNumProjectionIterations(), 1);
+		
+		auto FirstHandle = Scene.GetSolver()->GetEvolution()->GetParticleHandles().Handle(0)->CastToRigidParticle();
+		TArray<FPBDRigidParticleHandle*> Dynamics = { FirstHandle };
+		TArray<FVec3> Positions = { FVec3(1) };
+		FPBDPositionConstraints PositionConstraints(MoveTemp(Positions), MoveTemp(Dynamics), 1.f);
+		SinglePositionConstraint = PositionConstraints;
+		Scene.GetSolver()->GetEvolution()->AddConstraintContainer(SinglePositionConstraint);
+
+		Particle.SetPositionSolverIterations(BigIteration);
+		Particle.SetVelocitySolverIterations(BigIteration);
+		Particle.SetProjectionSolverIterations(BigIteration);
+
+		Scene.StartFrame();
+		Scene.EndFrame();
+		Chaos::Private::FIterationSettings GroupIterationsTemp = Scene.GetSolver()->GetEvolution()->GetIslandGroupManager().GetIslandGroupIterations(0);
+		EXPECT_EQ(GroupIterationsTemp.GetNumPositionIterations(), 100);
+		EXPECT_EQ(GroupIterationsTemp.GetNumVelocityIterations(), 100);
+		EXPECT_EQ(GroupIterationsTemp.GetNumProjectionIterations(), 100);
+	}
+
 	GTEST_TEST(EngineInterface, FlushCommand)
 	{
 		//Need to test:
diff --git a/Engine/Source/Runtime/Experimental/Chaos/Public/Chaos/Island/IslandGroupManager.h b/Engine/Source/Runtime/Experimental/Chaos/Public/Chaos/Island/IslandGroupManager.h
--- a/Engine/Source/Runtime/Experimental/Chaos/Public/Chaos/Island/IslandGroupManager.h
+++ b/Engine/Source/Runtime/Experimental/Chaos/Public/Chaos/Island/IslandGroupManager.h
@@ -121,6 +121,8 @@ namespace Chaos
 			void SetNumVelocityIterations(const int32 InNumIterations) { Iterations.SetNumVelocityIterations(InNumIterations); }
 			void SetNumProjectionIterations(const int32 InNumIterations) { Iterations.SetNumProjectionIterations(InNumIterations); }
 
+			FIterationSettings GetIslandGroupIterations(const int32 GroupIndex) const { return IslandGroups[GroupIndex]->GetIterationSettings(); }
+
 			/**
 			 * Get the default iteration settings. These are the minimum number of iterations that will be run. An island may run more
 			 * depending on the settings on the dynamic bodies in that island.
