Updated sweeps against height fields to return the result with the minimum TOI, even if there is an initial overlap. This matches the behavior of Overlaps.
Previously, the sweep would terminate on the first test with an initial overlap, even though it may not have been the minimum TOI.

diff --git a/Engine/Source/Programs/HeadlessChaos/Private/HeadlessChaosTestHeightfield.cpp b/Engine/Source/Programs/HeadlessChaos/Private/HeadlessChaosTestHeightfield.cpp
--- a/Engine/Source/Programs/HeadlessChaos/Private/HeadlessChaosTestHeightfield.cpp
+++ b/Engine/Source/Programs/HeadlessChaos/Private/HeadlessChaosTestHeightfield.cpp
@@ -1288,6 +1288,54 @@ namespace ChaosTest {
 		SUCCEED();
 	}
 
+	struct SimpleSweepFixture : public testing::Test
+	{
+		void ConstructFlatHeightFieldAndTestSweep(const FVec3& SweepStart, const FVec3& SweepDirection,
+			const bool bExpectedResult, const FReal ExpectedTOI, const FVec3& ExpectedPosition, const FVec3& ExpectedNormal, const int32 ExpectedFaceId)
+		{
+			TArray<FReal> Heights;
+			Heights.AddZeroed(Rows * Columns);
+			FHeightField Heightfield(MoveTemp(Heights), TArray<uint8>(), Rows, Columns, Scale);
+			const FAABB3 Bounds = Heightfield.BoundingBox();	//Current API forces us to do this to cache the bounds
+
+			const TSphere<FReal, 3> Sphere(FVec3::ZeroVector, SphereRadius);
+			const FRigidTransform3 StartTM(SweepStart, TRotation<FReal, 3>::Identity);
+			
+			FReal ResultTOI;
+			FVec3 ResultPosition, ResultNormal, ResultFaceNormal;
+			int32 ResultFaceId;
+			bool bResult = Heightfield.SweepGeom(Sphere, StartTM, SweepDirection, SweepLength, ResultTOI, ResultPosition, ResultNormal, ResultFaceId, ResultFaceNormal, 0.0, true);
+
+			EXPECT_EQ(bResult, bExpectedResult);
+			EXPECT_VECTOR_NEAR(ResultPosition, ExpectedPosition, 0.01);
+			EXPECT_VECTOR_NEAR(ResultNormal, ExpectedNormal, 0.01);
+			EXPECT_EQ(ResultFaceId, ExpectedFaceId);
+			EXPECT_NEAR(ResultTOI, ExpectedTOI, 0.01);
+		}
+
+		int32 Rows = 3;
+		int32 Columns = 3;
+		FVec3 Scale = FVec3(100, 100, 100);
+		FReal SphereRadius = 25.0;
+		FReal SweepLength = 100;
+	};
+
+	// These tests check that when a sweep has multiple initial overlaps, the "best" result is returned instead of the first.
+	TEST_F(SimpleSweepFixture, FlatHeightField_TestInitialOverlapSweeps)
+	{
+		// Test border of two triangles within the same cell
+		ConstructFlatHeightFieldAndTestSweep(FVec3(60, 50, 10), -FVec3::ZAxisVector, true, -15, FVec3(60, 50, 0), FVec3(0, 0, 1), 0);
+		ConstructFlatHeightFieldAndTestSweep(FVec3(50, 60, 10), -FVec3::ZAxisVector, true, -15, FVec3(50, 60, 0), FVec3(0, 0, 1), 1);
+
+		// Test borders between the 4 cells
+		ConstructFlatHeightFieldAndTestSweep(FVec3(90, 90, 10), -FVec3::ZAxisVector, true, -15, FVec3(90, 90, 0), FVec3(0, 0, 1), 0);
+		ConstructFlatHeightFieldAndTestSweep(FVec3(110, 90, 10), -FVec3::ZAxisVector, true, -15, FVec3(110, 90, 0), FVec3(0, 0, 1), 3);
+		ConstructFlatHeightFieldAndTestSweep(FVec3(90, 110, 10), -FVec3::ZAxisVector, true, -15, FVec3(90, 110, 0), FVec3(0, 0, 1), 4);
+		ConstructFlatHeightFieldAndTestSweep(FVec3(110, 110, 10), -FVec3::ZAxisVector, true, -15, FVec3(110, 110, 0), FVec3(0, 0, 1), 6);
+
+		// Is it worth doing test with mixed TOI / non-zero height?
+	}
+
 	// A flat heightfield and a sphere swept downwards using the CCD API.
 	// Check that the IgnoreThreshold is being used to ignore sweeps that have a penetration depth
 	// less that that threshold at T=1.
diff --git a/Engine/Source/Runtime/Experimental/Chaos/Private/Chaos/HeightField.cpp b/Engine/Source/Runtime/Experimental/Chaos/Private/Chaos/HeightField.cpp
--- a/Engine/Source/Runtime/Experimental/Chaos/Private/Chaos/HeightField.cpp
+++ b/Engine/Source/Runtime/Experimental/Chaos/Private/Chaos/HeightField.cpp
@@ -325,11 +325,6 @@ namespace Chaos
 
 			auto TestTriangle = [&](int32 FaceIndex, const VectorRegister4Float& A, const VectorRegister4Float& B, const VectorRegister4Float& C) -> bool
 			{
-				if(OutTime == 0)
-				{
-					return false;
-				}
-
 				//Convert into local space of A to get better precision
 				FTriangleRegister Triangle(VectorZero(), VectorSubtract(B, A), VectorSubtract(C, A));
 
@@ -373,6 +368,8 @@ namespace Chaos
 							OutTime = Time;
 							OutFaceIndex = FaceIndex;
 
+							// If there's an initial overlap, we can't terminate right away as the first hit might not be the "best" (min distance). 
+							// We have to process all initial overlaps no matter what, however we can still update CurrentLength to not process non-initial overlaps.
 							if(Time <= 0) //initial overlap or MTD, so stop
 							{
 								// This is incorrect. To prevent objects pushing through the surface of the heightfield
@@ -390,8 +387,14 @@ namespace Chaos
 								OutNormal.Y = OutFloat[1];
 								OutNormal.Z = OutFloat[2];
 
-								CurrentLength = 0;
-								return false;
+								// In an ideal world, we'd update CurrentLength to 0 so as to prevent getting any non-zero TOI intersections,
+								// however GJKRaycast does not allow a ray length of 0. We could switch to GJKPenetration, 
+								// however there is currently no SIMD implementation. To work around this, we can set CurrentLength to a small number instead.
+								// Doing this will always produce the correct result (in the edge case we get a TOI between 0 and espilon, we'll still pick the smallest TOI).
+								// This might be slightly slower in this edge case as GJK might've been able to terminate earlier.
+								// #TODO Replace this pathway with a call to GJKPenetration when a SIMD version is created.
+								CurrentLength = UE_KINDA_SMALL_NUMBER;
+								return true;
 							}
 
 							CurrentLength = Time;
@@ -415,8 +418,8 @@ namespace Chaos
 					TestTriangle(Payload * 2 + 1, Points[0], Points[3], Points[2]);
 				}
 			}
-
-			return OutTime > 0;
+			// There is no way to early terminate, even if there is an initial overlap since another cell might have a "better" initial overlap.
+			return true;
 		}
 
 		FReal OutTime;
