Updated sweeps against triangle meshes to return the result with the minimum TOI, even if there is an initial overlap. This matches the behavior of Overlaps.
Previously, the sweep would terminate on the first test with an initial overlap, even though it may not have been the minimum TOI.

diff --git a/Engine/Source/Programs/HeadlessChaos/Private/HeadlessChaosTestSweep.cpp b/Engine/Source/Programs/HeadlessChaos/Private/HeadlessChaosTestSweep.cpp
--- a/Engine/Source/Programs/HeadlessChaos/Private/HeadlessChaosTestSweep.cpp
+++ b/Engine/Source/Programs/HeadlessChaos/Private/HeadlessChaosTestSweep.cpp
@@ -883,4 +883,132 @@ namespace ChaosTest
 		EXPECT_NEAR(WorldPosition.Y, 728.80212, KINDA_SMALL_NUMBER);
 		EXPECT_NEAR(WorldPosition.Z, 303.77856, KINDA_SMALL_NUMBER);
 	}
-}
\ No newline at end of file
+
+	struct FTriangleMeshSweepFixture : public testing::Test
+	{
+		struct FMeshInitData
+		{
+			FTriangleMeshImplicitObject::ParticlesType TrimeshParticles;
+			TArray<TVec3<int32>> Indices;
+			TArray<uint16> Materials;
+
+			void AutoBuildMaterials()
+			{
+				Materials.Empty();
+				for (int32 i = 0; i < Indices.Num(); ++i)
+				{
+					Materials.Emplace(0);
+				}
+			}
+		};
+
+		struct FSweepResultData
+		{
+			bool bResult;
+			FReal TOI;
+			FVec3 Position, Normal, FaceNormal;
+			int32 FaceId;
+		};
+
+		static FMeshInitData CreateSimpleTriInitData()
+		{
+			FMeshInitData Result;
+			Result.TrimeshParticles = FTriangleMeshImplicitObject::ParticlesType(
+			{
+				{0, 0, 0},
+				{100, 0, 0},
+				{100, 100, 0},
+			});
+			Result.Indices =
+			{
+				{0, 1, 2},
+			};
+			Result.AutoBuildMaterials();
+			return Result;
+		}
+
+		static FMeshInitData CreateSimpleQuadInitData()
+		{
+			FMeshInitData Result;
+			Result.TrimeshParticles = FTriangleMeshImplicitObject::ParticlesType(
+			{
+			   {0, 0, 0},
+			   {100, 0, 0},
+			   {100, 100, 0},
+			   {0, 100, 0},
+			});
+			Result.Indices =
+			{
+				{0, 1, 2},
+				{0, 2, 3},
+			};
+			Result.AutoBuildMaterials();
+			return Result;
+		}
+
+		FSweepResultData SweepSphere(const FTriangleMeshImplicitObject& TriangleMesh, const FVec3& SweepStart, const FVec3& SweepDirection)
+		{
+			const FReal SweepLength = 100;
+			const FReal SphereRadius = 25;
+			const TSphere<FReal, 3> Sphere(FVec3::ZeroVector, SphereRadius);
+			const FRigidTransform3 StartTM(SweepStart, TRotation<FReal, 3>::Identity);
+			const FVec3 NormalizedSweepDirection = SweepDirection.GetSafeNormal();
+
+			FSweepResultData Result;
+			Result.bResult = TriangleMesh.SweepGeom(Sphere, StartTM, NormalizedSweepDirection, SweepLength, Result.TOI, Result.Position, Result.Normal, Result.FaceId, Result.FaceNormal, 0.0, true);
+			return Result;
+		}
+
+		FSweepResultData SweepSphere(FMeshInitData& InitData, const FVec3& SweepStart, const FVec3& SweepDirection)
+		{
+			FTriangleMeshImplicitObject TriangleMesh(MoveTemp(InitData.TrimeshParticles), MoveTemp(InitData.Indices), MoveTemp(InitData.Materials), nullptr, nullptr, true);
+			return SweepSphere(TriangleMesh, SweepStart, SweepDirection);
+		}
+	};
+
+	TEST_F(FTriangleMeshSweepFixture, SphereWithInitialIntersectionOfTwoTriangles_SweepAgainstTriMesh_CorrectTriangleIsHit)
+	{
+		// Test a straight down sweep where the sphere has an initial intersection with both triangles.
+		FMeshInitData InitData = CreateSimpleQuadInitData();
+		FTriangleMeshImplicitObject TriangleMesh(MoveTemp(InitData.TrimeshParticles), MoveTemp(InitData.Indices), MoveTemp(InitData.Materials), nullptr, nullptr, true);
+
+		// Test where tri0 has the first TOI.
+		FSweepResultData Result = SweepSphere(TriangleMesh, FVec3(60, 50, 10), FVec3(0, 0, -1));
+		EXPECT_EQ(Result.bResult, true);
+		EXPECT_NEAR(Result.TOI, -15, KINDA_SMALL_NUMBER);
+		EXPECT_VECTOR_NEAR(Result.Position, FVec3(60, 50, 0), KINDA_SMALL_NUMBER);
+		EXPECT_VECTOR_NEAR(Result.Normal, FVec3(0, 0, 1), KINDA_SMALL_NUMBER);
+		EXPECT_EQ(Result.FaceId, 0);
+
+		// Test where tri1 has the first TOI.
+		Result = SweepSphere(TriangleMesh, FVec3(40, 50, 10), FVec3(0, 0, -1));
+		EXPECT_EQ(Result.bResult, true);
+		EXPECT_NEAR(Result.TOI, -15, KINDA_SMALL_NUMBER);
+		EXPECT_VECTOR_NEAR(Result.Position, FVec3(40, 50, 0), KINDA_SMALL_NUMBER);
+		EXPECT_VECTOR_NEAR(Result.Normal, FVec3(0, 0, 1), KINDA_SMALL_NUMBER);
+		EXPECT_EQ(Result.FaceId, 1);
+	}
+
+	TEST_F(FTriangleMeshSweepFixture, SphereWithInitialIntersectionOfOneTriangle_SweepAgainstTriMesh_CorrectTriangleIsHit)
+	{
+		// Test a diagonal sweep where one triangle has an initial intersection and the other has a positive TOI.
+		FMeshInitData InitData = CreateSimpleQuadInitData();
+		FTriangleMeshImplicitObject TriangleMesh(MoveTemp(InitData.TrimeshParticles), MoveTemp(InitData.Indices), MoveTemp(InitData.Materials), nullptr, nullptr, true);
+
+		// Test where tri0 has the initial intersection.
+		FSweepResultData Result = SweepSphere(TriangleMesh, FVec3(90, 50, 10), FVec3(0, 0, -1));
+		EXPECT_EQ(Result.bResult, true);
+		EXPECT_NEAR(Result.TOI, -15, KINDA_SMALL_NUMBER);
+		EXPECT_VECTOR_NEAR(Result.Position, FVec3(90, 50, 0), KINDA_SMALL_NUMBER);
+		EXPECT_VECTOR_NEAR(Result.Normal, FVec3(0, 0, 1), KINDA_SMALL_NUMBER);
+		EXPECT_EQ(Result.FaceId, 0);
+
+		// Test where tri1 has the initial intersection.
+		Result = SweepSphere(TriangleMesh, FVec3(10, 50, 10), FVec3(0, 0, -1));
+		EXPECT_EQ(Result.bResult, true);
+		EXPECT_NEAR(Result.TOI, -15, KINDA_SMALL_NUMBER);
+		EXPECT_VECTOR_NEAR(Result.Position, FVec3(10, 50, 0), KINDA_SMALL_NUMBER);
+		EXPECT_VECTOR_NEAR(Result.Normal, FVec3(0, 0, 1), KINDA_SMALL_NUMBER);
+		EXPECT_EQ(Result.FaceId, 1);
+	}
+}
diff --git a/Engine/Source/Runtime/Experimental/Chaos/Private/Chaos/TriangleMeshImplicitObject.cpp b/Engine/Source/Runtime/Experimental/Chaos/Private/Chaos/TriangleMeshImplicitObject.cpp
--- a/Engine/Source/Runtime/Experimental/Chaos/Private/Chaos/TriangleMeshImplicitObject.cpp
+++ b/Engine/Source/Runtime/Experimental/Chaos/Private/Chaos/TriangleMeshImplicitObject.cpp
@@ -1351,12 +1351,17 @@ struct FTriangleMeshSweepVisitor
 				OutFaceIndex = TriIdx;
 				VectorStoreFloat3(TriNormal, &OutFaceNormal);
 
-				if(Time <= 0)	//MTD or initial overlap
+				// We have an initial overlap (MTD of 0 counts). We cannot terminate early as there might be a "better" hit (one with a more negative Time).
+				if(Time <= 0)
 				{
-					CurDataLength = 0.0f;
-
-					//initial overlap, no one will beat this
-					return false;
+					// In an ideal world, we'd update CurrentLength to 0 so as to prevent getting any non-zero TOI intersections,
+					// however GJKRaycast does not allow a ray length of 0. We could switch to GJKPenetration, 
+					// however there is currently no SIMD implementation. To work around this, we can set CurrentLength to a small number instead.
+					// Doing this will always produce the correct result (in the edge case we get a TOI between 0 and espilon, we'll still pick the smallest TOI).
+					// This might be slightly slower in this edge case as GJK might've been able to terminate earlier.
+					// #TODO Replace this pathway with a call to GJKPenetration when a SIMD version is created.
+					CurDataLength = UE_KINDA_SMALL_NUMBER;
+					return true;
 				}
 
 				CurDataLength = static_cast<FRealSingle>(OutTime);
