Fix GJK issue.
Made the GJK more robust.
Remove the instability issue which was introduced by an optimization originally. But with the vectorized version it s only three subtracts. So no significant cost would be added. This solve parallel case where a cross product should strictly equals to zero.

diff --git a/Engine/Source/Programs/HeadlessChaos/Private/HeadlessChaosTestGJK.cpp b/Engine/Source/Programs/HeadlessChaos/Private/HeadlessChaosTestGJK.cpp
--- a/Engine/Source/Programs/HeadlessChaos/Private/HeadlessChaosTestGJK.cpp
+++ b/Engine/Source/Programs/HeadlessChaos/Private/HeadlessChaosTestGJK.cpp
@@ -2330,4 +2330,82 @@ namespace ChaosTest
 		EXPECT_NEAR(Position.X, 10500.365723, 1e-1);
 		EXPECT_NEAR(Position.Y, -14500.4132690, 1e-1);
 	}
+
+	GTEST_TEST(GJKTests, GJK_SweepSphereBox)
+	{
+		Chaos::TSphere<Chaos::FReal, 3> QueryGeom({ 0, 0, 0 }, 12.0f);
+
+		Chaos::TBox<Chaos::FReal, 3> Box(FVec3(-240.742279/2.0, -16.483643/2.0, -17.078735/2.0), FVec3(240.742279 / 2.0, 16.483643 / 2.0, 17.078735 / 2.0));
+
+		FVec3 BoxWorldTranslation(791.002986, -13679.957812, 7977.850960);
+		FQuat BoxWorldRotation = FQuat(0.006953, -0.013390, 0.887363, -0.460825);
+		TRigidTransform<FReal, 3> BoxWorldTransform(BoxWorldTranslation, BoxWorldRotation);
+		
+		FVec3 Start(715.18582, -13666.121838, 7860.814454);
+		FVec3 End(727.595994, -13746.006452, 8140.664592);
+				
+		TRigidTransform<FReal, 3> TransformSphere(Start, FQuat::Identity);
+
+		TRigidTransform<FReal, 3> FromBoxToSphere1 = TRigidTransform<FReal, 3>::MultiplyNoScale(BoxWorldTransform, TransformSphere.Inverse());
+		TRigidTransform<FReal, 3> FromBoxToSphere2 = BoxWorldTransform.GetRelativeTransform(TransformSphere);
+
+
+		FVec3 Dir(0.042604, -0.274241, 0.960717);
+		const FVec3 LocalDir = TransformSphere.InverseTransformVectorNoScale(Dir);
+		FVec3 Ray = End - Start;
+		FReal Length = Ray.Length();
+
+		FReal OutTime = -1;
+		FVec3 Normal(0.0f);
+		FVec3 Position(0.0f);
+		int32 FaceIndex = -1;
+		FVec3 FaceNormal(0.0);
+		bool bResult = GJKRaycast2(QueryGeom, Box, FromBoxToSphere2, Dir, Length, OutTime, Position, Normal, 0.0, true);
+
+		Position = TransformSphere.TransformPositionNoScale(Position);
+		EXPECT_FALSE(bResult);
+	}
+
+	GTEST_TEST(GJKTests, GJK_SweepSphereTriangleMesh)
+	{
+		{
+			FImplicitSphere3 Sphere{ FVec3(0.0f, 0.0f, 0.0f), 8.16201973f };
+			VectorRegister4Float A = MakeVectorRegister(755.133118f, 69.0732956f, 184.725311f, 0.0f);
+			VectorRegister4Float B = MakeVectorRegister(713.519043f, 205.570648f, 172.950577f, 0.0f);
+			VectorRegister4Float C = MakeVectorRegister(686.582520f, 252.185165f, 365.669464f, 0.0f);
+			FTriangleRegister Tri{ A,B,C };
+
+			VectorRegister4Float StartRotation = MakeVectorRegister(0.0f, 0.0f, 0.0f, 1.0f);
+			VectorRegister4Float StartPosition = MakeVectorRegister(727.730835f, 221.037231f, 363.199402f, 0.0f);
+			VectorRegister4Float RayDirection = MakeVectorRegister(0.0423882306f, 0.0380843617f, -0.998375118f, 0.0f);
+
+			FRealSingle RayLength = 200.0;
+			FRealSingle OutTime;
+			VectorRegister4Float OutPosition;
+			VectorRegister4Float OutNormal;
+
+			const bool bDidHit = GJKRaycast2ImplSimd<VectorRegister4Float>(Tri, Sphere, StartRotation, StartPosition, RayDirection, RayLength, OutTime, OutPosition, OutNormal, true, GlobalVectorConstants::Float1000);
+			EXPECT_FALSE(bDidHit);
+		}
+		{
+			FImplicitSphere3 Sphere{ FVec3(0.0f, 0.0f, 0.0f), 8.16201973f };
+			VectorRegister4Float A = MakeVectorRegister(755.133118f, 69.0732956f, 184.725311f, 0.0f);
+			VectorRegister4Float B = MakeVectorRegister(713.519043f, 205.570648f, 172.950577f, 0.0f);
+			VectorRegister4Float C = MakeVectorRegister(686.582520f, 252.185165f, 365.669464f, 0.0f);
+			FTriangleRegister Tri{ A,B,C };
+
+			VectorRegister4Double StartRotation = MakeVectorRegister(0.0f, 0.0f, 0.0f, 1.0f);
+			VectorRegister4Double StartPosition = MakeVectorRegister(727.730835f, 221.037231f, 363.199402f, 0.0f);
+			VectorRegister4Double RayDirection = MakeVectorRegister(0.0423882306f, 0.0380843617f, -0.998375118f, 0.0f);
+
+			FRealSingle RayLength = 200.0;
+			FRealSingle OutTime;
+			VectorRegister4Double OutPosition;
+			VectorRegister4Double OutNormal;
+
+			const bool bDidHit = GJKRaycast2ImplSimd<VectorRegister4Double>(Tri, Sphere, StartRotation, StartPosition, RayDirection, RayLength, OutTime, OutPosition, OutNormal, true, GlobalVectorConstants::Float1000);
+			EXPECT_FALSE(bDidHit);
+		}
+	}
+
 }
\ No newline at end of file
diff --git a/Engine/Source/Runtime/Experimental/Chaos/Public/Chaos/GJK.h b/Engine/Source/Runtime/Experimental/Chaos/Public/Chaos/GJK.h
--- a/Engine/Source/Runtime/Experimental/Chaos/Public/Chaos/GJK.h
+++ b/Engine/Source/Runtime/Experimental/Chaos/Public/Chaos/GJK.h
@@ -1645,11 +1645,12 @@ namespace Chaos
 					X = VectorMultiplyAdd(Lambda, RayDir, StartPoint);
 					Normal = V;
 
-					//Update simplex from (OldX - P) to (X - P)
-					T XMinusOldX = VectorSubtract(X, OldX);
-					Simplex[0] = VectorAdd(Simplex[0], XMinusOldX);
-					Simplex[1] = VectorAdd(Simplex[1], XMinusOldX);
-					Simplex[2] = VectorAdd(Simplex[2], XMinusOldX);
+					//Update simplex with the new X
+					// As and Bs should have been correcly updated by enabling the flag CalculatExtraInformation
+					// inside VectorSimplexFindClosestToOrigin
+					Simplex[0] = VectorSubtract(X, VectorSubtract(As[0], Bs[0]));
+					Simplex[1] = VectorSubtract(X, VectorSubtract(As[1], Bs[1]));
+					Simplex[2] = VectorSubtract(X, VectorSubtract(As[2], Bs[2]));
 					Simplex[NumVerts++] = VectorSubtract(X, P);
 
 					GJKPreDist2 = LimitMax; //translated origin so restart gjk search
