Fixed an issue where a sweep with an initial overlap between two different shapes returned the first hit object instead of the best hit object. This also fixes separate bodies with initial overlaps.

diff --git a/Engine/Source/Programs/HeadlessChaos/Private/HeadlessChaosTestPhysicsScene.cpp b/Engine/Source/Programs/HeadlessChaos/Private/HeadlessChaosTestPhysicsScene.cpp
--- a/Engine/Source/Programs/HeadlessChaos/Private/HeadlessChaosTestPhysicsScene.cpp
+++ b/Engine/Source/Programs/HeadlessChaos/Private/HeadlessChaosTestPhysicsScene.cpp
@@ -2017,6 +2017,132 @@ namespace ChaosTest {
 		EXPECT_TRUE(HitBuffer.HasBlockingHit());
 	}
 
+	GTEST_TEST(EngineInterface, BodyWithTwoShapes_SweepWithInitialOverlap_MinimalTOIIsReturned)
+	{
+		FChaosScene Scene(nullptr, /*AsyncDt=*/-1);
+
+		FActorCreationParams Params;
+		Params.Scene = &Scene;
+		Params.bSimulatePhysics = false;
+		Params.bStatic = true;
+		Params.InitialTM = FTransform::Identity;
+		Params.Scene = &Scene;
+
+		FPhysicsActorHandle Actor0 = nullptr;
+
+		FChaosEngineInterface::CreateActor(Params, Actor0);
+		ASSERT_NE(Actor0, nullptr);
+
+		TArray<Chaos::FImplicitObjectPtr> Geoms;
+		Geoms.Emplace(MakeImplicitObjectPtr<TBox<FReal, 3>>(FVec3(-50, -50, -50), FVec3(0, 50, 50)));
+		Geoms.Emplace(MakeImplicitObjectPtr<TBox<FReal, 3>>(FVec3(0, -50, -50), FVec3(50, 50, 50)));
+
+		FRigidBodyHandle_External& Particle = Actor0->GetGameThreadAPI();
+		{
+			Chaos::FImplicitObjectPtr GeomUnion = MakeImplicitObjectPtr<FImplicitObjectUnion>(MoveTemp(Geoms));
+			Particle.SetGeometry(GeomUnion);
+		}
+
+		TArray<FPhysicsActorHandle> Particles{ Actor0 };
+		Scene.AddActorsToScene_AssumesLocked(Particles);
+
+		FChaosSQAccelerator SQ{ *Scene.GetSpacialAcceleration() };
+		FSQHitBuffer<ChaosInterface::FSweepHit> HitBuffer;
+		FBlockAllQueryCallback QueryCallback;
+
+		const FReal SphereRadius = 25;
+		const TSphere<FReal, 3> QuerySphere(FVec3::ZeroVector, SphereRadius);
+
+		// Sweep down over both shapes with an initial overlap, however center over shape 0 more.
+		FTransform QueryTM{ FVec3{-5, 0, 50} };
+		const FVec3 Dir(0,0,-1);
+		const FReal Length = 200;
+		SQ.Sweep(QuerySphere, QueryTM, Dir, Length, HitBuffer, EHitFlags::MTD, FQueryFilterData(), QueryCallback, FQueryDebugParams());
+
+		EXPECT_TRUE(HitBuffer.HasBlockingHit());
+		const FSweepHit* BlockingHit = HitBuffer.GetBlock(); 
+		EXPECT_EQ(BlockingHit->Shape, Particle.ShapesArray()[0].Get());
+		EXPECT_EQ(BlockingHit->Distance, -25);
+		EXPECT_VECTOR_NEAR(BlockingHit->WorldPosition, FVector(-5, 0, 50), KINDA_SMALL_NUMBER);
+		EXPECT_VECTOR_NEAR(BlockingHit->WorldNormal, FVector(0, 0, 1), KINDA_SMALL_NUMBER);
+
+		// Now do a second sweep down, this time over shape 1 more.
+		QueryTM = FTransform{ FVec3{5, 0, 50} };
+		SQ.Sweep(QuerySphere, QueryTM, Dir, Length, HitBuffer, EHitFlags::MTD, FQueryFilterData(), QueryCallback, FQueryDebugParams());
+
+		EXPECT_TRUE(HitBuffer.HasBlockingHit());
+		BlockingHit = HitBuffer.GetBlock();
+		EXPECT_EQ(BlockingHit->Shape, Particle.ShapesArray()[1].Get());
+		EXPECT_EQ(BlockingHit->Distance, -25);
+		EXPECT_VECTOR_NEAR(BlockingHit->WorldPosition, FVector(5, 0, 50), KINDA_SMALL_NUMBER);
+		EXPECT_VECTOR_NEAR(BlockingHit->WorldNormal, FVector(0, 0, 1), KINDA_SMALL_NUMBER);
+	}
+
+	GTEST_TEST(EngineInterface, TwoBodies_SweepWithInitialOverlap_MinimalTOIIsReturned)
+	{
+		FChaosScene Scene(nullptr, /*AsyncDt=*/-1);
+
+		FActorCreationParams Params;
+		Params.Scene = &Scene;
+		Params.bSimulatePhysics = false;
+		Params.bStatic = true;
+		Params.InitialTM = FTransform::Identity;
+		Params.Scene = &Scene;
+
+		auto CreateActor = [](const FActorCreationParams& Params, const FImplicitObjectPtr& Geometry) -> FPhysicsActorHandle
+		{
+			FPhysicsActorHandle ActorHandle = nullptr;
+			FChaosEngineInterface::CreateActor(Params, ActorHandle);
+			FRigidBodyHandle_External& Particle = ActorHandle->GetGameThreadAPI();
+			Particle.SetGeometry(Geometry);
+			return ActorHandle;
+		};
+		
+		FImplicitObjectPtr Box0 = MakeImplicitObjectPtr<TBox<FReal, 3>>(FVec3(-50, -50, -50), FVec3(0, 50, 50));
+		FPhysicsActorHandle Actor0 = CreateActor(Params, Box0);
+		ASSERT_NE(Actor0, nullptr);
+		FSingleParticlePhysicsProxy* Actor0Proxy = Actor0->GetGameThreadAPI().GetProxy();
+		
+		FImplicitObjectPtr Box1 = MakeImplicitObjectPtr<TBox<FReal, 3>>(FVec3(0, -50, -50), FVec3(50, 50, 50));
+		FPhysicsActorHandle Actor1 = CreateActor(Params, Box1);
+		ASSERT_NE(Actor1, nullptr);
+		FSingleParticlePhysicsProxy* Actor1Proxy = Actor1->GetGameThreadAPI().GetProxy();
+
+		TArray<FPhysicsActorHandle> Particles{ Actor0, Actor1 };
+		Scene.AddActorsToScene_AssumesLocked(Particles);
+
+		FChaosSQAccelerator SQ{ *Scene.GetSpacialAcceleration() };
+		FSQHitBuffer<ChaosInterface::FSweepHit> HitBuffer;
+		FBlockAllQueryCallback QueryCallback;
+
+		const FReal SphereRadius = 25;
+		const TSphere<FReal, 3> QuerySphere(FVec3::ZeroVector, SphereRadius);
+
+		// Sweep down over both shapes with an initial overlap, however center over shape 0 more.
+		FTransform QueryTM{ FVec3{-5, 0, 50} };
+		const FVec3 Dir(0, 0, -1);
+		const FReal Length = 200;
+		SQ.Sweep(QuerySphere, QueryTM, Dir, Length, HitBuffer, EHitFlags::MTD, FQueryFilterData(), QueryCallback, FQueryDebugParams());
+
+		EXPECT_TRUE(HitBuffer.HasBlockingHit());
+		const FSweepHit* BlockingHit = HitBuffer.GetBlock();
+		EXPECT_EQ(BlockingHit->Actor->GetProxy(), Actor0Proxy);
+		EXPECT_EQ(BlockingHit->Distance, -25);
+		EXPECT_VECTOR_NEAR(BlockingHit->WorldPosition, FVector(-5, 0, 50), KINDA_SMALL_NUMBER);
+		EXPECT_VECTOR_NEAR(BlockingHit->WorldNormal, FVector(0, 0, 1), KINDA_SMALL_NUMBER);
+
+		// Now do a second sweep down, this time over shape 1 more.
+		QueryTM = FTransform{ FVec3{5, 0, 50} };
+		SQ.Sweep(QuerySphere, QueryTM, Dir, Length, HitBuffer, EHitFlags::MTD, FQueryFilterData(), QueryCallback, FQueryDebugParams());
+
+		EXPECT_TRUE(HitBuffer.HasBlockingHit());
+		BlockingHit = HitBuffer.GetBlock();
+		EXPECT_EQ(BlockingHit->Actor->GetProxy(), Actor1Proxy);
+		EXPECT_EQ(BlockingHit->Distance, -25);
+		EXPECT_VECTOR_NEAR(BlockingHit->WorldPosition, FVector(5, 0, 50), KINDA_SMALL_NUMBER);
+		EXPECT_VECTOR_NEAR(BlockingHit->WorldNormal, FVector(0, 0, 1), KINDA_SMALL_NUMBER);
+	}
+
 	// Disable a moving kinematic particle and switch it to dynamic while disabled. Verify that when
 	// disabled it is not in any active lists, and that when re-enabled it is not duplicated
 	//
diff --git a/Engine/Source/Runtime/PhysicsCore/Public/SQVisitor.h b/Engine/Source/Runtime/PhysicsCore/Public/SQVisitor.h
--- a/Engine/Source/Runtime/PhysicsCore/Public/SQVisitor.h
+++ b/Engine/Source/Runtime/PhysicsCore/Public/SQVisitor.h
@@ -322,6 +322,7 @@ private:
 					if constexpr (SQ == ESQType::Sweep && std::is_same_v<THitType, ChaosInterface::FSweepHit>)
 					{
 						const THitType* CurrentHit = HitBuffer.GetCurrentHit();
+						bool bIsCoLocatedHit = false;
 						if (FaceIdx != INDEX_NONE || !FaceNormal.IsNearlyZero())
 						{
 							if (CurrentHit)
@@ -331,6 +332,7 @@ private:
 
 								if (DistDelta < CoLocationEpsilon && !HitFaceNormal.IsNearlyZero())
 								{
+									bIsCoLocatedHit = true;
 									// We already have a face hit from another triangle mesh - see if this one is better (more opposing the sweep)
 									const FReal OldDot = FVec3::DotProduct(CurData->Dir, HitFaceNormal);
 									const FReal NewDot = FVec3::DotProduct(CurData->Dir, FaceNormal);
@@ -355,6 +357,13 @@ private:
 								HitFaceNormal = FaceNormal;
 							}
 						}
+						// If we're not in a co-located hit, we need to make sure the new hit is actually better than the old.
+						// For normal hits this is redundant due to clipping length.
+						// When there is an initial overlaps though, we cannot clip length < 0, so we must check the returned distance.
+						if (CurrentHit != nullptr && !bIsCoLocatedHit && CurrentHit->Distance < Distance)
+						{
+							bAcceptHit = false;
+						}
 					}
 
 					if (bAcceptHit)
@@ -382,10 +391,13 @@ private:
 
 							if (bBlocker && SQ != ESQType::Overlap)
 							{
-								CurData->SetLength(FMath::Max((FReal)0., Distance));	//Max is needed for MTD which returns negative distance
-								if (CurData->CurrentLength == 0 && (SQ == ESQType::Raycast || HitBuffer.WantsSingleResult()))	//raycasts always fail with distance 0, sweeps only matter if we want multi overlaps
+								// When there is an initial overlap, MTD will return a negative distance. We need to clamp this to 0 as many checks don't work properly with a negative length.
+								CurData->SetLength(FMath::Max((FReal)0., Distance));
+								// If we have a zero length raycast, this is an initial overlap and there's no better result possible.
+								// Note: We do not do this check with sweeps as some initial overlap sweeps are better than others (more negative MTD).
+								if (CurData->CurrentLength == 0 && (SQ == ESQType::Raycast))
 								{
-									bContinue = false; //initial overlap so nothing will be better than this
+									bContinue = false;
 									break;
 								}
 							}
