Chaos test does not implement the proper CCD callback when attempting to convert contacts to probes when using CCD. This change adds this callback.

diff --git a/Engine/Source/Programs/HeadlessChaos/Private/HeadlessChaosTestContactModification.cpp b/Engine/Source/Programs/HeadlessChaos/Private/HeadlessChaosTestContactModification.cpp
--- a/Engine/Source/Programs/HeadlessChaos/Private/HeadlessChaosTestContactModification.cpp
+++ b/Engine/Source/Programs/HeadlessChaos/Private/HeadlessChaosTestContactModification.cpp
@@ -1,6 +1,7 @@
 // Copyright Epic Games, Inc. All Rights Reserved.
 
 #include "Chaos/ContactModification.h"
+#include "Chaos/CCDModification.h"
 #include "ChaosSolversModule.h"
 #include "HeadlessChaosTestUtility.h"
 #include "PBDRigidsSolver.h"
@@ -13,14 +14,16 @@ namespace ChaosTest
 	class FContactModificationTestCallback : public Chaos::TSimCallbackObject<
 		Chaos::FSimCallbackNoInput,
 		Chaos::FSimCallbackNoOutput,
-		Chaos::ESimCallbackOptions::Presimulate | Chaos::ESimCallbackOptions::ContactModification>
+		Chaos::ESimCallbackOptions::Presimulate | Chaos::ESimCallbackOptions::ContactModification | Chaos::ESimCallbackOptions::CCDModification>
 	{
 	public:
 		TUniqueFunction<void(Chaos::FCollisionContactModifier&)> TestLambda;
+		TUniqueFunction<void(Chaos::FCCDModifierAccessor&)> TestCCDLambda = nullptr;
 
 	private:
 		virtual void OnPreSimulate_Internal() override {}
 		virtual void OnContactModification_Internal(Chaos::FCollisionContactModifier& Modifier) override;
+		virtual void OnCCDModification_Internal(Chaos::FCCDModifierAccessor& Accessor) override;
 	};
 
 	void FContactModificationTestCallback::OnContactModification_Internal(Chaos::FCollisionContactModifier& Modifier)
@@ -28,6 +31,14 @@ namespace ChaosTest
 		TestLambda(Modifier);
 	}
 
+	void FContactModificationTestCallback::OnCCDModification_Internal(Chaos::FCCDModifierAccessor& Accessor)
+	{
+		if (TestCCDLambda)
+		{
+			TestCCDLambda(Accessor);
+		}
+	}
+
 
 	GTEST_TEST(AllTraits, ContactModification_Disable)
 	{
@@ -121,8 +132,7 @@ namespace ChaosTest
 		Module->DestroySolver(Solver);
 	}
 
-	// Disabling due to: UE-218719; CCD is being done before Contact Modification which is resulting in the object B colliding
-	GTEST_TEST(AllTraits, DISABLED_ContactModification_Probe)
+	GTEST_TEST(AllTraits, ContactModification_Probe)
 	{
 		FChaosSolversModule* Module = FChaosSolversModule::GetModule();
 		auto* Solver = Module->CreateSolver(nullptr, /*AsyncDt=*/-1);
@@ -177,6 +187,18 @@ namespace ChaosTest
 				PairModifier.ConvertToProbe();
 			}
 		};
+		
+		Callback->TestCCDLambda = [CubeProxyB](Chaos::FCCDModifierAccessor& Accessor)
+		{
+			if (Chaos::FGeometryParticleHandle* ParticleHandle = CubeProxyB->GetHandle_LowLevel())
+			{
+				for (FCCDModifier& CCDModifier : Accessor.GetModifiers(ParticleHandle))
+				{
+					CCDModifier.ConvertToProbe();
+				}
+			}
+		};
+
 
 		const float Dt = 1.0f;
 		const int32 Steps = 10;
