HeadlessChaos and Phsyics compiling with Warning as Error

diff --git a/Engine/Source/Programs/HeadlessChaos/HeadlessChaos.Target.cs b/Engine/Source/Programs/HeadlessChaos/HeadlessChaos.Target.cs
--- a/Engine/Source/Programs/HeadlessChaos/HeadlessChaos.Target.cs
+++ b/Engine/Source/Programs/HeadlessChaos/HeadlessChaos.Target.cs
@@ -17,5 +17,6 @@ public class HeadlessChaosTarget : TestTargetRules
 		// Enable additional debug functionality in the constraint management system to allow better unit testing
 		GlobalDefinitions.Add("CHAOS_CONSTRAINTHANDLE_DEBUG_ENABLED=1");
 		BuildEnvironment = TargetBuildEnvironment.Unique;
+		bWarningsAsErrors = true;
 	}
 }
