Fix Headless chaos CIS

diff --git a/Engine/Source/Programs/HeadlessChaos/Private/HeadlessChaosTestModularVehicles.cpp b/Engine/Source/Programs/HeadlessChaos/Private/HeadlessChaosTestModularVehicles.cpp
--- a/Engine/Source/Programs/HeadlessChaos/Private/HeadlessChaosTestModularVehicles.cpp
+++ b/Engine/Source/Programs/HeadlessChaos/Private/HeadlessChaosTestModularVehicles.cpp
@@ -936,7 +936,9 @@ namespace ChaosTest
 		const FWheelSimModule& Wheel2 = *static_cast<const FWheelSimModule*>(Tree.GetSimModule(Wheel2Idx));
 		const FWheelSimModule& Wheel3 = *static_cast<const FWheelSimModule*>(Tree.GetSimModule(Wheel3Idx));
 
-		FClusterUnionPhysicsProxy* Proxy = nullptr;
+		IPhysicsProxyBase* Proxy = nullptr;
+		FPBDRigidParticleHandle* Particle = nullptr;
+
 		float DeltaTime = 1.0f / 30.0f;
 
 		// We need to setup reverse pointer to tree
@@ -959,7 +961,7 @@ namespace ChaosTest
 		// Simulate
 		for (int I=0; I<50; I++)
 		{
-			Tree.Simulate(DeltaTime, Inputs, Proxy);
+			Tree.Simulate(DeltaTime, Inputs, Proxy, Particle);
 		}
 
 		// driven wheels are turning
@@ -978,7 +980,7 @@ namespace ChaosTest
 		// Simulate
 		for (int I = 0; I < 50; I++)
 		{
-			Tree.Simulate(DeltaTime, Inputs, Proxy);
+			Tree.Simulate(DeltaTime, Inputs, Proxy, Particle);
 		}
 
 		// all wheels stop spinning
@@ -996,7 +998,7 @@ namespace ChaosTest
 		// Simulate
 		for (int I = 0; I < 50; I++)
 		{
-			Tree.Simulate(DeltaTime, Inputs, Proxy);
+			Tree.Simulate(DeltaTime, Inputs, Proxy, Particle);
 		}
 
 		// driven wheels are not turning due to disengaged clutch
@@ -1053,7 +1055,8 @@ namespace ChaosTest
 		FWheelSimModule& Wheel0 = *static_cast<FWheelSimModule*>(Tree.AccessSimModule(Wheel0Idx));
 		FWheelSimModule& Wheel1 = *static_cast<FWheelSimModule*>(Tree.AccessSimModule(Wheel1Idx));
 
-		FClusterUnionPhysicsProxy* Proxy = nullptr;
+		IPhysicsProxyBase* Proxy = nullptr;
+		FPBDRigidParticleHandle* Particle = nullptr;
 		float DeltaTime = 1.0f / 30.0f;
 
 		// We need to setup reverse pointer to tree
@@ -1068,7 +1071,7 @@ namespace ChaosTest
 		{
 			Wheel0.SetLinearSpeed(100);
 			Wheel1.SetLinearSpeed(100);
-			Tree.Simulate(DeltaTime, Inputs, Proxy);
+			Tree.Simulate(DeltaTime, Inputs, Proxy, Particle);
 		}
 
 		EXPECT_NEAR(Engine.GetRPM(), Wheel0.GetRPM(), TOLERANCE);
@@ -1079,7 +1082,7 @@ namespace ChaosTest
 		{
 			Wheel0.SetLinearSpeed(300);
 			Wheel1.SetLinearSpeed(100);
-			Tree.Simulate(DeltaTime, Inputs, Proxy);
+			Tree.Simulate(DeltaTime, Inputs, Proxy, Particle);
 		}
 
 		EXPECT_NEAR(Engine.GetRPM(), (Wheel0.GetRPM() + Wheel1.GetRPM())*0.5f, TOLERANCE);
