Triangle meshes were incorrectly reporting sweeps against backfaces even when backface culling was enabled. To preserve the behavior of getting hits when there is an intiial overlap, the backface case was split into front, back, and parallel cases. Backface cases will always cull now whereas parallel cases will check for initial overlap.

diff --git a/Engine/Source/Programs/HeadlessChaos/Private/HeadlessChaosTestSweep.cpp b/Engine/Source/Programs/HeadlessChaos/Private/HeadlessChaosTestSweep.cpp
--- a/Engine/Source/Programs/HeadlessChaos/Private/HeadlessChaosTestSweep.cpp
+++ b/Engine/Source/Programs/HeadlessChaos/Private/HeadlessChaosTestSweep.cpp
@@ -1042,4 +1042,155 @@ namespace ChaosTest
 		Result = SweepSphere(TriangleMesh, FVec3(20, 20, 0), FVec3(0, 0, 0), 1);
 		EXPECT_EQ(Result.bResult, true);
 	}
+
+	// All front-face hits (sweeping opposed to the normal) should generate valid hits.
+	TEST_F(FTriangleMeshSweepFixture, SphereInFront_SweepOpposedToNormal_HitIsExpected)
+	{
+		FMeshInitData InitData = CreateSimpleTriInitData();
+		const FSweepResultData Result = SweepSphere(InitData, FVec3(20, 20, 50), FVec3(0, 0, -1));
+
+		EXPECT_EQ(Result.bResult, true);
+		EXPECT_NEAR(Result.TOI, 25, KINDA_SMALL_NUMBER);
+		EXPECT_VECTOR_NEAR(Result.Position, FVec3(20, 20, 0), KINDA_SMALL_NUMBER);
+		EXPECT_VECTOR_NEAR(Result.Normal, FVec3(0, 0, 1), KINDA_SMALL_NUMBER);
+		EXPECT_EQ(Result.FaceId, 0);
+	}
+
+	TEST_F(FTriangleMeshSweepFixture, SphereInFrontWithInitialOverlap_SweepOpposedToNormal_HitIsExpected)
+	{
+		FMeshInitData InitData = CreateSimpleTriInitData();
+		FSweepResultData Result = SweepSphere(InitData, FVec3(20, 20, 10), FVec3(0, 0, -1));
+
+		EXPECT_EQ(Result.bResult, true);
+		EXPECT_NEAR(Result.TOI, -15, KINDA_SMALL_NUMBER);
+		EXPECT_VECTOR_NEAR(Result.Position, FVec3(20, 20, 0), KINDA_SMALL_NUMBER);
+		EXPECT_VECTOR_NEAR(Result.Normal, FVec3(0, 0, 1), KINDA_SMALL_NUMBER);
+		EXPECT_EQ(Result.FaceId, 0);
+	}
+
+	TEST_F(FTriangleMeshSweepFixture, SphereBehindWithInitialOverlap_SweepOpposedToNormal_HitIsExpected)
+	{
+		FMeshInitData InitData = CreateSimpleTriInitData();
+		FSweepResultData Result = SweepSphere(InitData, FVec3(20, 20, -10), FVec3(0, 0, -1));
+
+		EXPECT_EQ(Result.bResult, true);
+		EXPECT_NEAR(Result.TOI, -15, KINDA_SMALL_NUMBER);
+		EXPECT_VECTOR_NEAR(Result.Position, FVec3(20, 20, 0), KINDA_SMALL_NUMBER);
+		EXPECT_VECTOR_NEAR(Result.Normal, FVec3(0, 0, -1), KINDA_SMALL_NUMBER);
+		EXPECT_EQ(Result.FaceId, 0);
+	}
+
+	struct FTriangleMeshBackfaceSweepFixture : public FTriangleMeshSweepFixture
+	{
+		static constexpr FReal ValueWithinParallelEpsilon = UE_SMALL_NUMBER;
+		static constexpr FReal ValueOutsideParallelEpsilon = 4 * UE_KINDA_SMALL_NUMBER;
+	};
+
+	TEST_F(FTriangleMeshBackfaceSweepFixture, SphereBehind_SweepOpposedToNormal_HitIsNotExpected)
+	{
+		FMeshInitData InitData = CreateSimpleTriInitData();
+		FSweepResultData Result = SweepSphere(InitData, FVec3(20, 20, -50), FVec3(0, 0, -1));
+
+		EXPECT_EQ(Result.bResult, false);
+	}
+
+	// All back-face hits (sweeping along the normal) should not generate a valid hit.
+	TEST_F(FTriangleMeshBackfaceSweepFixture, SphereBehind_SweepAlongNormal_HitIsNotExpected)
+	{
+		FMeshInitData InitData = CreateSimpleTriInitData();
+		FSweepResultData Result = SweepSphere(InitData, FVec3(20, 20, -50), FVec3(0, 0, 1));
+
+		EXPECT_EQ(Result.bResult, false);
+	}
+
+	TEST_F(FTriangleMeshBackfaceSweepFixture, SphereBehindWithInitialOverlap_SweepAlongNormal_HitIsNotExpected)
+	{
+		FMeshInitData InitData = CreateSimpleTriInitData();
+		FSweepResultData Result = SweepSphere(InitData, FVec3(20, 20, -10), FVec3(0, 0, 1));
+
+		EXPECT_EQ(Result.bResult, false);
+	}
+
+	TEST_F(FTriangleMeshBackfaceSweepFixture, SphereInFrontWithInitialOverlap_SweepAlongNormal_HitIsNotExpected)
+	{
+		FMeshInitData InitData = CreateSimpleTriInitData();
+		FSweepResultData Result = SweepSphere(InitData, FVec3(20, 20, 10), FVec3(0, 0, 1));
+
+		EXPECT_EQ(Result.bResult, false);
+	}
+
+	TEST_F(FTriangleMeshBackfaceSweepFixture, SphereInFront_SweepAlongNormal_HitIsNotExpected)
+	{
+		FMeshInitData InitData = CreateSimpleTriInitData();
+		FSweepResultData Result = SweepSphere(InitData, FVec3(20, 20, 50), FVec3(0, 0, 1));
+
+		EXPECT_EQ(Result.bResult, false);
+	}
+
+	// Parallel checks are more interesting. If there is an initial overlap then no hit should be generated. 
+	// This parallel checks uses an epsilon to determine the boundaries.
+	TEST_F(FTriangleMeshBackfaceSweepFixture, SphereWithInitialOverlap_SweepExactlyParallel_HitIsExpected)
+	{
+		FMeshInitData InitData = CreateSimpleTriInitData();
+		FSweepResultData Result = SweepSphere(InitData, FVec3(20, 20, 10), FVec3(1, 0, 0));
+
+		EXPECT_EQ(Result.bResult, true);
+	}
+
+	TEST_F(FTriangleMeshBackfaceSweepFixture, SphereWithInitialOverlap_SweepOpposedToNormalWithinParallelEpsilon_HitIsExpected)
+	{
+		FMeshInitData InitData = CreateSimpleTriInitData();
+		FSweepResultData Result = SweepSphere(InitData, FVec3(20, 20, 10), FVec3(1, 0, ValueWithinParallelEpsilon));
+
+		EXPECT_EQ(Result.bResult, true);
+	}
+
+	TEST_F(FTriangleMeshBackfaceSweepFixture, SphereWithInitialOverlap_SweepOpposedToNormalOutsideParallelEpsilon_HitIsNotExpected)
+	{
+		FMeshInitData InitData = CreateSimpleTriInitData();
+		FSweepResultData Result = SweepSphere(InitData, FVec3(20, 20, 10), FVec3(1, 0, ValueOutsideParallelEpsilon));
+
+		EXPECT_EQ(Result.bResult, false);
+	}
+
+	TEST_F(FTriangleMeshBackfaceSweepFixture, SphereWithInitialOverlap_SweepAlongNormalWithinParallelEpsilon_HitIsExpected)
+	{
+		FMeshInitData InitData = CreateSimpleTriInitData();
+		FSweepResultData Result = SweepSphere(InitData, FVec3(20, 20, 10), FVec3(1, 0, -ValueWithinParallelEpsilon));
+
+		EXPECT_EQ(Result.bResult, true);
+	}
+
+	TEST_F(FTriangleMeshBackfaceSweepFixture, SphereWithInitialOverlap_SweepAlongNormalOutsideParallelEpsilon_HitIsExpected)
+	{
+		FMeshInitData InitData = CreateSimpleTriInitData();
+		FSweepResultData Result = SweepSphere(InitData, FVec3(20, 20, 10), FVec3(1, 0, 4 * -ValueOutsideParallelEpsilon));
+
+		EXPECT_EQ(Result.bResult, true);
+	}
+
+	// Need a few extra tests for parallel to validate that no initial overlap always returns false.
+	TEST_F(FTriangleMeshBackfaceSweepFixture, SphereInAabbWithNoInitialOverlap_SweepExactlyParallel_HitIsNotExpected)
+	{
+		FMeshInitData InitData = CreateSimpleTriInitData();
+		FSweepResultData Result = SweepSphere(InitData, FVec3(0, 100, 0), FVec3(1, -1, 0));
+
+		EXPECT_EQ(Result.bResult, false);
+	}
+
+	TEST_F(FTriangleMeshBackfaceSweepFixture, SphereInAabbWithNoInitialOverlap_AlongNormalWithinPositiveEpsilon_HitIsNotExpected)
+	{
+		FMeshInitData InitData = CreateSimpleTriInitData();
+		FSweepResultData Result = SweepSphere(InitData, FVec3(0, 100, 0), FVec3(1, -1, ValueWithinParallelEpsilon));
+
+		EXPECT_EQ(Result.bResult, false);
+	}
+
+	TEST_F(FTriangleMeshBackfaceSweepFixture, SphereInAabbWithNoInitialOverlap_AlongNormalWithinNegativeEpsilon_HitIsNotExpected)
+	{
+		FMeshInitData InitData = CreateSimpleTriInitData();
+		FSweepResultData Result = SweepSphere(InitData, FVec3(0, 100, 0), FVec3(1, -1, -ValueWithinParallelEpsilon));
+
+		EXPECT_EQ(Result.bResult, false);
+	}
 }
diff --git a/Engine/Source/Runtime/Experimental/Chaos/Private/Chaos/TriangleMeshImplicitObject.cpp b/Engine/Source/Runtime/Experimental/Chaos/Private/Chaos/TriangleMeshImplicitObject.cpp
--- a/Engine/Source/Runtime/Experimental/Chaos/Private/Chaos/TriangleMeshImplicitObject.cpp
+++ b/Engine/Source/Runtime/Experimental/Chaos/Private/Chaos/TriangleMeshImplicitObject.cpp
@@ -1293,10 +1293,34 @@ struct FTriangleMeshSweepVisitor
 		FTriangleRegister Tri(A, B, C);
 		const VectorRegister4Float TriNormal = VectorCross(VectorSubtract(B, A), VectorSubtract(C, A));
 
-		const VectorRegister4Float IsBackFace = VectorCompareGT(VectorMultiply(VectorDot3(TriNormal, VectorScaledDirNormalized), VectorCullsBackFaceSweepsCode), VectorZero());
+		bool bIsParallelSweepAndCullsBackFace = false;
 		if (CullsBackFaceSweepsCode != 0)
 		{
-			if (VectorMaskBits(IsBackFace))
+			// For the parallel epsilon check to be consistent, we must normalize the triangle normal
+			const VectorRegister4Float ReciprocalNormalLen = VectorReciprocalLen(TriNormal);
+			const VectorRegister4Float NormalizedNormal = VectorMultiply(TriNormal, ReciprocalNormalLen);
+			const VectorRegister4Float NDotD = VectorDot3(NormalizedNormal, VectorScaledDirNormalized);
+
+			const VectorRegister4Float PositiveEpsilon = MakeVectorRegister(UE_KINDA_SMALL_NUMBER, UE_KINDA_SMALL_NUMBER, UE_KINDA_SMALL_NUMBER, UE_KINDA_SMALL_NUMBER);
+			const VectorRegister4Float NegativeEpsilon = VectorNegate(PositiveEpsilon);
+			const VectorRegister4Float IsAlongPositiveNormal = VectorCompareGT(NDotD, PositiveEpsilon);
+			const VectorRegister4Float IsAlongNegativeNormal = VectorCompareLT(NDotD, NegativeEpsilon);
+
+			const VectorRegister4Float IsAlongPositiveNormalMask = VectorMultiply(VectorCullsBackFaceSweepsCode, IsAlongPositiveNormal);
+			const VectorRegister4Float IsAlongNegativeNormalMask = VectorMultiply(VectorCullsBackFaceSweepsCode, IsAlongNegativeNormal);
+
+			const bool bIsBackFace = VectorMaskBits(IsAlongPositiveNormalMask) != 0;
+			const bool bIsFrontFace = VectorMaskBits(IsAlongNegativeNormalMask) != 0;
+			bIsParallelSweepAndCullsBackFace = (!bIsBackFace && !bIsFrontFace);
+
+			// If this is a definitive back-face hit, always skip, even if there is an initial overlap.
+			if (bIsBackFace)
+			{
+				return true;
+			}
+			// If this is parallel, we want to skip if there's no initial intersection. 
+			// This is an early-out check to avoid calling GJK if there's no bounding volume intersection.
+			if (bIsParallelSweepAndCullsBackFace)
 			{
 				// Don't cull the back face if there is a chance that we are initially overlapping.
 				const VectorRegister4Float MinBounds = VectorMin(VectorMin(A, B), C);
@@ -1324,8 +1348,9 @@ struct FTriangleMeshSweepVisitor
 		}
 		if(bIsHitting)
 		{
-			// Don't return back faces if they are not initially overlapping
-			if (Time > 0 && VectorMaskBits(IsBackFace))
+			// Parallel sweeps should still be culled (if culls back-face is enabled) if there is no initial intersection. 
+			// We still want to keep a sweep if there is initial intersection.
+			if (Time > 0 && bIsParallelSweepAndCullsBackFace)
 			{
 				return true;
 			}
