Store Implicit in single precision
Now the following
ImplicitObject : margin is in single precision
Sphere was already using margin as radius which is stored in single precision.
Capsule segment is in single precision
Deprecated all double precision API and replace by the single precision API.
TaperedCapsule, TaperedCylinder and Cylinder are also now stored in single precision

diff --git a/Engine/Plugins/ChaosCloth/Source/ChaosCloth/Private/ChaosCloth/ChaosClothVisualization.cpp b/Engine/Plugins/ChaosCloth/Source/ChaosCloth/Private/ChaosCloth/ChaosClothVisualization.cpp
--- a/Engine/Plugins/ChaosCloth/Source/ChaosCloth/Private/ChaosCloth/ChaosClothVisualization.cpp
+++ b/Engine/Plugins/ChaosCloth/Source/ChaosCloth/Private/ChaosCloth/ChaosClothVisualization.cpp
@@ -586,16 +586,16 @@ namespace Chaos
 
 	static void DrawSphere(FPrimitiveDrawInterface* PDI, const ::Chaos::TSphere<::Chaos::FReal, 3>& Sphere, const FQuat& Rotation, const FVector& Position, const FReal Scale, const FLinearColor& Color)
 	{
-		const FVec3 Center = Position + Rotation.RotateVector(Scale * Sphere.GetCenter());
+		const FVec3 Center = Position + Rotation.RotateVector(Scale * FVec3(Sphere.GetCenterf()));
 		if (!PDI)
 		{
-			const FReal Radius = Sphere.GetRadius() * Scale;
+			const FReal Radius = Sphere.GetRadiusf() * Scale;
 			FDebugDrawQueue::GetInstance().DrawDebugSphere(Center, Radius, 12, Color.ToFColor(true), false, KINDA_SMALL_NUMBER, Chaos::Private::GetDepthPriority(), 0.f);
 			return;
 		}
 #if WITH_EDITOR
 		const FTransform Transform(Rotation, Center, FVector(Scale));
-		DrawWireSphere(PDI, Transform, Color, Sphere.GetRadius(), 12, SDPG_World, 0.0f, 0.001f, false);
+		DrawWireSphere(PDI, Transform, Color, Sphere.GetRadiusf(), 12, SDPG_World, 0.0f, 0.001f, false);
 #endif
 	}
 
@@ -615,9 +615,9 @@ namespace Chaos
 
 	static void DrawCapsule(FPrimitiveDrawInterface* PDI, const ::Chaos::FCapsule& Capsule, const FQuat& Rotation, const FVector& Position, FReal Scale, const FLinearColor& Color)
 	{
-		const FReal Radius = Scale * Capsule.GetRadius();
-		const FReal HalfHeight = Scale * Capsule.GetHeight() * 0.5f + Radius;
-		const FVec3 Center = Position + Rotation.RotateVector(Scale * Capsule.GetCenter());
+		const FReal Radius = Scale * Capsule.GetRadiusf();
+		const FReal HalfHeight = Scale * Capsule.GetHeightf() * 0.5f + Radius;
+		const FVec3 Center = Position + Rotation.RotateVector(Scale * FVec3(Capsule.GetCenterf()));
 		if (!PDI)
 		{
 			const FQuat Orientation = FQuat::FindBetweenNormals(FVec3::UpVector, Capsule.GetAxis());
@@ -1020,7 +1020,7 @@ namespace Chaos
 
 		// Draw bounds
 		DrawBox(PDI, FAABB3(-Bounds.BoxExtent, Bounds.BoxExtent), FQuat::Identity, Bounds.Origin, 1., FLinearColor(FColor::Purple));
-		DrawSphere(PDI, TSphere<FReal, 3>(FVector::ZeroVector, Bounds.SphereRadius), FQuat::Identity, Bounds.Origin, 1., FLinearColor(FColor::Orange));
+		DrawSphere(PDI, FSphere(FVector::ZeroVector, Bounds.SphereRadius), FQuat::Identity, Bounds.Origin, 1., FLinearColor(FColor::Orange));
 
 		// Draw individual cloth bounds
 		static const FLinearColor Color = FLinearColor(FColor::Purple).Desaturate(0.5);
@@ -1516,7 +1516,7 @@ namespace Chaos
 						switch (Object->GetType())
 						{
 						case ImplicitObjectType::Sphere:
-							DrawSphere(PDI, Object->GetObjectChecked<TSphere<FReal, 3>>(), Rotation, Position, Solver->GetLocalSpaceScale(), Color);
+							DrawSphere(PDI, Object->GetObjectChecked<FSphere>(), Rotation, Position, Solver->GetLocalSpaceScale(), Color);
 							break;
 
 						case ImplicitObjectType::Box:
@@ -1535,7 +1535,7 @@ namespace Chaos
 									switch (SubObject->GetType())
 									{
 									case ImplicitObjectType::Sphere:
-										DrawSphere(PDI, SubObject->GetObjectChecked<TSphere<FReal, 3>>(), Rotation, Position, Solver->GetLocalSpaceScale(), Color);
+										DrawSphere(PDI, SubObject->GetObjectChecked<FSphere>(), Rotation, Position, Solver->GetLocalSpaceScale(), Color);
 										break;
 
 									case ImplicitObjectType::TaperedCylinder:
@@ -1556,8 +1556,8 @@ namespace Chaos
 								const FVec3 X2 = TaperedCapsule.GetX2();
 								const FReal Radius1 = TaperedCapsule.GetRadius1();
 								const FReal Radius2 = TaperedCapsule.GetRadius2();
-								DrawSphere(PDI, TSphere<FReal, 3>(X1, Radius1), Rotation, Position, Solver->GetLocalSpaceScale(), Color);
-								DrawSphere(PDI, TSphere<FReal, 3>(X2, Radius2), Rotation, Position, Solver->GetLocalSpaceScale(), Color);
+								DrawSphere(PDI, FSphere(X1, Radius1), Rotation, Position, Solver->GetLocalSpaceScale(), Color);
+								DrawSphere(PDI, FSphere(X2, Radius2), Rotation, Position, Solver->GetLocalSpaceScale(), Color);
 								DrawTaperedCylinder(PDI, FTaperedCylinder(X1, X2, Radius1, Radius2), Rotation, Position, Solver->GetLocalSpaceScale(), Color);
 							}
 							break;
@@ -2639,7 +2639,7 @@ namespace Chaos
 					for (const int32 Vertex : *VertexSet)
 					{
 						const FVector P0 = GetWorldPosition(*Solver, Positions[Vertex + Offset]);
-						DrawSphere(PDI, TSphere<FReal, 3>(FVector::ZeroVector, Radius), FQuat::Identity,
+						DrawSphere(PDI, FSphere(FVector::ZeroVector, Radius), FQuat::Identity,
 							P0, 1., FLinearColor(FColor::Orange));
 					}
 				}
diff --git a/Engine/Plugins/ChaosVD/Source/ChaosVD/Private/ChaosVDGeometryBuilder.cpp b/Engine/Plugins/ChaosVD/Source/ChaosVD/Private/ChaosVDGeometryBuilder.cpp
--- a/Engine/Plugins/ChaosVD/Source/ChaosVD/Private/ChaosVDGeometryBuilder.cpp
+++ b/Engine/Plugins/ChaosVD/Source/ChaosVD/Private/ChaosVDGeometryBuilder.cpp
@@ -678,10 +678,10 @@ TSharedPtr<UE::Geometry::FMeshShapeGenerator> FChaosVDGeometryBuilder::CreateMes
 	{
 		case ImplicitObjectType::Sphere:
 		{
-			if (const Chaos::TSphere<FReal, 3>* Sphere = InImplicit->template GetObject<Chaos::TSphere<FReal, 3>>())
+			if (const Chaos::FSphere* Sphere = InImplicit->template GetObject<Chaos::FSphere>())
 			{
 				TSharedPtr<UE::Geometry::FSphereGenerator> SphereGen = MakeShared<UE::Geometry::FSphereGenerator>();
-				SphereGen->Radius = Sphere->GetRadius();
+				SphereGen->Radius = Sphere->GetRadiusf();
 				SphereGen->NumTheta = 25 * SimpleShapesComplexityFactor;
 				SphereGen->NumPhi = 25 * SimpleShapesComplexityFactor;
 				SphereGen->bPolygroupPerQuad = false;
@@ -708,7 +708,7 @@ TSharedPtr<UE::Geometry::FMeshShapeGenerator> FChaosVDGeometryBuilder::CreateMes
 			if (const Chaos::FCapsule* Capsule = InImplicit->template GetObject<Chaos::FCapsule>())
 			{
 				TSharedPtr<UE::Geometry::FCapsuleGenerator> CapsuleGenerator = MakeShared<UE::Geometry::FCapsuleGenerator>();
-				CapsuleGenerator->Radius = FMath::Max(FMathf::ZeroTolerance, Capsule->GetRadius());
+				CapsuleGenerator->Radius = FMath::Max(FMathf::ZeroTolerance, Capsule->GetRadiusf());
 				CapsuleGenerator->SegmentLength = FMath::Max(FMathf::ZeroTolerance, Capsule->GetSegment().GetLength());
 				CapsuleGenerator->NumHemisphereArcSteps = 12 * SimpleShapesComplexityFactor;
 				CapsuleGenerator->NumCircleSteps = 12 * SimpleShapesComplexityFactor;
@@ -806,8 +806,8 @@ void FChaosVDGeometryBuilder::AdjustedTransformForImplicit(const Chaos::FImplici
 			if (const FCapsule* Capsule = InImplicit->template GetObject<FCapsule>())
 			{
 				// Re-adjust the location so the pivot is not the center of the capsule, and transform it based on the provided transform
-				const FVector FinalLocation = OutAdjustedTransform.TransformPosition(Capsule->GetCenter() - Capsule->GetAxis() * Capsule->GetSegment().GetLength() * 0.5f);
-				const FQuat Rotation = FRotationMatrix::MakeFromZ(Capsule->GetAxis()).Rotator().Quaternion();
+				const FVector FinalLocation = OutAdjustedTransform.TransformPosition(FVector(Capsule->GetCenterf()) - FVector(Capsule->GetAxisf()) * Capsule->GetSegment().GetLength() * 0.5f);
+				const FQuat Rotation = FRotationMatrix::MakeFromZ(FVector(Capsule->GetAxisf())).Rotator().Quaternion();
 
 				OutAdjustedTransform.SetRotation(OutAdjustedTransform.GetRotation() * Rotation);
 				OutAdjustedTransform.SetLocation(FinalLocation);
@@ -816,9 +816,9 @@ void FChaosVDGeometryBuilder::AdjustedTransformForImplicit(const Chaos::FImplici
 		}
 		case ImplicitObjectType::Sphere:
 			{
-				if (const Chaos::TSphere<FReal, 3>* Sphere = InImplicit->template GetObject<Chaos::TSphere<FReal, 3>>())
+				if (const Chaos::FSphere* Sphere = InImplicit->template GetObject<Chaos::FSphere>())
 				{
-					const FVector FinalLocation = OutAdjustedTransform.TransformPosition(Sphere->GetCenter());
+					const FVector FinalLocation = OutAdjustedTransform.TransformPosition(FVector(Sphere->GetCenterf()));
 					OutAdjustedTransform.SetLocation(FinalLocation);
 				}
 				break;
diff --git a/Engine/Plugins/Experimental/Buoyancy/Source/Runtime/Private/BuoyancyAlgorithms.cpp b/Engine/Plugins/Experimental/Buoyancy/Source/Runtime/Private/BuoyancyAlgorithms.cpp
--- a/Engine/Plugins/Experimental/Buoyancy/Source/Runtime/Private/BuoyancyAlgorithms.cpp
+++ b/Engine/Plugins/Experimental/Buoyancy/Source/Runtime/Private/BuoyancyAlgorithms.cpp
@@ -356,8 +356,8 @@ namespace BuoyancyAlgorithms
 				// If we have a sphere, ignore rotation because submerged volume is independent of rotation
 				// and also we don't want to apply any torques on the wheel.
 				// @todo(chaos): ComputeSubmergedBounds special case for spheres
-				const FVec3 SphereCenter = ShapeWorldTransform.TransformPosition(Sphere->GetCenter());
-				const FVec3 SphereExtent = FVec3(Sphere->GetRadius());
+				const FVec3 SphereCenter = ShapeWorldTransform.TransformPosition(FVec3(Sphere->GetCenterf()));
+				const FVec3 SphereExtent = FVec3(Sphere->GetRadiusf());
 				LocalBox = FAABB3(-SphereExtent, SphereExtent);
 				ShapeWorldTransform.SetTranslation(SphereCenter);
 				ShapeWorldTransform.SetRotation(FRotationMatrix::MakeFromZ(WaterN).ToQuat());
@@ -714,8 +714,8 @@ namespace BuoyancyAlgorithms
 				{
 					// If we have a sphere, ignore rotation because submerged volume is independent of rotation
 					// and also we don't want to apply any torques on the wheel.			
-					const FVec3 SphereCenter = ShapeWorldTransform.TransformPosition(Sphere->GetCenter());
-					const FVec3 SphereExtent = FVec3(Sphere->GetRadius());
+					const FVec3 SphereCenter = ShapeWorldTransform.TransformPosition(FVec3(Sphere->GetCenterf()));
+					const FVec3 SphereExtent = FVec3(Sphere->GetRadiusf());
 										
 					LocalBox = FAABB3(-SphereExtent, SphereExtent);
 					ShapeWorldTransform.SetTranslation(SphereCenter);
diff --git a/Engine/Plugins/Experimental/ChaosFlesh/Source/ChaosFleshEngine/Private/ChaosFlesh/ChaosDeformableCollisionsComponent.cpp b/Engine/Plugins/Experimental/ChaosFlesh/Source/ChaosFleshEngine/Private/ChaosFlesh/ChaosDeformableCollisionsComponent.cpp
--- a/Engine/Plugins/Experimental/ChaosFlesh/Source/ChaosFleshEngine/Private/ChaosFlesh/ChaosDeformableCollisionsComponent.cpp
+++ b/Engine/Plugins/Experimental/ChaosFlesh/Source/ChaosFleshEngine/Private/ChaosFlesh/ChaosDeformableCollisionsComponent.cpp
@@ -75,7 +75,7 @@ UDeformablePhysicsComponent::FThreadingProxy* UDeformableCollisionsComponent::Ne
 	return new FCollisionThreadingProxy(this);
 }
 
-template<> class Chaos::TSphere<float, 3>;
+// template<> class Chaos::TSphere<float, 3>;
 
 UDeformableCollisionsComponent::FDataMapValue 
 UDeformableCollisionsComponent::NewDeformableData()
@@ -111,7 +111,7 @@ UDeformableCollisionsComponent::NewDeformableData()
 					{
 						FKSphereElem& S = BodySetup->AggGeom.SphereElems[Index];
 						FTransform ShapeTransform = FTransform(Scale * S.Center) * P;
-						Geometry = new ::Chaos::TSphere<Chaos::FReal, 3>(FVector(0), S.Radius * Scale.GetMax());
+						Geometry = new ::Chaos::FSphere(FVector(0), S.Radius * Scale.GetMax());
 						PrevXfs[PrevXfIndex] = ShapeTransform;
 						AddedBodiesData.Add(Chaos::Softs::FCollisionObjectAddedBodies(
 							{ CollisionBody, Chaos::Softs::ERigidCollisionShapeType::Sphere, Index },
diff --git a/Engine/Plugins/Experimental/ChaosFlesh/Source/ChaosFleshNodes/Private/Dataflow/ChaosFleshSetFleshBonePositionTargetBindingNode.cpp b/Engine/Plugins/Experimental/ChaosFlesh/Source/ChaosFleshNodes/Private/Dataflow/ChaosFleshSetFleshBonePositionTargetBindingNode.cpp
--- a/Engine/Plugins/Experimental/ChaosFlesh/Source/ChaosFleshNodes/Private/Dataflow/ChaosFleshSetFleshBonePositionTargetBindingNode.cpp
+++ b/Engine/Plugins/Experimental/ChaosFlesh/Source/ChaosFleshNodes/Private/Dataflow/ChaosFleshSetFleshBonePositionTargetBindingNode.cpp
@@ -72,21 +72,21 @@ void FSetFleshBonePositionTargetBindingDataflowNode::Evaluate(UE::Dataflow::FCon
 					Chaos::TVec3<float> CoordDiff = (CoordMaxs - CoordMins) * VertexRadiusRatio;
 					SphereRadius = Chaos::FReal(FGenericPlatformMath::Max(CoordDiff[0], FGenericPlatformMath::Max(CoordDiff[1], CoordDiff[2])));
 
-					TArray<Chaos::TSphere<Chaos::FReal, 3>*> VertexSpherePtrs;
-					TArray<Chaos::TSphere<Chaos::FReal, 3>> VertexSpheres;
+					TArray<Chaos::FSphere*> VertexSpherePtrs;
+					TArray<Chaos::FSphere> VertexSpheres;
 
-					VertexSpheres.Init(Chaos::TSphere<Chaos::FReal, 3>(Chaos::TVec3<Chaos::FReal>(0), SphereRadius), NumSkeletonVertices);
+					VertexSpheres.Init(Chaos::FSphere(Chaos::TVec3<Chaos::FReal>(0), SphereRadius), NumSkeletonVertices);
 					VertexSpherePtrs.SetNum(NumSkeletonVertices);
 
 					for (int32 i = 0; i < int32(NumSkeletonVertices); i++)
 					{
 						Chaos::TVec3<Chaos::FReal> SphereCenter(PositionVertexBuffer.VertexPosition(i));
-						Chaos::TSphere<Chaos::FReal, 3> VertexSphere(SphereCenter, SphereRadius);
-						VertexSpheres[i] = Chaos::TSphere<Chaos::FReal, 3>(SphereCenter, SphereRadius);
+						Chaos::FSphere VertexSphere(SphereCenter, SphereRadius);
+						VertexSpheres[i] = Chaos::FSphere(SphereCenter, SphereRadius);
 						VertexSpherePtrs[i] = &VertexSpheres[i];
 					}
 					Chaos::TBoundingVolumeHierarchy<
-						TArray<Chaos::TSphere<Chaos::FReal, 3>*>,
+						TArray<Chaos::FSphere*>,
 						TArray<int32>,
 						Chaos::FReal,
 						3> VertexBVH(VertexSpherePtrs);
diff --git a/Engine/Plugins/Experimental/ChaosFlesh/Source/ChaosFleshNodes/Private/Dataflow/ChaosFleshSetVertexVertexPositionTargetBindingNode.cpp b/Engine/Plugins/Experimental/ChaosFlesh/Source/ChaosFleshNodes/Private/Dataflow/ChaosFleshSetVertexVertexPositionTargetBindingNode.cpp
--- a/Engine/Plugins/Experimental/ChaosFlesh/Source/ChaosFleshNodes/Private/Dataflow/ChaosFleshSetVertexVertexPositionTargetBindingNode.cpp
+++ b/Engine/Plugins/Experimental/ChaosFlesh/Source/ChaosFleshNodes/Private/Dataflow/ChaosFleshSetVertexVertexPositionTargetBindingNode.cpp
@@ -44,21 +44,21 @@ void FSetVertexVertexPositionTargetBindingDataflowNode::Evaluate(UE::Dataflow::F
 
 				SphereRadius = Chaos::FReal(FGenericPlatformMath::Min(CoordDiff[0], FGenericPlatformMath::Min(CoordDiff[1], CoordDiff[2])));
 
-				TArray<Chaos::TSphere<Chaos::FReal, 3>*> VertexSpherePtrs;
-				TArray<Chaos::TSphere<Chaos::FReal, 3>> VertexSpheres;
+				TArray<Chaos::FSphere*> VertexSpherePtrs;
+				TArray<Chaos::FSphere> VertexSpheres;
 
-				VertexSpheres.Init(Chaos::TSphere<Chaos::FReal, 3>(Chaos::TVec3<Chaos::FReal>(0), SphereRadius), Vertices->Num());
+				VertexSpheres.Init(Chaos::FSphere(Chaos::TVec3<Chaos::FReal>(0), SphereRadius), Vertices->Num());
 				VertexSpherePtrs.SetNum(Vertices->Num());
 
 				for (int32 i = 0; i < Vertices->Num(); i++)
 				{
 					Chaos::TVec3<Chaos::FReal> SphereCenter((*Vertices)[i]);
-					Chaos::TSphere<Chaos::FReal, 3> VertexSphere(SphereCenter, SphereRadius);
-					VertexSpheres[i] = Chaos::TSphere<Chaos::FReal, 3>(SphereCenter, SphereRadius);
+					Chaos::FSphere VertexSphere(SphereCenter, SphereRadius);
+					VertexSpheres[i] = Chaos::FSphere(SphereCenter, SphereRadius);
 					VertexSpherePtrs[i] = &VertexSpheres[i];
 				}
 				Chaos::TBoundingVolumeHierarchy<
-					TArray<Chaos::TSphere<Chaos::FReal, 3>*>,
+					TArray<Chaos::FSphere*>,
 					TArray<int32>,
 					Chaos::FReal,
 					3> VertexBVH(VertexSpherePtrs);
diff --git a/Engine/Plugins/Experimental/GeometryCollectionPlugin/Source/GeometryCollectionDepNodes/Private/Dataflow/GeometryCollectionTransferVertexScalarAttributeDepNode.cpp b/Engine/Plugins/Experimental/GeometryCollectionPlugin/Source/GeometryCollectionDepNodes/Private/Dataflow/GeometryCollectionTransferVertexScalarAttributeDepNode.cpp
--- a/Engine/Plugins/Experimental/GeometryCollectionPlugin/Source/GeometryCollectionDepNodes/Private/Dataflow/GeometryCollectionTransferVertexScalarAttributeDepNode.cpp
+++ b/Engine/Plugins/Experimental/GeometryCollectionPlugin/Source/GeometryCollectionDepNodes/Private/Dataflow/GeometryCollectionTransferVertexScalarAttributeDepNode.cpp
@@ -431,17 +431,17 @@ void FGeometryCollectionTransferVertexScalarAttributeNode::BuildComponentSpaceVe
 
 UE::Private::BVH* FGeometryCollectionTransferVertexScalarAttributeNode::BuildParticleSphereBVH(const TArray<FVector>& Vertices, float Radius)
 {
-	TArray<Chaos::TSphere<Chaos::FReal, 3>*> VertexSpherePtrs;
-	TArray<Chaos::TSphere<Chaos::FReal, 3>> VertexSpheres;
+	TArray<Chaos::FSphere*> VertexSpherePtrs;
+	TArray<Chaos::FSphere> VertexSpheres;
 
-	VertexSpheres.Init(Chaos::TSphere<Chaos::FReal, 3>(Chaos::TVec3<Chaos::FReal>(0), Radius), Vertices.Num());
+	VertexSpheres.Init(Chaos::FSphere(Chaos::TVec3<Chaos::FReal>(0), Radius), Vertices.Num());
 	VertexSpherePtrs.SetNum(Vertices.Num());
 
 	for (int32 i = 0; i < Vertices.Num(); i++)
 	{
 		Chaos::TVec3<Chaos::FReal> SphereCenter(Vertices[i]);
-		Chaos::TSphere<Chaos::FReal, 3> VertexSphere(SphereCenter, Radius);
-		VertexSpheres[i] = Chaos::TSphere<Chaos::FReal, 3>(SphereCenter, Radius);
+		Chaos::FSphere VertexSphere(SphereCenter, Radius);
+		VertexSpheres[i] = Chaos::FSphere(SphereCenter, Radius);
 		VertexSpherePtrs[i] = &VertexSpheres[i];
 	}
 	return new UE::Private::BVH(VertexSpherePtrs);
diff --git a/Engine/Plugins/Experimental/GeometryCollectionPlugin/Source/GeometryCollectionDepNodes/Private/Dataflow/GeometryCollectionTransferVertexScalarAttributeDepNode.h b/Engine/Plugins/Experimental/GeometryCollectionPlugin/Source/GeometryCollectionDepNodes/Private/Dataflow/GeometryCollectionTransferVertexScalarAttributeDepNode.h
--- a/Engine/Plugins/Experimental/GeometryCollectionPlugin/Source/GeometryCollectionDepNodes/Private/Dataflow/GeometryCollectionTransferVertexScalarAttributeDepNode.h
+++ b/Engine/Plugins/Experimental/GeometryCollectionPlugin/Source/GeometryCollectionDepNodes/Private/Dataflow/GeometryCollectionTransferVertexScalarAttributeDepNode.h
@@ -14,7 +14,7 @@
 namespace UE::Private 
 {
 	class FTransferFacade;
-	typedef Chaos::TBoundingVolumeHierarchy<TArray<Chaos::TSphere<Chaos::FReal, 3>*>, TArray<int32>, Chaos::FReal, 3> BVH;
+	typedef Chaos::TBoundingVolumeHierarchy<TArray<Chaos::FSphere*>, TArray<int32>, Chaos::FReal, 3> BVH;
 };
 
 UENUM(BlueprintType)
diff --git a/Engine/Plugins/Experimental/GeometryCollectionPlugin/Source/GeometryCollectionNodes/Private/Dataflow/GeometryCollectionTransferVertexAttributeNode.cpp b/Engine/Plugins/Experimental/GeometryCollectionPlugin/Source/GeometryCollectionNodes/Private/Dataflow/GeometryCollectionTransferVertexAttributeNode.cpp
--- a/Engine/Plugins/Experimental/GeometryCollectionPlugin/Source/GeometryCollectionNodes/Private/Dataflow/GeometryCollectionTransferVertexAttributeNode.cpp
+++ b/Engine/Plugins/Experimental/GeometryCollectionPlugin/Source/GeometryCollectionNodes/Private/Dataflow/GeometryCollectionTransferVertexAttributeNode.cpp
@@ -16,7 +16,7 @@
 
 namespace UE::Private
 {
-	typedef Chaos::TSphere<Chaos::FReal, 3> SphereType;
+	typedef Chaos::FSphere SphereType;
 	typedef Chaos::TBoundingVolumeHierarchy<TArray<SphereType*>, TArray<int32>, Chaos::FReal, 3> BVH;
 	class FTransferFacade
 	{
diff --git a/Engine/Source/Programs/HeadlessChaos/Private/GeometryCollection/GeometryCollectionTestBoneHierarchy.cpp b/Engine/Source/Programs/HeadlessChaos/Private/GeometryCollection/GeometryCollectionTestBoneHierarchy.cpp
--- a/Engine/Source/Programs/HeadlessChaos/Private/GeometryCollection/GeometryCollectionTestBoneHierarchy.cpp
+++ b/Engine/Source/Programs/HeadlessChaos/Private/GeometryCollection/GeometryCollectionTestBoneHierarchy.cpp
@@ -21,7 +21,7 @@ namespace GeometryCollectionTest
 
 	void TestSphere(const Chaos::FVec3& Center, const Chaos::FReal Radius, const FTransform &BoneRelXf)
 	{
-		Chaos::TSphere<Chaos::FReal, 3>* Sphere = new Chaos::TSphere<Chaos::FReal, 3>(Center, Radius);
+		Chaos::FSphere* Sphere = new Chaos::FSphere(Center, Radius);
 		TUniquePtr<FAnalyticImplicitGroup> Group1(new FAnalyticImplicitGroup("Root", 0));
 		Group1->Init(1);
 		Group1->SetParentBoneIndex(INDEX_NONE);
@@ -40,8 +40,8 @@ namespace GeometryCollectionTest
 		const FTransform &BoneRelXf1,
 		const FTransform &BoneRelXf2)
 	{
-		Chaos::TSphere<Chaos::FReal, 3>* Sphere1 = new Chaos::TSphere<Chaos::FReal, 3>(Center1, Radius1);
-		Chaos::TSphere<Chaos::FReal, 3>* Sphere2 = new Chaos::TSphere<Chaos::FReal, 3>(Center2, Radius2);
+		Chaos::FSphere* Sphere1 = new Chaos::FSphere(Center1, Radius1);
+		Chaos::FSphere* Sphere2 = new Chaos::FSphere(Center2, Radius2);
 		TUniquePtr<FAnalyticImplicitGroup> Group1(new FAnalyticImplicitGroup("Root", 0));
 		Group1->Init(2);
 		Group1->SetParentBoneIndex(INDEX_NONE);
@@ -84,13 +84,13 @@ namespace GeometryCollectionTest
 			TUniquePtr<FAnalyticImplicitGroup> Group1(
 				new FAnalyticImplicitGroup("Root", 0));
 			Group1->SetParentBoneIndex(INDEX_NONE);
-			Group1->Add(FTransform::Identity, new Chaos::TSphere<Chaos::FReal, 3>(Chaos::FVec3(0), 1.));
+			Group1->Add(FTransform::Identity, new Chaos::FSphere(Chaos::FVec3(0), 1.));
 			Hierarchy.Add(MoveTemp(Group1));
 
 			TUniquePtr<FAnalyticImplicitGroup> Group2(
 				new FAnalyticImplicitGroup("Bone1", 1));
 			Group2->SetParentBoneIndex(0);
-			Group2->Add(FTransform::Identity, new Chaos::TSphere<Chaos::FReal, 3>(Chaos::FVec3(0), 1.));
+			Group2->Add(FTransform::Identity, new Chaos::FSphere(Chaos::FVec3(0), 1.));
 			Hierarchy.Add(MoveTemp(Group2));
 
 			Hierarchy.InitPostAdd();
diff --git a/Engine/Source/Programs/HeadlessChaos/Private/GeometryCollection/GeometryCollectionTestFramework.cpp b/Engine/Source/Programs/HeadlessChaos/Private/GeometryCollection/GeometryCollectionTestFramework.cpp
--- a/Engine/Source/Programs/HeadlessChaos/Private/GeometryCollection/GeometryCollectionTestFramework.cpp
+++ b/Engine/Source/Programs/HeadlessChaos/Private/GeometryCollection/GeometryCollectionTestFramework.cpp
@@ -28,7 +28,7 @@ namespace GeometryCollectionTest
 		SphereGen.Generate();
 
 		// FSphereGenerator's points drift off the surface just a bit, so we correct for that.
-		Chaos::TSphere<Chaos::FReal, 3> Sphere(Chaos::FVec3(0), SphereGen.Radius);
+		Chaos::FSphere Sphere(Chaos::FVec3(0), SphereGen.Radius);
 		Chaos::FVec3 Normal;
 		for (int32 Idx = 0; Idx < SphereGen.Vertices.Num(); Idx++)
 		{
diff --git a/Engine/Source/Programs/HeadlessChaos/Private/GeometryCollection/GeometryCollectionTestImplicitCapsule.cpp b/Engine/Source/Programs/HeadlessChaos/Private/GeometryCollection/GeometryCollectionTestImplicitCapsule.cpp
--- a/Engine/Source/Programs/HeadlessChaos/Private/GeometryCollection/GeometryCollectionTestImplicitCapsule.cpp
+++ b/Engine/Source/Programs/HeadlessChaos/Private/GeometryCollection/GeometryCollectionTestImplicitCapsule.cpp
@@ -10,42 +10,42 @@ namespace GeometryCollectionTest
 
 	void RunTestComputeSamplePoints(const Chaos::FCapsule &Capsule)
 	{
-		const Chaos::FCapsule OACapsule = Chaos::FCapsule::NewFromOriginAndAxis(Capsule.GetOrigin(), Capsule.GetAxis(), Capsule.GetHeight(), Capsule.GetRadius());
-		EXPECT_NEAR((Capsule.GetOrigin() - OACapsule.GetOrigin()).Size(), 0, KINDA_SMALL_NUMBER) << *FString("Capsule != OACapsule, origin.");
+		const Chaos::FCapsule OACapsule = Chaos::FCapsule::NewFromOriginAndAxis(Capsule.GetOriginf(), Capsule.GetAxisf(), Capsule.GetHeightf(), Capsule.GetRadiusf());
+		EXPECT_NEAR((Capsule.GetOriginf() - OACapsule.GetOriginf()).Size(), 0, KINDA_SMALL_NUMBER) << *FString("Capsule != OACapsule, origin.");
 		EXPECT_NEAR((Capsule.GetInsertion() - OACapsule.GetInsertion()).Size(), 0, KINDA_SMALL_NUMBER) << *FString("Capsule != OACapsule, insertion.");
-		EXPECT_NEAR((Capsule.GetAxis() - OACapsule.GetAxis()).Size(), 0, KINDA_SMALL_NUMBER) << *FString("Capsule != OACapsule, axis.");
-		EXPECT_NEAR(Capsule.GetHeight() - OACapsule.GetHeight(), 0, KINDA_SMALL_NUMBER) << *FString("Capsule != OACapsule, height.");
-		EXPECT_NEAR(Capsule.GetRadius() - OACapsule.GetRadius(), 0, KINDA_SMALL_NUMBER) << *FString("Capsule != OACapsule, radius.");
+		EXPECT_NEAR((Capsule.GetAxisf() - OACapsule.GetAxisf()).Size(), 0, KINDA_SMALL_NUMBER) << *FString("Capsule != OACapsule, axis.");
+		EXPECT_NEAR(Capsule.GetHeightf() - OACapsule.GetHeightf(), 0, KINDA_SMALL_NUMBER) << *FString("Capsule != OACapsule, height.");
+		EXPECT_NEAR(Capsule.GetRadiusf() - OACapsule.GetRadiusf(), 0, KINDA_SMALL_NUMBER) << *FString("Capsule != OACapsule, radius.");
 
 		FVec3 Point;
 		FReal Phi;
 		EXPECT_EQ(Capsule.GetType(), Chaos::ImplicitObjectType::Capsule) << *FString("Implicit object type is not 'capsule'.");
 
-		Point = Capsule.GetAxis();
+		Point = Capsule.GetAxisf();
 		EXPECT_LT(FMath::Abs(Point.Size() - 1.0), KINDA_SMALL_NUMBER) << *FString("Capsule axis is not unit length.");
 		
-		Point = Capsule.GetOrigin() + Capsule.GetAxis() * (Capsule.GetHeight() + Capsule.GetRadius() * 2);
+		Point = Capsule.GetOriginf() + Capsule.GetAxisf() * (Capsule.GetHeightf() + Capsule.GetRadiusf() * 2);
 		EXPECT_LT((Point - Capsule.GetInsertion()).Size(), KINDA_SMALL_NUMBER) << *FString("Capsule is broken.");
 
-		Point = Capsule.GetInsertion();// Capsule.GetOrigin() + Capsule.GetAxis() * Capsule.GetHeight();
+		Point = Capsule.GetInsertion();// Capsule.GetOriginf() + Capsule.GetAxis() * Capsule.GetHeight();
 		Phi = Capsule.SignedDistance(Point);
 		EXPECT_NEAR(Phi, 0, KINDA_SMALL_NUMBER) << *FString("Capsule failed phi surface (insertion) sanity test.");
 
-		Point = Capsule.GetOrigin() + Capsule.GetAxis() * (Capsule.GetHeight() + Capsule.GetRadius() * 2) * 0.25;
+		Point = Capsule.GetOriginf() + Capsule.GetAxisf() * (Capsule.GetHeightf() + Capsule.GetRadiusf() * 2) * 0.25;
 		Phi = Capsule.SignedDistance(Point);
 		EXPECT_LE(Phi, (FReal)0.0) << *FString("Capsule failed phi depth (1/4 origin) sanity test.");
 
-		Point = Capsule.GetOrigin() + Capsule.GetAxis() * (Capsule.GetHeight() + Capsule.GetRadius() * 2) * 0.75;
+		Point = Capsule.GetOriginf() + Capsule.GetAxisf() * (Capsule.GetHeightf() + Capsule.GetRadiusf() * 2) * 0.75;
 		Phi = Capsule.SignedDistance(Point);
 		EXPECT_LE(Phi, (FReal)0.0) << *FString("Capsule failed phi depth (3/4 origin) sanity test.");
 
-		EXPECT_NEAR((Capsule.GetCenter() - FVec3(Capsule.GetOrigin() + Capsule.GetAxis() * (Capsule.GetHeight() + Capsule.GetRadius() * 2) * 0.5)).Size(), 0, KINDA_SMALL_NUMBER) << *FString("Capsule center is off mid axis.");
+		EXPECT_NEAR((Capsule.GetCenterf() - FVec3(Capsule.GetOriginf() + Capsule.GetAxisf() * (Capsule.GetHeightf() + Capsule.GetRadiusf() * 2) * 0.5)).Size(), 0, KINDA_SMALL_NUMBER) << *FString("Capsule center is off mid axis.");
 
-		Point = Capsule.GetCenter();
+		Point = Capsule.GetCenterf();
 		Phi = Capsule.SignedDistance(Point);
 		EXPECT_LT(Phi, (FReal)0.0) << *FString("Capsule failed phi depth sanity test.");
 		
-		Point = Capsule.GetOrigin();
+		Point = Capsule.GetOriginf();
 		Phi = Capsule.SignedDistance(Point);
 		EXPECT_NEAR(FMath::Abs(Phi), 0, KINDA_SMALL_NUMBER) << *FString("Capsule failed phi surface (origin) sanity test.");
 
@@ -53,11 +53,11 @@ namespace GeometryCollectionTest
 		Phi = Capsule.SignedDistance(Point);
 		EXPECT_NEAR(FMath::Abs(Phi), 0, KINDA_SMALL_NUMBER) << *FString("Capsule failed phi surface (origin+axis*height) sanity test.");
 
-		Point = Capsule.GetOrigin() + Capsule.GetAxis() * Capsule.GetRadius() + Capsule.GetAxis().GetOrthogonalVector().GetSafeNormal() * Capsule.GetRadius();
+		Point = Capsule.GetOriginf() + Capsule.GetAxisf() * Capsule.GetRadiusf() + Capsule.GetAxisf().GetOrthogonalVector().GetSafeNormal() * Capsule.GetRadiusf();
 		Phi = Capsule.SignedDistance(Point);
 		EXPECT_NEAR(FMath::Abs(Phi), 0, KINDA_SMALL_NUMBER) << *FString("Capsule failed phi surface (origin+orthogonalAxis*radius) sanity test.");
 
-		Point = Capsule.GetCenter() + Capsule.GetAxis().GetOrthogonalVector().GetSafeNormal() * Capsule.GetRadius();
+		Point = Capsule.GetCenterf() + Capsule.GetAxisf().GetOrthogonalVector().GetSafeNormal() * Capsule.GetRadiusf();
 		Phi = Capsule.SignedDistance(Point);
 		EXPECT_NEAR(FMath::Abs(Phi), 0, KINDA_SMALL_NUMBER) << *FString("Capsule failed phi surface (center+orthogonalAxis*radius) sanity test.");
 
diff --git a/Engine/Source/Programs/HeadlessChaos/Private/GeometryCollection/GeometryCollectionTestImplicitSphere.cpp b/Engine/Source/Programs/HeadlessChaos/Private/GeometryCollection/GeometryCollectionTestImplicitSphere.cpp
--- a/Engine/Source/Programs/HeadlessChaos/Private/GeometryCollection/GeometryCollectionTestImplicitSphere.cpp
+++ b/Engine/Source/Programs/HeadlessChaos/Private/GeometryCollection/GeometryCollectionTestImplicitSphere.cpp
@@ -26,7 +26,7 @@ namespace GeometryCollectionTest
 		const FVec3& EndPt,
 		const IntersectionResult ExpectedRes)
 	{
-		TSphere<FReal, 3> Sphere(Center, Radius);
+		Chaos::FSphere Sphere(Center, Radius);
 		IntersectionResult Res = Sphere.FindClosestIntersection(StartPt, EndPt, Thickness);
 		EXPECT_TRUE(
 			(!Res.Second && !ExpectedRes.Second) ||
@@ -311,15 +311,15 @@ namespace GeometryCollectionTest
 	// Sample points tests
 	//==========================================================================
 
-	void RunTestComputeSamplePoints(const TSphere<FReal, 3> &Sphere)
+	void RunTestComputeSamplePoints(const Chaos::FSphere &Sphere)
 	{
 		EXPECT_EQ(Sphere.GetType(), Chaos::ImplicitObjectType::Sphere) << *FString("Implicit object type is not 'sphere'.");
 
-		FVec3 Point = Sphere.GetCenter();
+		FVec3 Point = Sphere.GetCenterf();
 		FReal Phi = Sphere.SignedDistance(Point);
-		EXPECT_LE(FMath::Abs(Phi + Sphere.GetRadius()), SMALL_NUMBER) << *FString("Sphere failed phi depth sanity test.");
+		EXPECT_LE(FMath::Abs(Phi + Sphere.GetRadiusf()), SMALL_NUMBER) << *FString("Sphere failed phi depth sanity test.");
 		
-		Point[0] += Sphere.GetRadius();
+		Point[0] += Sphere.GetRadiusf();
 		Phi = Sphere.SignedDistance(Point);
 		EXPECT_LE(FMath::Abs(Phi), KINDA_SMALL_NUMBER) << *FString("Sphere failed phi surface sanity test.");
 
@@ -339,7 +339,7 @@ namespace GeometryCollectionTest
 		}
 	}
 
-	void RunTestComputeSemispherePoints(const TSphere<FReal, 3> &Sphere)
+	void RunTestComputeSemispherePoints(const Chaos::FSphere &Sphere)
 	{
 		TArray<FVec3> Points;
 		Chaos::TSphereSpecializeSamplingHelper<FReal, 3>::ComputeBottomHalfSemiSphere(Points, Sphere, 100);
@@ -350,7 +350,7 @@ namespace GeometryCollectionTest
 			check(OnSurface);
 			EXPECT_TRUE(OnSurface) << *FString("Produced a point not on the surface of the sphere.");
 
-			const bool BelowCenter = Pt[2] < Sphere.GetCenter()[2] + KINDA_SMALL_NUMBER;
+			const bool BelowCenter = Pt[2] < Sphere.GetCenterf()[2] + KINDA_SMALL_NUMBER;
 			check(BelowCenter);
 			EXPECT_TRUE(BelowCenter) << *FString("Bottom semisphere produced a point above midline.");
 		}
@@ -364,7 +364,7 @@ namespace GeometryCollectionTest
 			check(OnSurface);
 			EXPECT_TRUE(OnSurface) << *FString("Produced a point not on the surface of the sphere.");
 
-			const bool BelowCenter = Pt[2] > Sphere.GetCenter()[2] - KINDA_SMALL_NUMBER;
+			const bool BelowCenter = Pt[2] > Sphere.GetCenterf()[2] - KINDA_SMALL_NUMBER;
 			check(BelowCenter);
 			EXPECT_TRUE(BelowCenter) << *FString("Top semisphere produced a point above midline.");
 		}
@@ -374,7 +374,7 @@ namespace GeometryCollectionTest
 	{
 		// At the origin with radius 1
 		{
-			TSphere<FReal, 3> Sphere(FVec3((FReal)0.0), (FReal)1.0);
+			Chaos::FSphere Sphere(FVec3((FReal)0.0), (FReal)1.0);
 			RunTestComputeSemispherePoints(Sphere);
 		}	
 	}
@@ -383,32 +383,32 @@ namespace GeometryCollectionTest
 	{
 		// At the origin with radius 1
 		{
-			TSphere<FReal, 3> Sphere(FVec3(0.0), (FReal)1.0);
+			Chaos::FSphere Sphere(FVec3(0.0), (FReal)1.0);
 			RunTestComputeSamplePoints(Sphere);
 		}
 		// At the origin with radius > 1
 		{
-			TSphere<FReal, 3> Sphere(FVec3(0.0), (FReal)10.0);
+			Chaos::FSphere Sphere(FVec3(0.0), (FReal)10.0);
 			RunTestComputeSamplePoints(Sphere);
 		}
 		// At the origin with radius < 1
 		{
-			TSphere<FReal, 3> Sphere(FVec3(0.0), (FReal).1);
+			Chaos::FSphere Sphere(FVec3(0.0), (FReal).1);
 			RunTestComputeSamplePoints(Sphere);
 		}
 		// Off the origin with radius 1
 		{
-			TSphere<FReal, 3> Sphere(FVec3(10.0), (FReal)1.0);
+			Chaos::FSphere Sphere(FVec3(10.0), (FReal)1.0);
 			RunTestComputeSamplePoints(Sphere);
 		}
 		// Off the origin with radius > 1
 		{
-			TSphere<FReal, 3> Sphere(FVec3(10.0), (FReal)10.0);
+			Chaos::FSphere Sphere(FVec3(10.0), (FReal)10.0);
 			RunTestComputeSamplePoints(Sphere);
 		}
 		// Off the origin with radius < 1
 		{
-			TSphere<FReal, 3> Sphere(FVec3(10.0), (FReal).1);
+			Chaos::FSphere Sphere(FVec3(10.0), (FReal).1);
 			RunTestComputeSamplePoints(Sphere);
 		}
 	}
diff --git a/Engine/Source/Programs/HeadlessChaos/Private/HeadlessChaosTestCCD.cpp b/Engine/Source/Programs/HeadlessChaos/Private/HeadlessChaosTestCCD.cpp
--- a/Engine/Source/Programs/HeadlessChaos/Private/HeadlessChaosTestCCD.cpp
+++ b/Engine/Source/Programs/HeadlessChaos/Private/HeadlessChaosTestCCD.cpp
@@ -171,7 +171,7 @@ namespace ChaosTest
 		PhysicsMaterial->Restitution = 1.0f;
 
 		// Create Sphere geometry (Radius = 100)
-		Chaos::FImplicitObjectPtr Sphere(new TSphere<FReal, 3>(FVec3(0, 0, 0), SphereRadius));
+		Chaos::FImplicitObjectPtr Sphere(new Chaos::FSphere(FVec3(0, 0, 0), SphereRadius));
 
 		// Assign sphere geometry to both particles 
 		Static->SetGeometry(Sphere);
diff --git a/Engine/Source/Programs/HeadlessChaos/Private/HeadlessChaosTestCapsuleTriangleCollision.cpp b/Engine/Source/Programs/HeadlessChaos/Private/HeadlessChaosTestCapsuleTriangleCollision.cpp
--- a/Engine/Source/Programs/HeadlessChaos/Private/HeadlessChaosTestCapsuleTriangleCollision.cpp
+++ b/Engine/Source/Programs/HeadlessChaos/Private/HeadlessChaosTestCapsuleTriangleCollision.cpp
@@ -24,11 +24,11 @@ namespace ChaosTest
 
 	bool IsPointOnCapsuleSurface(const FImplicitCapsule3& Capsule, const FVec3& Point, const FReal Tolerance)
 	{
-		const FReal SegmentT = Utilities::ClosestTimeOnLineSegment(Point, Capsule.GetX1(), Capsule.GetX2());
-		const FVec3 SegmentPoint = FMath::Lerp(Capsule.GetX1(), Capsule.GetX2(), SegmentT);
+		const FReal SegmentT = Utilities::ClosestTimeOnLineSegment(Point, FVec3(Capsule.GetX1f()), FVec3(Capsule.GetX2f()));
+		const FVec3 SegmentPoint = FMath::Lerp(Capsule.GetX1f(), Capsule.GetX2f(), SegmentT);
 		const FVec3 Delta = Point - SegmentPoint;
 		const FReal Distance = Delta.Size();
-		return FMath::IsNearlyEqual(Distance, Capsule.GetRadius(), Tolerance);
+		return FMath::IsNearlyEqual(Distance, Capsule.GetRadiusf(), Tolerance);
 	}
 
 	// Double face contact: both capsule points within the edge planes.
@@ -59,8 +59,8 @@ namespace ChaosTest
 		EXPECT_EQ(ContactPoints.Num(), 2);
 		if (ContactPoints.Num() == 2)
 		{
-			EXPECT_TRUE(ManifoldContainsPoint(ContactPoints, 0, Capsule.GetX1() - Capsule.GetRadius() * Triangle.GetNormal(), Tolerance));
-			EXPECT_TRUE(ManifoldContainsPoint(ContactPoints, 0, Capsule.GetX2() - Capsule.GetRadius() * Triangle.GetNormal(), Tolerance));
+			EXPECT_TRUE(ManifoldContainsPoint(ContactPoints, 0, Capsule.GetX1f() - Capsule.GetRadiusf() * Triangle.GetNormal(), Tolerance));
+			EXPECT_TRUE(ManifoldContainsPoint(ContactPoints, 0, Capsule.GetX2f() - Capsule.GetRadiusf() * Triangle.GetNormal(), Tolerance));
 		}
 	}
 
@@ -98,8 +98,8 @@ namespace ChaosTest
 		EXPECT_EQ(ContactPoints.Num(), 2);
 		if (ContactPoints.Num() == 2)
 		{
-			EXPECT_TRUE(ManifoldContainsPoint(ContactPoints, 0, Capsule.GetX1() - Capsule.GetRadius() * Triangle.GetNormal(), Tolerance));
-			EXPECT_TRUE(ManifoldContainsPoint(ContactPoints, 0, Capsule.GetX2() - Capsule.GetRadius() * Triangle.GetNormal(), Tolerance));
+			EXPECT_TRUE(ManifoldContainsPoint(ContactPoints, 0, Capsule.GetX1f() - Capsule.GetRadiusf() * Triangle.GetNormal(), Tolerance));
+			EXPECT_TRUE(ManifoldContainsPoint(ContactPoints, 0, Capsule.GetX2f() - Capsule.GetRadiusf() * Triangle.GetNormal(), Tolerance));
 		}
 	}
 
@@ -131,8 +131,8 @@ namespace ChaosTest
 		EXPECT_EQ(ContactPoints.Num(), 2);
 		if (ContactPoints.Num() == 2)
 		{
-			EXPECT_TRUE(ManifoldContainsPoint(ContactPoints, 0, Capsule.GetX1() - Capsule.GetRadius() * Triangle.GetNormal(), Tolerance));
-			EXPECT_TRUE(ManifoldContainsPoint(ContactPoints, 0, FVec3(100, -10, 10) - Capsule.GetRadius() * Triangle.GetNormal(), Tolerance));
+			EXPECT_TRUE(ManifoldContainsPoint(ContactPoints, 0, Capsule.GetX1f() - Capsule.GetRadiusf() * Triangle.GetNormal(), Tolerance));
+			EXPECT_TRUE(ManifoldContainsPoint(ContactPoints, 0, FVec3(100, -10, 10) - Capsule.GetRadiusf() * Triangle.GetNormal(), Tolerance));
 		}
 	}
 
@@ -203,7 +203,7 @@ namespace ChaosTest
 		EXPECT_EQ(ContactPoints.Num(), 1);
 		if (ContactPoints.Num() == 1)
 		{
-			EXPECT_NEAR(ContactPoints[0].Phi, -Capsule.GetRadius(), PositionTolerance);
+			EXPECT_NEAR(ContactPoints[0].Phi, -Capsule.GetRadiusf(), PositionTolerance);
 			EXPECT_NEAR(FVec3::DotProduct(ContactPoints[0].ShapeContactNormal, FVec3(0,0,1)), FMath::Cos(FMath::DegreesToRadians(45)), NormalTolerance);
 		}
 	}
@@ -234,9 +234,9 @@ namespace ChaosTest
 		EXPECT_EQ(ContactPoints.Num(), 1);
 		if (ContactPoints.Num() == 1)
 		{
-			EXPECT_NEAR(ContactPoints[0].ShapeContactPoints[0].X, Capsule.GetX1().X, PositionTolerance);
-			EXPECT_NEAR(ContactPoints[0].ShapeContactPoints[0].Y, Capsule.GetX1().Y, PositionTolerance);
-			EXPECT_NEAR(ContactPoints[0].ShapeContactPoints[0].Z, Capsule.GetX1().Z - Capsule.GetRadius(), PositionTolerance);
+			EXPECT_NEAR(ContactPoints[0].ShapeContactPoints[0].X, Capsule.GetX1f().X, PositionTolerance);
+			EXPECT_NEAR(ContactPoints[0].ShapeContactPoints[0].Y, Capsule.GetX1f().Y, PositionTolerance);
+			EXPECT_NEAR(ContactPoints[0].ShapeContactPoints[0].Z, Capsule.GetX1f().Z - Capsule.GetRadiusf(), PositionTolerance);
 
 			EXPECT_NEAR(ContactPoints[0].ShapeContactNormal.Z, FReal(1), NormalTolerance);
 		}
diff --git a/Engine/Source/Programs/HeadlessChaos/Private/HeadlessChaosTestDetectCollision.cpp b/Engine/Source/Programs/HeadlessChaos/Private/HeadlessChaosTestDetectCollision.cpp
--- a/Engine/Source/Programs/HeadlessChaos/Private/HeadlessChaosTestDetectCollision.cpp
+++ b/Engine/Source/Programs/HeadlessChaos/Private/HeadlessChaosTestDetectCollision.cpp
@@ -85,7 +85,7 @@ namespace ChaosTest
 		{
 			const FVec3 SphereOffset = FVec3(SpherePos.X - 0.5f * BoxSize.X, 0, SpherePos.Z);
 			const FVec3 ExpectedNormal = SphereOffset.GetSafeNormal();
-			const FReal ExpectedDistance = SphereOffset.Size() - Sphere.GetRadius();
+			const FReal ExpectedDistance = SphereOffset.Size() - Sphere.GetRadiusf();
 
 			CheckContactPoint(ContactPoint, -SphereRadius * ExpectedNormal, 0.5f * FVec3(BoxSize.X, 0, BoxSize.Z), ExpectedNormal, 1, ExpectedNormal, ExpectedDistance);
 		}
@@ -112,7 +112,7 @@ namespace ChaosTest
 		{
 			const FVec3 SphereOffset = FVec3(SpherePos.X - 0.5f * BoxSize.X, SpherePos.Y - 0.5f * BoxSize.Y, SpherePos.Z);
 			const FVec3 ExpectedNormal = SphereOffset.GetSafeNormal();
-			const FReal ExpectedDistance = SphereOffset.Size() - Sphere.GetRadius();
+			const FReal ExpectedDistance = SphereOffset.Size() - Sphere.GetRadiusf();
 
 			CheckContactPoint(ContactPoint, -SphereRadius * ExpectedNormal, 0.5f * FVec3(BoxSize.X, BoxSize.Y, BoxSize.Z), ExpectedNormal, 1, ExpectedNormal, ExpectedDistance);
 		}
@@ -139,7 +139,7 @@ namespace ChaosTest
 		{
 			const FVec3 SphereOffset = FVec3(SpherePos.X - 0.5f * BoxSize.X, SpherePos.Y - 0.5f * BoxSize.Y, SpherePos.Z);
 			const FVec3 ExpectedNormal = SphereOffset.GetSafeNormal();
-			const FReal ExpectedDistance = SphereOffset.Size() - Sphere.GetRadius();
+			const FReal ExpectedDistance = SphereOffset.Size() - Sphere.GetRadiusf();
 
 			CheckContactPoint(ContactPoint, -SphereRadius * ExpectedNormal, 0.5f * FVec3(BoxSize.X, BoxSize.Y, BoxSize.Z), ExpectedNormal, 1, ExpectedNormal, ExpectedDistance);
 		}
diff --git a/Engine/Source/Programs/HeadlessChaos/Private/HeadlessChaosTestEPA.cpp b/Engine/Source/Programs/HeadlessChaos/Private/HeadlessChaosTestEPA.cpp
--- a/Engine/Source/Programs/HeadlessChaos/Private/HeadlessChaosTestEPA.cpp
+++ b/Engine/Source/Programs/HeadlessChaos/Private/HeadlessChaosTestEPA.cpp
@@ -688,7 +688,7 @@ namespace ChaosTest
 			FConvexPtr Convex( new FConvex(SurfaceParticles, 0.0f));
 			TImplicitObjectScaled<FConvex> ScaledConvex(Convex, FVec3(1.0f), 0.0f);
 
-			TSphere<FReal, 3> Sphere(FVec3(0.0f), 34.2120171);
+			Chaos::FSphere Sphere(FVec3(0.0f), 34.2120171);
 
 			const FRigidTransform3 BToATM({ 568.001648, -535.998352, 8.00000000 }, FRotation3::FromElements(0.000000000, 0.000000000, -0.707105696, 0.707107902));
 			const FVec3 LocalDir(0.000000000, 0.000000000, -1.00000000);
@@ -713,7 +713,7 @@ namespace ChaosTest
 				MakeVectorRegisterFloat( 0.000000000f, 0.000000000f, 0.000000000f, 0.0f ), 
 				MakeVectorRegisterFloat(128.000000f, 0.000000000f, -114.064575, 0.0f), 
 				MakeVectorRegisterFloat(128.000000, 128.000000, 2.35327148, 0.0f));
-			const TSphere<FReal, 3> Sphere({ 0.0, 0.0, 0.0 }, 4);
+			const Chaos::FSphere Sphere({ 0.0, 0.0, 0.0 }, 4);
 			const TRigidTransform<FReal, 3> Transform({ 174.592773, -161.781250, -68.0469971 }, FQuat::Identity);
 			const FVec3 Dir(-0.406315684, 0.913382649, -0.0252906363);
 			const FReal Length = 430.961548;
diff --git a/Engine/Source/Programs/HeadlessChaos/Private/HeadlessChaosTestGJK.cpp b/Engine/Source/Programs/HeadlessChaos/Private/HeadlessChaosTestGJK.cpp
--- a/Engine/Source/Programs/HeadlessChaos/Private/HeadlessChaosTestGJK.cpp
+++ b/Engine/Source/Programs/HeadlessChaos/Private/HeadlessChaosTestGJK.cpp
@@ -434,8 +434,8 @@ namespace ChaosTest
 
 	void GJKSphereSphereTest()
 	{
-		TSphere<FReal, 3> A(FVec3(10, 0, 0), 5);
-		TSphere<FReal, 3> B(FVec3(4, 0, 0), 2);
+		Chaos::FSphere A(FVec3(10, 0, 0), 5);
+		Chaos::FSphere B(FVec3(4, 0, 0), 2);
 
 		FVec3 InitialDirs[] = { FVec3(1,0,0), FVec3(-1,0,0), FVec3(0,1,0), FVec3(0,-1,0), FVec3(0,0,1), FVec3(0,0,-1) };
 
@@ -459,7 +459,7 @@ namespace ChaosTest
 			EXPECT_FALSE(GJKIntersection<FReal>(A, B, FRigidTransform3(FVec3(6.5, 0, 0), FRotation3::FromVector(FVec3(0, 0, PI))), 0.01, InitialDir));
 
 			//hit tiny
-			TSphere<FReal, 3> Tiny(FVec3(0), 1e-2);
+			Chaos::FSphere Tiny(FVec3(0), 1e-2);
 			EXPECT_TRUE(GJKIntersection<FReal>(A, Tiny, FRigidTransform3(FVec3(15, 0, 0), FRotation3::Identity), 0, InitialDir));
 
 			//miss tiny
@@ -470,7 +470,7 @@ namespace ChaosTest
 
 	void GJKSphereBoxTest()
 	{
-		TSphere<FReal, 3> A(FVec3(10, 0, 0), 5);
+		Chaos::FSphere A(FVec3(10, 0, 0), 5);
 		FAABB3 B(FVec3(-4, -2, -4), FVec3(4,2,4));
 
 		FVec3 InitialDirs[] = { FVec3(1,0,0), FVec3(-1,0,0), FVec3(0,1,0), FVec3(0,-1,0), FVec3(0,0,1), FVec3(0,0,-1) };
@@ -510,7 +510,7 @@ namespace ChaosTest
 
 	void GJKSphereCapsuleTest()
 	{
-		TSphere<FReal, 3> A(FVec3(10, 0, 0), 5);
+		Chaos::FSphere A(FVec3(10, 0, 0), 5);
 		FCapsule B(FVec3(0, 0, -3), FVec3(0, 0, 3), 3);
 
 		FVec3 InitialDirs[] = { FVec3(1,0,0), FVec3(-1,0,0), FVec3(0,1,0), FVec3(0,-1,0), FVec3(0,0,1), FVec3(0,0,-1) };
@@ -547,7 +547,7 @@ namespace ChaosTest
 	void GJKSphereConvexTest()
 	{
 		FVec3 InitialDirs[] = { FVec3(1,0,0), FVec3(-1,0,0), FVec3(0,1,0), FVec3(0,-1,0), FVec3(0,0,1), FVec3(0,0,-1) };
-		TSphere<FReal, 3> A(FVec3(10, 0, 0), 5);
+		Chaos::FSphere A(FVec3(10, 0, 0), 5);
 
 		{
 			//Tetrahedron
@@ -623,11 +623,11 @@ namespace ChaosTest
 
 	void GJKSphereScaledSphereTest()
 	{
-		TSphere<FReal, 3> A(FVec3(10, 0, 0), 5);
-		FSpherePtr Sphere( new TSphere<FReal, 3>(FVec3(4, 0, 0), 2));
-		TImplicitObjectScaled<TSphere<FReal, 3>> Unscaled(Sphere, FVec3(1));
-		TImplicitObjectScaled<TSphere<FReal, 3>> UniformScaled(Sphere, FVec3(2));
-		TImplicitObjectScaled<TSphere<FReal, 3>> NonUniformScaled(Sphere, FVec3(2,1,1));
+		Chaos::FSphere A(FVec3(10, 0, 0), 5);
+		FSpherePtr Sphere( new Chaos::FSphere(FVec3(4, 0, 0), 2));
+		TImplicitObjectScaled<Chaos::FSphere> Unscaled(Sphere, FVec3(1));
+		TImplicitObjectScaled<Chaos::FSphere> UniformScaled(Sphere, FVec3(2));
+		TImplicitObjectScaled<Chaos::FSphere> NonUniformScaled(Sphere, FVec3(2,1,1));
 
 		FVec3 InitialDirs[] = { FVec3(1,0,0), FVec3(-1,0,0), FVec3(0,1,0), FVec3(0,-1,0), FVec3(0,0,1), FVec3(0,0,-1) };
 
@@ -676,8 +676,8 @@ namespace ChaosTest
 
 	void GJKSphereSphereSweep()
 	{
-		TSphere<FReal, 3> A(FVec3(10, 0, 0), 5);
-		TSphere<FReal, 3> B(FVec3(1, 0, 0), 2);
+		Chaos::FSphere A(FVec3(10, 0, 0), 5);
+		Chaos::FSphere B(FVec3(1, 0, 0), 2);
 
 		FVec3 InitialDirs[] = { FVec3(1,0,0), FVec3(-1,0,0), FVec3(0,1,0), FVec3(0,-1,0), FVec3(0,0,1), FVec3(0,0,-1) };
 
@@ -740,7 +740,7 @@ namespace ChaosTest
 			EXPECT_FALSE(GJKRaycast<FReal>(A, B, FRigidTransform3(FVec3(0, 0, 7 + 1e-2), FRotation3::Identity), FVec3(1, 0, 0), 20, Time, Position, Normal, 0, InitialDir));
 
 			//degenerate
-			TSphere<FReal, 3> Tiny(FVec3(1, 0, 0), 1e-8);
+			Chaos::FSphere Tiny(FVec3(1, 0, 0), 1e-8);
 			EXPECT_TRUE(GJKRaycast<FReal>(A, Tiny, FRigidTransform3::Identity, FVec3(1, 0, 0), 8, Time, Position, Normal, 0, InitialDir));
 			EXPECT_NEAR(Time, 4, Eps);
 			EXPECT_VECTOR_NEAR(Normal, FVec3(-1, 0, 0), Eps);
@@ -758,7 +758,7 @@ namespace ChaosTest
 	void GJKSphereBoxSweep()
 	{
 		FAABB3 A(FVec3(3, -1, 0), FVec3(4, 1, 4));
-		TSphere<FReal, 3> B(FVec3(0, 0, 0), 1);
+		Chaos::FSphere B(FVec3(0, 0, 0), 1);
 
 		FVec3 InitialDirs[] = { FVec3(1,0,0), FVec3(-1,0,0), FVec3(0,1,0), FVec3(0,-1,0), FVec3(0,0,1), FVec3(0,0,-1) };
 
@@ -766,9 +766,9 @@ namespace ChaosTest
 
 		for (const FVec3& InitialDir : InitialDirs)
 		{
-			FReal Time;
-			FVec3 Position;
-			FVec3 Normal;
+			FReal Time(0.0);
+			FVec3 Position(0.0);
+			FVec3 Normal(0.0);
 
 			//hit
 			EXPECT_TRUE(GJKRaycast<FReal>(A, B, FRigidTransform3::Identity, FVec3(1, 0, 0), 4, Time, Position, Normal, 0, InitialDir));
@@ -844,7 +844,7 @@ namespace ChaosTest
 
 	void GJKSphereCapsuleSweep()
 	{
-		TSphere<FReal, 3> A(FVec3(10, 0, 0), 5);
+		Chaos::FSphere A(FVec3(10, 0, 0), 5);
 		FCapsule B(FVec3(1, 0, 0), FVec3(-3, 0, 0), 2);
 
 		FVec3 InitialDirs[] = { FVec3(1,0,0), FVec3(-1,0,0), FVec3(0,1,0), FVec3(0,-1,0), FVec3(0,0,1), FVec3(0,0,-1) };
@@ -902,7 +902,7 @@ namespace ChaosTest
 			EXPECT_FALSE(GJKRaycast<FReal>(A, B, FRigidTransform3(FVec3(0, 0, 7 + 1e-2), FRotation3::Identity), FVec3(1, 0, 0), 20, Time, Position, Normal, 0, InitialDir));
 
 			//degenerate
-			TSphere<FReal, 3> Tiny(FVec3(1, 0, 0), 1e-8);
+			Chaos::FSphere Tiny(FVec3(1, 0, 0), 1e-8);
 			EXPECT_TRUE(GJKRaycast<FReal>(A, Tiny, FRigidTransform3::Identity, FVec3(1, 0, 0), 8, Time, Position, Normal, 0, InitialDir));
 			EXPECT_NEAR(Time, 4, Eps);
 			EXPECT_VECTOR_NEAR(Normal, FVec3(-1, 0, 0), Eps);
@@ -928,7 +928,7 @@ namespace ChaosTest
 		HullParticles[2] = { 3,-1,0 };
 		HullParticles[3] = { 4,0,2 };
 		FConvex A(HullParticles, 0.0f);
-		TSphere<FReal, 3> B(FVec3(0, 0, 0), 1);
+		Chaos::FSphere B(FVec3(0, 0, 0), 1);
 
 		FVec3 InitialDirs[] = { FVec3(1,0,0), FVec3(-1,0,0), FVec3(0,1,0), FVec3(0,-1,0), FVec3(0,0,1), FVec3(0,0,-1) };
 
@@ -937,7 +937,7 @@ namespace ChaosTest
 		for (const FVec3& InitialDir : InitialDirs)
 		{
 			FReal Time;
-			FVec3 Position;
+			FVec3 Position(0);
 			FVec3 Normal;
 
 			//hit
@@ -1006,11 +1006,11 @@ namespace ChaosTest
 
 	void GJKSphereScaledSphereSweep()
 	{
-		TSphere<FReal, 3> A(FVec3(10, 0, 0), 5);
-		FSpherePtr Sphere( new TSphere<FReal, 3>(FVec3(0, 0, 0), 2));
-		TImplicitObjectScaled<TSphere<FReal, 3>> Unscaled(Sphere, FVec3(1));
-		TImplicitObjectScaled<TSphere<FReal, 3>> UniformScaled(Sphere, FVec3(2));
-		TImplicitObjectScaled<TSphere<FReal, 3>> NonUniformScaled(Sphere, FVec3(2, 1, 1));
+		Chaos::FSphere A(FVec3(10, 0, 0), 5);
+		FSpherePtr Sphere( new Chaos::FSphere(FVec3(0, 0, 0), 2));
+		TImplicitObjectScaled<Chaos::FSphere> Unscaled(Sphere, FVec3(1));
+		TImplicitObjectScaled<Chaos::FSphere> UniformScaled(Sphere, FVec3(2));
+		TImplicitObjectScaled<Chaos::FSphere> NonUniformScaled(Sphere, FVec3(2, 1, 1));
 
 		FVec3 InitialDirs[] = { FVec3(1,0,0), FVec3(-1,0,0), FVec3(0,1,0), FVec3(0,-1,0), FVec3(0,0,1), FVec3(0,0,-1) };
 
@@ -1090,7 +1090,7 @@ namespace ChaosTest
 			EXPECT_FALSE(GJKRaycast<FReal>(A, Unscaled, FRigidTransform3(FVec3(0, 0, 7 + 1e-2), FRotation3::Identity), FVec3(1, 0, 0), 20, Time, Position, Normal, 0, InitialDir));
 
 			//degenerate
-			TSphere<FReal, 3> Tiny(FVec3(1, 0, 0), 1e-8);
+			Chaos::FSphere Tiny(FVec3(1, 0, 0), 1e-8);
 			EXPECT_TRUE(GJKRaycast<FReal>(A, Tiny, FRigidTransform3::Identity, FVec3(1, 0, 0), 8, Time, Position, Normal, 0, InitialDir));
 			EXPECT_NEAR(Time, 4, Eps);
 			EXPECT_VECTOR_NEAR(Normal, FVec3(-1, 0, 0), Eps);
@@ -1108,11 +1108,11 @@ namespace ChaosTest
 
 	void GJKSphereTransformedSphereSweep()
 	{
-		TSphere<FReal, 3> A(FVec3(10, 0, 0), 5);
+		Chaos::FSphere A(FVec3(10, 0, 0), 5);
 
-		TSphere<FReal, 3> Sphere(FVec3(0), 2);
-		TSphere<FReal, 3> Translated(Sphere.GetCenter() + FVec3(1, 0, 0), Sphere.GetRadius());
-		TSphere<FReal,3> Transformed(FRigidTransform3(FVec3(1, 0, 0), FRotation3::FromVector(FVec3(0, 0, PI))).TransformPosition(Sphere.GetCenter()), Sphere.GetRadius());
+		Chaos::FSphere Sphere(FVec3(0), 2);
+		Chaos::FSphere Translated(Sphere.GetCenterf() + FVec3f(1, 0, 0), Sphere.GetRadiusf());
+		Chaos::FSphere Transformed(FRigidTransform3(FVec3(1, 0, 0), FRotation3::FromVector(FVec3(0, 0, PI))).TransformPosition(FVec3(Sphere.GetCenterf())), Sphere.GetRadiusf());
 
 		FVec3 InitialDirs[] = { FVec3(1,0,0), FVec3(-1,0,0), FVec3(0,1,0), FVec3(0,-1,0), FVec3(0,0,1), FVec3(0,0,-1) };
 
@@ -1492,7 +1492,7 @@ namespace ChaosTest
 			int32 ClosestVertexIndexA, ClosestVertexIndexB;
 			EXPECT_TRUE((GJKPenetration<false, FReal>(A,B,BToATM, Penetration, ClosestA, ClosestB, Normal, ClosestVertexIndexA, ClosestVertexIndexB, 0.0, 0.0, FVec3(0,0,23.4092140))));
 			EXPECT_FLOAT_EQ(Normal.Z,0);
-			EXPECT_FLOAT_EQ(Penetration,A.GetRadius() + B.GetRadius());
+			EXPECT_FLOAT_EQ(Penetration,A.GetRadiusf() + B.GetRadiusf());
 		}
 
 		{
@@ -1551,7 +1551,7 @@ namespace ChaosTest
 			
 
 			const FCapsule A(Pt1, Pt2, 1.0);
-			const TSphere<FReal, 3> B(Pt0, 1.0);
+			const Chaos::FSphere B(Pt0, 1.0);
 
 			const FRigidTransform3 BToATM(FVec3(0, 0, 0), FRotation3::FromElements(0.0, 0.0, 0, 1)); // Unit transform
 			const FVec3 InitDir(0.1, 0.0, 0.0);
@@ -1889,7 +1889,7 @@ namespace ChaosTest
 		TArray<FConvex::FVec3Type> ConvexVertices(MoveTemp(ConvexVerts));
 		FConvexPtr CoreConvex( new FImplicitConvex3(ConvexVertices, 0.0f));
 		const TImplicitObjectScaled<FImplicitConvex3> ScaledConvex(CoreConvex, FVec3(-1,1,1), 38.4000015);
-		const TSphere<FReal, 3> Sphere(FVec3(0,0,0), 32);
+		const Chaos::FSphere Sphere(FVec3(0,0,0), 32);
 		const FRigidTransform3 StartTM(FVec3( -172.000000, -48.0000000, 52.0000000 ), FRotation3::FromIdentity());
 
 
@@ -1925,7 +1925,7 @@ namespace ChaosTest
 		};
 		TArray<FConvex::FVec3Type> ConvexVertices(MoveTemp(ConvexVerts));
 		FImplicitConvex3 CoreConvex = FImplicitConvex3(ConvexVertices, 38.4000015);
-		const TSphere<FReal, 3> Sphere(FVec3(0, 0, 0), 32);
+		const Chaos::FSphere Sphere(FVec3(0, 0, 0), 32);
 		const FRigidTransform3 StartTM(FVec3(-172.000000, -48.0000000, 52.0000000), FRotation3::FromIdentity());
 
 
@@ -2122,7 +2122,7 @@ namespace ChaosTest
 
 		FRealSingle Time;
 		VectorRegister4Float OutPositionSimd, OutNormalSimd;
-		Chaos::TSphere<Chaos::FReal, 3> QueryGeom({ 0, 0, 0 }, 5.0f);
+		Chaos::FSphere QueryGeom({ 0, 0, 0 }, 5.0f);
 		VectorRegister4Float RotationSimd = MakeVectorRegister(0.0f, 0.0f, 0.0f, 0.0f);
 
 		VectorRegister4Float TranslationSimd = MakeVectorRegister(100000.0f, 100022.820f, 213.002823f, 0.0f);
@@ -2247,7 +2247,7 @@ namespace ChaosTest
 
 	GTEST_TEST(GJKTests, GJK_SweepSphereLargeTriangleInMesh)
 	{
-		Chaos::TSphere<Chaos::FReal, 3> QueryGeom({ 0, 0, 0 }, 5.0f);
+		Chaos::FSphere QueryGeom({ 0, 0, 0 }, 5.0f);
 
 		FTriangleMeshImplicitObject::ParticlesType TrimeshParticles(
 			{
@@ -2333,7 +2333,7 @@ namespace ChaosTest
 
 	GTEST_TEST(GJKTests, GJK_SweepSphereBox)
 	{
-		Chaos::TSphere<Chaos::FReal, 3> QueryGeom({ 0, 0, 0 }, 12.0f);
+		Chaos::FSphere QueryGeom({ 0, 0, 0 }, 12.0f);
 
 		Chaos::TBox<Chaos::FReal, 3> Box(FVec3(-240.742279/2.0, -16.483643/2.0, -17.078735/2.0), FVec3(240.742279 / 2.0, 16.483643 / 2.0, 17.078735 / 2.0));
 
diff --git a/Engine/Source/Programs/HeadlessChaos/Private/HeadlessChaosTestGJKDistance.cpp b/Engine/Source/Programs/HeadlessChaos/Private/HeadlessChaosTestGJKDistance.cpp
--- a/Engine/Source/Programs/HeadlessChaos/Private/HeadlessChaosTestGJKDistance.cpp
+++ b/Engine/Source/Programs/HeadlessChaos/Private/HeadlessChaosTestGJKDistance.cpp
@@ -22,8 +22,8 @@ namespace ChaosTest
 		FVec3 Normal = {0, 0, 1};
 		// Fail - overlapping
 		{
-			TSphere<FReal, 3> A(FVec3(12, 0, 0), 5);
-			TSphere<FReal, 3> B(FVec3(4, 0, 0), 2);
+			Chaos::FSphere A(FVec3(12, 0, 0), 5);
+			Chaos::FSphere B(FVec3(4, 0, 0), 2);
 			const FRigidTransform3 BToATm = FRigidTransform3(FVec3(2, 0, 0), FRotation3::FromIdentity());
 			EGJKDistanceResult Result = GJKDistance<FReal>(
 				TGJKCoreShape(A),
@@ -35,8 +35,8 @@ namespace ChaosTest
 
 		// Success - not overlapping
 		{
-			TSphere<FReal, 3> A(FVec3(12, 0, 0), 5);
-			TSphere<FReal, 3> B(FVec3(4, 0, 0), 2);
+			Chaos::FSphere A(FVec3(12, 0, 0), 5);
+			Chaos::FSphere B(FVec3(4, 0, 0), 2);
 			const FRigidTransform3 BToATm = FRigidTransform3::Identity;
 			EGJKDistanceResult Result = GJKDistance<FReal>(
 				TGJKCoreShape(A), 
@@ -55,8 +55,8 @@ namespace ChaosTest
 
 		// Success - not overlapping
 		{
-			TSphere<FReal, 3> A(FVec3(0, 0, 0), 2);
-			TSphere<FReal, 3> B(FVec3(0, 0, 0), 2);
+			Chaos::FSphere A(FVec3(0, 0, 0), 2);
+			Chaos::FSphere B(FVec3(0, 0, 0), 2);
 			FVec3 BPos = FVec3(3, 3, 0);
 			const FRigidTransform3 BToATm = FRigidTransform3(BPos, FRotation3::FromIdentity());
 			EGJKDistanceResult Result = GJKDistance<FReal>(
@@ -65,21 +65,21 @@ namespace ChaosTest
 				GJKDistanceInitialVFromRelativeTransform(A, B, BToATm),
 				Distance, NearestA, NearestB, Normal);
 			EXPECT_EQ(Result, EGJKDistanceResult::Separated);
-			FVec3 CenterDelta = (B.GetCenter() + BPos) - A.GetCenter();
+			FVec3 CenterDelta = (B.GetCenterf() + BPos) - A.GetCenterf();
 			FVec3 CenterDir = CenterDelta.GetSafeNormal();
-			EXPECT_NEAR(Distance, CenterDelta.Size() - (A.GetRadius() + B.GetRadius()), Tolerance);
-			EXPECT_NEAR(NearestA.X, A.GetCenter().X + A.GetRadius() * CenterDir.X, Tolerance);
-			EXPECT_NEAR(NearestA.Y, A.GetCenter().Y + A.GetRadius() * CenterDir.Y, Tolerance);
-			EXPECT_NEAR(NearestA.Z, A.GetCenter().Z + A.GetRadius() * CenterDir.Z, Tolerance);
-			EXPECT_NEAR(NearestB.X, B.GetCenter().X - B.GetRadius() * CenterDir.X, Tolerance);
-			EXPECT_NEAR(NearestB.Y, B.GetCenter().Y - B.GetRadius() * CenterDir.Y, Tolerance);
-			EXPECT_NEAR(NearestB.Z, B.GetCenter().Z - B.GetRadius() * CenterDir.Z, Tolerance);
+			EXPECT_NEAR(Distance, CenterDelta.Size() - (A.GetRadiusf() + B.GetRadiusf()), Tolerance);
+			EXPECT_NEAR(NearestA.X, A.GetCenterf().X + A.GetRadiusf() * CenterDir.X, Tolerance);
+			EXPECT_NEAR(NearestA.Y, A.GetCenterf().Y + A.GetRadiusf() * CenterDir.Y, Tolerance);
+			EXPECT_NEAR(NearestA.Z, A.GetCenterf().Z + A.GetRadiusf() * CenterDir.Z, Tolerance);
+			EXPECT_NEAR(NearestB.X, B.GetCenterf().X - B.GetRadiusf() * CenterDir.X, Tolerance);
+			EXPECT_NEAR(NearestB.Y, B.GetCenterf().Y - B.GetRadiusf() * CenterDir.Y, Tolerance);
+			EXPECT_NEAR(NearestB.Z, B.GetCenterf().Z - B.GetRadiusf() * CenterDir.Z, Tolerance);
 		}
 
 		// Success - very close not overlapping
 		{
-			TSphere<FReal, 3> A(FVec3(12, 0, 0), 5);
-			TSphere<FReal, 3> B(FVec3(4, 0, 0), 2);
+			Chaos::FSphere A(FVec3(12, 0, 0), 5);
+			Chaos::FSphere B(FVec3(4, 0, 0), 2);
 			FVec3 BPos = FVec3(0.99, 0, 0);
 			const FRigidTransform3 BToATm = FRigidTransform3(BPos, FRotation3::FromIdentity());
 			EGJKDistanceResult Result = GJKDistance<FReal>(
@@ -115,7 +115,7 @@ namespace ChaosTest
 		// Fail - overlapping
 		{
 			FAABB3 A(FVec3(5, -2, -2), FVec3(8, 2, 2));
-			TSphere<FReal, 3> B(FVec3(2, 0, 0), 2);
+			Chaos::FSphere B(FVec3(2, 0, 0), 2);
 			const FRigidTransform3 BToATm = FRigidTransform3(FVec3(2, 0, 0), FRotation3::FromIdentity());
 			EGJKDistanceResult Result = GJKDistance<FReal>(
 				TGJKShape(A),
@@ -128,7 +128,7 @@ namespace ChaosTest
 		// Success - not overlapping - mid-face near point
 		{
 			FAABB3 A(FVec3(5, -2, -2), FVec3(8, 2, 2));
-			TSphere<FReal, 3> B(FVec3(2, 0, 0), 2);
+			Chaos::FSphere B(FVec3(2, 0, 0), 2);
 			const FRigidTransform3 BToATm = FRigidTransform3::Identity;
 			EGJKDistanceResult Result = GJKDistance<FReal>(
 				TGJKShape(A),
@@ -147,7 +147,7 @@ namespace ChaosTest
 		// Other way round
 		{
 			FAABB3 A(FVec3(5, -2, -2), FVec3(8, 2, 2));
-			TSphere<FReal, 3> B(FVec3(2, 0, 0), 2);
+			Chaos::FSphere B(FVec3(2, 0, 0), 2);
 			const FRigidTransform3 BToATm = FRigidTransform3::Identity;
 			EGJKDistanceResult Result = GJKDistance<FReal>(
 				TGJKShape(A),
@@ -167,7 +167,7 @@ namespace ChaosTest
 		// Success - not overlapping - vertex near point
 		{
 			FAABB3 A(FVec3(5, 2, 2), FVec3(8, 4, 4));
-			TSphere<FReal, 3> B(FVec3(2, 0, 0), 2);
+			Chaos::FSphere B(FVec3(2, 0, 0), 2);
 			const FRigidTransform3 BToATm = FRigidTransform3::Identity;
 			EGJKDistanceResult Result = GJKDistance<FReal>(
 				TGJKShape(A),
@@ -175,8 +175,8 @@ namespace ChaosTest
 				GJKDistanceInitialVFromRelativeTransform(A, B, BToATm),
 				Distance, NearestA, NearestB, Normal);
 			FVec3 NearPointOnA = A.Min();
-			FVec3 SphereNearPointDir = (NearPointOnA - B.GetCenter()).GetSafeNormal();
-			FVec3 NearPointOnB = B.GetCenter() + SphereNearPointDir * B.GetRadius();
+			FVec3 SphereNearPointDir = (NearPointOnA - B.GetCenterf()).GetSafeNormal();
+			FVec3 NearPointOnB = B.GetCenterf() + SphereNearPointDir * B.GetRadiusf();
 			EXPECT_EQ(Result, EGJKDistanceResult::Separated);
 			EXPECT_NEAR(Distance, (NearPointOnA - NearPointOnB).Size(), Tolerance);
 			EXPECT_NEAR(NearestA.X, NearPointOnA.X, Tolerance);
@@ -189,7 +189,7 @@ namespace ChaosTest
 		// Other way round
 		{
 			FAABB3 A(FVec3(5, 2, 2), FVec3(8, 4, 4));
-			TSphere<FReal, 3> B(FVec3(2, 0, 0), 2);
+			Chaos::FSphere B(FVec3(2, 0, 0), 2);
 			const FRigidTransform3 BToATm = FRigidTransform3::Identity;
 			EGJKDistanceResult Result = GJKDistance<FReal>(
 				TGJKShape(A),
@@ -197,8 +197,8 @@ namespace ChaosTest
 				GJKDistanceInitialVFromRelativeTransform(A, B, BToATm),
 				Distance, NearestA, NearestB, Normal);
 			FVec3 NearPointOnA = A.Min();
-			FVec3 SphereNearPointDir = (NearPointOnA - B.GetCenter()).GetSafeNormal();
-			FVec3 NearPointOnB = B.GetCenter() + SphereNearPointDir * B.GetRadius();
+			FVec3 SphereNearPointDir = (NearPointOnA - B.GetCenterf()).GetSafeNormal();
+			FVec3 NearPointOnB = B.GetCenterf() + SphereNearPointDir * B.GetRadiusf();
 			EXPECT_EQ(Result, EGJKDistanceResult::Separated);
 			EXPECT_NEAR(Distance, (NearPointOnA - NearPointOnB).Size(), Tolerance);
 			EXPECT_NEAR(NearestA.X, NearPointOnA.X, Tolerance);
@@ -212,7 +212,7 @@ namespace ChaosTest
 		// Rotated
 		{
 			FAABB3 A(FVec3(-2, -2, -2), FVec3(4, 4, 4));
-			TSphere<FReal, 3> B(FVec3(0, 0, 0), 2);
+			Chaos::FSphere B(FVec3(0, 0, 0), 2);
 			FRigidTransform3 BToATm = FRigidTransform3(FVec3(8, 0, 0), FRotation3::FromAxisAngle(FVec3(0, 1, 0), FMath::DegreesToRadians(45)));	// Rotation won't affect contact depth, but does affect local contact position
 			EGJKDistanceResult Result = GJKDistance<FReal>(
 				TGJKShape(A),
@@ -220,9 +220,9 @@ namespace ChaosTest
 				GJKDistanceInitialVFromRelativeTransform(A, B, BToATm),
 				Distance, NearestA, NearestB, Normal);
 			FVec3 NearPointOnA = FVec3(4, 0, 0);
-			FVec3 BPos = BToATm.TransformPositionNoScale(B.GetCenter());
+			FVec3 BPos = BToATm.TransformPositionNoScale(FVec3(B.GetCenterf()));
 			FVec3 NearPointDir = (NearPointOnA - BPos).GetSafeNormal();
-			FVec3 NearPointOnB = BPos + NearPointDir * B.GetRadius();
+			FVec3 NearPointOnB = BPos + NearPointDir * B.GetRadiusf();
 			FVec3 NearPointOnBLocal = BToATm.InverseTransformPositionNoScale(NearPointOnB);
 			EXPECT_EQ(Result, EGJKDistanceResult::Separated);
 			EXPECT_NEAR(Distance, (NearPointOnA - NearPointOnB).Size(), Tolerance);
@@ -236,7 +236,7 @@ namespace ChaosTest
 		// Other way round
 		{
 			FAABB3 A(FVec3(-2, -2, -2), FVec3(4, 4, 4));
-			TSphere<FReal, 3> B(FVec3(0, 0, 0), 2);
+			Chaos::FSphere B(FVec3(0, 0, 0), 2);
 			FRigidTransform3 BToATm = FRigidTransform3(FVec3(-8, 0, 0), FRotation3::FromAxisAngle(FVec3(0, 1, 0), FMath::DegreesToRadians(45)));
 			EGJKDistanceResult Result = GJKDistance<FReal>(
 				TGJKCoreShape(B),
@@ -244,9 +244,9 @@ namespace ChaosTest
 				GJKDistanceInitialVFromRelativeTransform(B, A, BToATm),
 				Distance, NearestB, NearestA, Normal);
 			FVec3 NearPointOnA = FVec3(4, 0, 4);
-			FVec3 BPos = BToATm.InverseTransformPositionNoScale(B.GetCenter());
+			FVec3 BPos = BToATm.InverseTransformPositionNoScale(FVec3(B.GetCenterf()));
 			FVec3 NearPointDir = (NearPointOnA - BPos).GetSafeNormal();
-			FVec3 NearPointOnB = BPos + NearPointDir * B.GetRadius();
+			FVec3 NearPointOnB = BPos + NearPointDir * B.GetRadiusf();
 			FVec3 NearPointOnBLocal = BToATm.TransformPositionNoScale(NearPointOnB);
 			EXPECT_EQ(Result, EGJKDistanceResult::Separated);
 			EXPECT_NEAR(Distance, (NearPointOnA - NearPointOnB).Size(), Tolerance);
@@ -261,7 +261,7 @@ namespace ChaosTest
 		// Success - specific test case that initially failed (using incorrect initialization of V which works for Overlap but not Distance)
 		{
 			FAABB3 A(FVec3(5, -2, 2), FVec3(8, 2, 4));
-			TSphere<FReal, 3> B(FVec3(2, 0, 0), 2);
+			Chaos::FSphere B(FVec3(2, 0, 0), 2);
 
 			bool bOverlap = GJKIntersection<FReal>(A, B, FRigidTransform3::Identity);
 			EXPECT_FALSE(bOverlap);
@@ -273,8 +273,8 @@ namespace ChaosTest
 				GJKDistanceInitialVFromRelativeTransform(A, B, BToATm),
 				Distance, NearestA, NearestB, Normal);
 			FVec3 NearPointOnA = FVec3(5, 0, 2);
-			FVec3 NearPointDir = (NearPointOnA - B.GetCenter()).GetSafeNormal();
-			FVec3 NearPointOnB = B.GetCenter() + NearPointDir * B.GetRadius();
+			FVec3 NearPointDir = (NearPointOnA - B.GetCenterf()).GetSafeNormal();
+			FVec3 NearPointOnB = B.GetCenterf() + NearPointDir * B.GetRadiusf();
 			EXPECT_EQ(Result, EGJKDistanceResult::Separated);
 			EXPECT_NEAR(Distance, (NearPointOnA - NearPointOnB).Size(), Tolerance);
 			EXPECT_NEAR(NearestA.X, NearPointOnA.X, Tolerance);
@@ -345,8 +345,8 @@ namespace ChaosTest
 				GJKDistanceInitialVFromRelativeTransform(A, B, BToATm),
 				Distance, NearestA, NearestB, Normal);
 			FVec3 ExpectedNearestA = FVec3(5, 0, -2);
-			FVec3 ExpectedDir = (ExpectedNearestA - B.GetX2()).GetSafeNormal();
-			FVec3 ExpectedNearestB = B.GetX2() + ExpectedDir * B.GetRadius();
+			FVec3 ExpectedDir = (ExpectedNearestA - B.GetX2f()).GetSafeNormal();
+			FVec3 ExpectedNearestB = B.GetX2f() + ExpectedDir * B.GetRadiusf();
 
 			const FReal Tolerance = (FReal)2e-3;
 			EXPECT_EQ(Result, EGJKDistanceResult::Separated);
@@ -439,8 +439,8 @@ namespace ChaosTest
 				GJKDistanceInitialVFromRelativeTransform(A, B, BToATm),
 				Distance, NearestA, NearestB, Normal);
 			FVec3 ExpectedNearestA = FVec3(5, 0, -2);
-			FVec3 ExpectedDir = (ExpectedNearestA - B.GetX2()).GetSafeNormal();
-			FVec3 ExpectedNearestB = B.GetX2() + ExpectedDir * B.GetRadius();
+			FVec3 ExpectedDir = (ExpectedNearestA - B.GetX2f()).GetSafeNormal();
+			FVec3 ExpectedNearestB = B.GetX2f() + ExpectedDir * B.GetRadiusf();
 
 			const FReal Tolerance = (FReal)0.3;
 			EXPECT_EQ(Result, EGJKDistanceResult::Separated);
diff --git a/Engine/Source/Programs/HeadlessChaos/Private/HeadlessChaosTestHeightfield.cpp b/Engine/Source/Programs/HeadlessChaos/Private/HeadlessChaosTestHeightfield.cpp
--- a/Engine/Source/Programs/HeadlessChaos/Private/HeadlessChaosTestHeightfield.cpp
+++ b/Engine/Source/Programs/HeadlessChaos/Private/HeadlessChaosTestHeightfield.cpp
@@ -630,7 +630,7 @@ namespace ChaosTest {
 		const FReal CountToWorldScale = 1;
 
 
-		TSphere<FReal, 3> Sphere(FVec3(0.0, 0.0, 0.0), 1.0);
+		Chaos::FSphere Sphere(FVec3(0.0, 0.0, 0.0), 1.0);
 		{
 
 			TArray<FReal> Heights;
@@ -718,7 +718,7 @@ namespace ChaosTest {
 		const FReal CountToWorldScale = 100;
 
 
-		TSphere<FReal, 3> Sphere(FVec3(0.0, 0.0, 0.0), 250.0);
+		Chaos::FSphere Sphere(FVec3(0.0, 0.0, 0.0), 250.0);
 
 		TArray<FReal> Heights;
 		Heights.AddZeroed(Rows * Columns);
@@ -1298,7 +1298,7 @@ namespace ChaosTest {
 			FHeightField Heightfield(MoveTemp(Heights), TArray<uint8>(), Rows, Columns, Scale);
 			const FAABB3 Bounds = Heightfield.BoundingBox();	//Current API forces us to do this to cache the bounds
 
-			const TSphere<FReal, 3> Sphere(FVec3::ZeroVector, SphereRadius);
+			const Chaos::FSphere Sphere(FVec3::ZeroVector, SphereRadius);
 			const FRigidTransform3 StartTM(SweepStart, TRotation<FReal, 3>::Identity);
 			
 			FReal ResultTOI;
@@ -1351,7 +1351,7 @@ namespace ChaosTest {
 		FHeightField Heightfield(MoveTemp(HeightsCopy), TArray<uint8>(), Rows, Columns, Scale);
 		const auto& Bounds = Heightfield.BoundingBox();	//Current API forces us to do this to cache the bounds
 
-		TSphere<FReal, 3> Sphere(FVec3(0.0, 0.0, 0.0), 50.0);
+		Chaos::FSphere Sphere(FVec3(0.0, 0.0, 0.0), 50.0);
 
 		FReal TOI, Phi;
 		FVec3 Position, Normal, FaceNormal;
@@ -1393,7 +1393,7 @@ namespace ChaosTest {
 		const auto& Bounds = Heightfield.BoundingBox();	//Current API forces us to do this to cache the bounds
 
 		// NOTE: Sphere center 50cm up
-		TSphere<FReal, 3> Sphere(FVec3(0.0, 0.0, 50.0), 50.0);
+		Chaos::FSphere Sphere(FVec3(0.0, 0.0, 50.0), 50.0);
 
 		FReal TOI, Phi;
 		FVec3 Position, Normal, FaceNormal;
diff --git a/Engine/Source/Programs/HeadlessChaos/Private/HeadlessChaosTestImplicits.cpp b/Engine/Source/Programs/HeadlessChaos/Private/HeadlessChaosTestImplicits.cpp
--- a/Engine/Source/Programs/HeadlessChaos/Private/HeadlessChaosTestImplicits.cpp
+++ b/Engine/Source/Programs/HeadlessChaos/Private/HeadlessChaosTestImplicits.cpp
@@ -539,16 +539,16 @@ namespace ChaosTest {
 		
 		FString Caller("ImplicitSphere()");
 
-		TSphere<FReal, 3> Subject(FVec3(0), 1);
+		Chaos::FSphere Subject(FVec3(0), 1);
 		UnitImplicitObjectNormalsInternal(Subject, Caller);
 		UnitImplicitObjectNormalsExternal(Subject, Caller);
 		UnitImplicitObjectIntersections(Subject, Caller);
 		UnitImplicitObjectSupportPhis<TSphere<FReal,3>>(Subject, Caller);
 
 		// intersection
-		EXPECT_TRUE(Subject.Intersects(TSphere<FReal, 3>(FVec3(0.f), 2.f)));
-		EXPECT_TRUE(Subject.Intersects(TSphere<FReal, 3>(FVec3(.5f), 1.f)));
-		EXPECT_FALSE(Subject.Intersects(TSphere<FReal, 3>(FVec3(2.f), 1.f)));
+		EXPECT_TRUE(Subject.Intersects(Chaos::FSphere(FVec3(0.f), 2.f)));
+		EXPECT_TRUE(Subject.Intersects(Chaos::FSphere(FVec3(.5f), 1.f)));
+		EXPECT_FALSE(Subject.Intersects(Chaos::FSphere(FVec3(2.f), 1.f)));
 
 		{// closest point near origin (+)
 			FVec3 InputPoint(0, 0, 2.);
@@ -564,7 +564,7 @@ namespace ChaosTest {
 		}
 
 		{// closest point off origin (+)
-			TSphere<FReal, 3> Subject2(FVec3(2), 2);
+			Chaos::FSphere Subject2(FVec3(2), 2);
 			FVec3 InputPoint(2, 2, 5);
 			TestFindClosestIntersection(Subject2, InputPoint, FVec3(2, 2, 4), Caller);
 			EXPECT_VECTOR_NEAR(Subject2.FindClosestPoint(InputPoint), FVector(2, 2, 4), 0.001);
@@ -944,7 +944,7 @@ namespace ChaosTest {
 		{	// Inertia tensor and rotation of mass
 			FCapsule AlignedSubject(FVec3(0, 0, 1), FVec3(0, 0, -1), 1);
 			FCapsule OffsetedAlignedSubject(FVec3(5, 10, 1), FVec3(5, 10, -1), 1);
-			FCapsule NonAlignedSubject(FVec3(-1, -1, -1).GetSafeNormal(), FVec3(1, 1, 1).GetSafeNormal(), 1);
+			FCapsule NonAlignedSubject(FVec3(-1, -1, -1).GetSafeNormal(), FVec3(1, 1, 1).GetSafeNormal(), 1.0);
 
 			UnitImplicitObjectInertiaTensorAndRotationOfMass(AlignedSubject, OffsetedAlignedSubject, NonAlignedSubject, Caller);
 		}
@@ -959,10 +959,10 @@ namespace ChaosTest {
 		UnitImplicitObjectNormalsExternal(UnitUnscaled, FString("ImplicitTransformed()"));
 		UnitImplicitObjectIntersections(UnitUnscaled, FString("ImplicitTransformed()"));
 
-		FSpherePtr Sphere( new TSphere<FReal, 3>(FVec3(3, 0, 0), 5));
-		TImplicitObjectScaled<TSphere<FReal, 3>> Unscaled(Sphere, FVec3(1));
-		TImplicitObjectScaled<TSphere<FReal, 3>> UniformScale(Sphere, FVec3(2));
-		TImplicitObjectScaled<TSphere<FReal, 3>> NonUniformScale(Sphere, FVec3(2, 1, 1));
+		FSpherePtr Sphere( new Chaos::FSphere(FVec3(3, 0, 0), 5));
+		TImplicitObjectScaled<Chaos::FSphere> Unscaled(Sphere, FVec3(1));
+		TImplicitObjectScaled<Chaos::FSphere> UniformScale(Sphere, FVec3(2));
+		TImplicitObjectScaled<Chaos::FSphere> NonUniformScale(Sphere, FVec3(2, 1, 1));
 
 		{//phi
 			const FVec3 NearEdge(7.5, 0, 0);
@@ -1478,7 +1478,7 @@ namespace ChaosTest {
 		{// Union of a union containing all the unit geometries overlapping - should still pass all the normal unit tests. 
 			TArray<Chaos::FImplicitObjectPtr> Objects1;
 			Objects1.Add(MakeImplicitObjectPtr<FCylinder>(FVec3(0, 0, 1), FVec3(0, 0, -1), 1));
-			Objects1.Add(MakeImplicitObjectPtr<TSphere<FReal, 3>>(FVec3(0, 0, 0), 1));
+			Objects1.Add(MakeImplicitObjectPtr<Chaos::FSphere>(FVec3(0, 0, 0), 1));
 			Objects1.Add(MakeImplicitObjectPtr<TBox<FReal, 3>>(FVec3(-1, -1, -1), FVec3(1, 1, 1)));
 			Objects1.Add(MakeImplicitObjectPtr<FTaperedCylinder>(FVec3(0, 0, 1), FVec3(0, 0, -1), 1, 1));
 
@@ -1495,7 +1495,7 @@ namespace ChaosTest {
 			TArray<Chaos::FImplicitObjectPtr> ObjectsA;
 			TArray<Chaos::FImplicitObjectPtr> ObjectsB;
 			ObjectsA.Add(MakeImplicitObjectPtr<FCylinder>(FVec3(0, 0, 1), FVec3(0, 0, -1), 1));
-			ObjectsA.Add(MakeImplicitObjectPtr<TSphere<FReal, 3>>(FVec3(0, 0, 0), 1));
+			ObjectsA.Add(MakeImplicitObjectPtr<Chaos::FSphere>(FVec3(0, 0, 0), 1));
 			ObjectsB.Add(MakeImplicitObjectPtr<TBox<FReal, 3>>(FVec3(-1, -1, -1), FVec3(1, 1, 1)));
 			ObjectsB.Add(MakeImplicitObjectPtr<FTaperedCylinder>(FVec3(0, 0, 1), FVec3(0, 0, -1), 1, 1));
 
@@ -1976,12 +1976,12 @@ namespace ChaosTest {
 		// ImplicitObjectScaled is ignored.
 		FReal Thickness = 0.1;
 		FSpherePtr Sphere( new TSphere<FReal,3>(FVec3(3, 0, 0), 5));
-		TImplicitObjectScaled<TSphere<FReal, 3>> Unscaled(Sphere, FVec3(1));
-		TImplicitObjectScaled<TSphere<FReal, 3>> UnscaledThickened(Sphere, FVec3(1), Thickness);
-		TImplicitObjectScaled<TSphere<FReal, 3>> UniformScale(Sphere, FVec3(2));
-		TImplicitObjectScaled<TSphere<FReal, 3>> UniformScaleThickened(Sphere, FVec3(2), Thickness);
-		TImplicitObjectScaled<TSphere<FReal, 3>> NonUniformScale(Sphere, FVec3(2, 1, 1));
-		TImplicitObjectScaled<TSphere<FReal, 3>> NonUniformScaleThickened(Sphere, FVec3(2, 1, 1), Thickness);
+		TImplicitObjectScaled<Chaos::FSphere> Unscaled(Sphere, FVec3(1));
+		TImplicitObjectScaled<Chaos::FSphere> UnscaledThickened(Sphere, FVec3(1), Thickness);
+		TImplicitObjectScaled<Chaos::FSphere> UniformScale(Sphere, FVec3(2));
+		TImplicitObjectScaled<Chaos::FSphere> UniformScaleThickened(Sphere, FVec3(2), Thickness);
+		TImplicitObjectScaled<Chaos::FSphere> NonUniformScale(Sphere, FVec3(2, 1, 1));
+		TImplicitObjectScaled<Chaos::FSphere> NonUniformScaleThickened(Sphere, FVec3(2, 1, 1), Thickness);
 
 		//phi
 		{
@@ -2113,9 +2113,9 @@ namespace ChaosTest {
 		FImplicitObjectUnionPtr MUnionedObjects(new Chaos::FImplicitObjectUnion(std::move(Objects)));
 
 		TArray<Chaos::FImplicitObjectPtr> Objects2;
-		Objects2.Add(MakeImplicitObjectPtr<TSphere<FReal, 3>>(FVec3(4, 0, 0), 1));
-		Objects2.Add(MakeImplicitObjectPtr<TSphere<FReal, 3>>(FVec3(5, 0, 0), 2));
-		Objects2.Add(MakeImplicitObjectPtr<TSphere<FReal, 3>>(FVec3(10, 0, 0), 3));
+		Objects2.Add(MakeImplicitObjectPtr<Chaos::FSphere>(FVec3(4, 0, 0), 1));
+		Objects2.Add(MakeImplicitObjectPtr<Chaos::FSphere>(FVec3(5, 0, 0), 2));
+		Objects2.Add(MakeImplicitObjectPtr<Chaos::FSphere>(FVec3(10, 0, 0), 3));
 
 		const FAABB3 OriginalBounds = MUnionedObjects->BoundingBox();
 
diff --git a/Engine/Source/Programs/HeadlessChaos/Private/HeadlessChaosTestOverlap.cpp b/Engine/Source/Programs/HeadlessChaos/Private/HeadlessChaosTestOverlap.cpp
--- a/Engine/Source/Programs/HeadlessChaos/Private/HeadlessChaosTestOverlap.cpp
+++ b/Engine/Source/Programs/HeadlessChaos/Private/HeadlessChaosTestOverlap.cpp
@@ -1205,7 +1205,7 @@ namespace ChaosTest
 			}
 			{
 				// Sphere test
-				const TSphere<FReal, 3> Sphere = TSphere<FReal, 3>({ 0.0, 0.0, 0.0 }, 2.0);
+				const Chaos::FSphere Sphere = Chaos::FSphere({ 0.0, 0.0, 0.0 }, 2.0);
 				{
 					const FRigidTransform3 QueryTM(FVec3(0.0, 0.0, 0.0), FQuat::Identity);
 					bool bResult = TriangleMesh->OverlapGeom(Sphere, QueryTM, 0.0);
diff --git a/Engine/Source/Programs/HeadlessChaos/Private/HeadlessChaosTestPerf.cpp b/Engine/Source/Programs/HeadlessChaos/Private/HeadlessChaosTestPerf.cpp
--- a/Engine/Source/Programs/HeadlessChaos/Private/HeadlessChaosTestPerf.cpp
+++ b/Engine/Source/Programs/HeadlessChaos/Private/HeadlessChaosTestPerf.cpp
@@ -76,7 +76,7 @@ namespace ChaosTest {
 		Chaos::FImplicitObjectPtr Box(new TBox<FReal, 3>(-HalfExtents, HalfExtents));
 
 		const FReal Radius = (Interval * 0.30);
-		Chaos::FImplicitObjectPtr Sphere(new TSphere<FReal, 3>(FVec3(0, 0, 0), Radius));
+		Chaos::FImplicitObjectPtr Sphere(new Chaos::FSphere(FVec3(0, 0, 0), Radius));
 
 		const FVec3 Offset(0.5, 0.5, 0.); // => if (0,0,0) produces a Nan in P() for the particle aligned with origin
 		int32 DynamicParticleIndex = 0;
diff --git a/Engine/Source/Programs/HeadlessChaos/Private/HeadlessChaosTestPhysicsScene.cpp b/Engine/Source/Programs/HeadlessChaos/Private/HeadlessChaosTestPhysicsScene.cpp
--- a/Engine/Source/Programs/HeadlessChaos/Private/HeadlessChaosTestPhysicsScene.cpp
+++ b/Engine/Source/Programs/HeadlessChaos/Private/HeadlessChaosTestPhysicsScene.cpp
@@ -139,7 +139,7 @@ namespace ChaosTest {
 		FChaosSQAccelerator SQAccelerator(*Scene.GetSpacialAcceleration());
 		FSQHitBuffer<ChaosInterface::FOverlapHit> HitBuffer;
 		FOverlapAllQueryCallback QueryCallback;
-		SQAccelerator.Overlap(TSphere<FReal, 3>(FVec3(0), Radius), InTM, HitBuffer, FChaosQueryFilterData(), QueryCallback, FQueryDebugParams());
+		SQAccelerator.Overlap(Chaos::FSphere(FVec3(0), Radius), InTM, HitBuffer, FChaosQueryFilterData(), QueryCallback, FQueryDebugParams());
 		return HitBuffer;
 	}
 
@@ -157,7 +157,7 @@ namespace ChaosTest {
 		EXPECT_NE(Proxy, nullptr);
 
 		{
-			auto Sphere = MakeImplicitObjectPtr<TSphere<FReal, 3>>(FVec3(0), 3);
+			auto Sphere = MakeImplicitObjectPtr<Chaos::FSphere>(FVec3(0), 3);
 			Proxy->GetGameThreadAPI().SetGeometry(Sphere);
 		}
 
@@ -179,7 +179,7 @@ namespace ChaosTest {
 		EXPECT_NE(Proxy, nullptr);
 
 		{
-			auto Sphere = MakeImplicitObjectPtr<TSphere<FReal, 3>>(FVec3(0), 3);
+			auto Sphere = MakeImplicitObjectPtr<Chaos::FSphere>(FVec3(0), 3);
 			Particle.SetGeometry(Sphere);
 		}
 
@@ -362,7 +362,7 @@ namespace ChaosTest {
 		auto& Particle = Proxy->GetGameThreadAPI();
 		EXPECT_NE(Proxy, nullptr);
 		{
-			auto Sphere = MakeImplicitObjectPtr<TSphere<FReal, 3>>(FVec3(0), 3);
+			auto Sphere = MakeImplicitObjectPtr<Chaos::FSphere>(FVec3(0), 3);
 			Particle.SetGeometry(Sphere);
 		}
 		FPhysicsActorHandle Proxy2 = nullptr;
@@ -370,7 +370,7 @@ namespace ChaosTest {
 		auto& Particle2 = Proxy2->GetGameThreadAPI();
 		EXPECT_NE(Proxy2, nullptr);
 		{
-			auto Sphere = MakeImplicitObjectPtr<TSphere<FReal, 3>>(FVec3(0), 3);
+			auto Sphere = MakeImplicitObjectPtr<Chaos::FSphere>(FVec3(0), 3);
 			Particle2.SetGeometry(Sphere);
 		}
 		TArray<FPhysicsActorHandle> Proxys = { Proxy, Proxy2 };
@@ -463,7 +463,7 @@ namespace ChaosTest {
 		FChaosEngineInterface::CreateActor(Params, Proxy);
 		auto& Particle = Proxy->GetGameThreadAPI();
 		{
-			auto Sphere = MakeImplicitObjectPtr<TSphere<FReal, 3>>(FVec3(0), 3);
+			auto Sphere = MakeImplicitObjectPtr<Chaos::FSphere>(FVec3(0), 3);
 			Particle.SetGeometry(Sphere);
 		}
 		TArray<FPhysicsActorHandle> Proxys = { Proxy };
@@ -583,7 +583,7 @@ namespace ChaosTest {
 		FChaosEngineInterface::CreateActor(Params, Proxy);
 		auto& Particle = Proxy->GetGameThreadAPI();
 		{
-			auto Sphere = MakeImplicitObjectPtr<TSphere<FReal, 3>>(FVec3(0), 3);
+			auto Sphere = MakeImplicitObjectPtr<Chaos::FSphere>(FVec3(0), 3);
 			Particle.SetGeometry(Sphere);
 		}
 		TArray<FPhysicsActorHandle> Proxys = { Proxy };
@@ -667,7 +667,7 @@ namespace ChaosTest {
 		EXPECT_NE(Proxy, nullptr);
 
 		{
-			auto Sphere = MakeImplicitObjectPtr<TSphere<FReal, 3>>(FVec3(0), 3);
+			auto Sphere = MakeImplicitObjectPtr<Chaos::FSphere>(FVec3(0), 3);
 			Particle.SetGeometry(Sphere);
 		}
 
@@ -707,7 +707,7 @@ namespace ChaosTest {
 		EXPECT_NE(Proxy, nullptr);
 
 		{
-			auto Sphere = MakeImplicitObjectPtr<TSphere<FReal, 3>>(FVec3(0), 3);
+			auto Sphere = MakeImplicitObjectPtr<Chaos::FSphere>(FVec3(0), 3);
 			Particle.SetGeometry(Sphere);
 		}
 
@@ -751,7 +751,7 @@ namespace ChaosTest {
 		EXPECT_NE(Proxy, nullptr);
 
 		{
-			auto Sphere = MakeImplicitObjectPtr<TSphere<FReal, 3>>(FVec3(0), 3);
+			auto Sphere = MakeImplicitObjectPtr<Chaos::FSphere>(FVec3(0), 3);
 			Particle.SetGeometry(Sphere);
 		}
 
@@ -794,7 +794,7 @@ namespace ChaosTest {
 		EXPECT_NE(Proxy, nullptr);
 
 		{
-			auto Sphere = MakeImplicitObjectPtr<TSphere<FReal, 3>>(FVec3(0), 3);
+			auto Sphere = MakeImplicitObjectPtr<Chaos::FSphere>(FVec3(0), 3);
 			Particle.SetGeometry(Sphere);
 		}
 
@@ -848,7 +848,7 @@ namespace ChaosTest {
 		EXPECT_NE(Proxy, nullptr);
 
 		{
-			auto Sphere = MakeImplicitObjectPtr<TSphere<FReal, 3>>(FVec3(0), 3);
+			auto Sphere = MakeImplicitObjectPtr<Chaos::FSphere>(FVec3(0), 3);
 			Particle.SetGeometry(Sphere);
 		}
 
@@ -884,7 +884,7 @@ namespace ChaosTest {
 			EXPECT_NE(Proxy, nullptr);
 
 			{
-				auto Sphere = MakeImplicitObjectPtr<TSphere<FReal, 3>>(FVec3(0), 3);
+				auto Sphere = MakeImplicitObjectPtr<Chaos::FSphere>(FVec3(0), 3);
 				Particle.SetGeometry(Sphere);
 			}
 
@@ -1016,7 +1016,7 @@ namespace ChaosTest {
 			EXPECT_NE(Proxy, nullptr);
 
 			{
-				auto Sphere = MakeImplicitObjectPtr<TSphere<FReal, 3>>(FVec3(0), 3);
+				auto Sphere = MakeImplicitObjectPtr<Chaos::FSphere>(FVec3(0), 3);
 				Particle.SetGeometry(Sphere);
 				Particle.SetV(FVec3(0, 0, -1));
 			}
@@ -1029,7 +1029,7 @@ namespace ChaosTest {
 			auto& Particle2 = Proxy2->GetGameThreadAPI();
 			EXPECT_NE(Proxy2, nullptr);
 			{
-				auto Sphere = MakeImplicitObjectPtr<TSphere<FReal, 3>>(FVec3(0), 3);
+				auto Sphere = MakeImplicitObjectPtr<Chaos::FSphere>(FVec3(0), 3);
 				Particle2.SetGeometry(Sphere);
 				Particle2.SetV(FVec3(0, -1, 0));
 			}
@@ -1111,7 +1111,7 @@ namespace ChaosTest {
 			EXPECT_NE(Proxy, nullptr);
 
 			{
-				auto Sphere = MakeImplicitObjectPtr<TSphere<FReal, 3>>(FVec3(0), 3);
+				auto Sphere = MakeImplicitObjectPtr<Chaos::FSphere>(FVec3(0), 3);
 				Particle.SetGeometry(Sphere);
 				Particle.SetV(FVec3(0, 0, -1));
 			}
@@ -1213,7 +1213,7 @@ namespace ChaosTest {
 		FChaosEngineInterface::CreateActor(Params, Proxy);
 		auto& Particle = Proxy->GetGameThreadAPI();
 		{
-			auto Sphere = MakeImplicitObjectPtr<TSphere<FReal, 3>>(FVec3(0), 3);
+			auto Sphere = MakeImplicitObjectPtr<Chaos::FSphere>(FVec3(0), 3);
 			Particle.SetGeometry(Sphere);
 		}
 
@@ -1258,7 +1258,7 @@ namespace ChaosTest {
 		FChaosEngineInterface::CreateActor(Params, Proxy);
 		auto& Particle = Proxy->GetGameThreadAPI();
 		{
-			auto Sphere = MakeImplicitObjectPtr<TSphere<FReal, 3>>(FVec3(0), 3);
+			auto Sphere = MakeImplicitObjectPtr<Chaos::FSphere>(FVec3(0), 3);
 			Particle.SetGeometry(Sphere);
 		}
 
@@ -1266,7 +1266,7 @@ namespace ChaosTest {
 		FChaosEngineInterface::CreateActor(Params, Proxy2);
 		auto& Particle2 = Proxy2->GetGameThreadAPI();
 		{
-			auto Sphere = MakeImplicitObjectPtr<TSphere<FReal, 3>>(FVec3(0), 3);
+			auto Sphere = MakeImplicitObjectPtr<Chaos::FSphere>(FVec3(0), 3);
 			Particle2.SetGeometry(Sphere);
 		}
 
@@ -1362,7 +1362,7 @@ namespace ChaosTest {
 		FChaosEngineInterface::CreateActor(Params, Proxy);
 		auto& Particle = Proxy->GetGameThreadAPI();
 		{
-			auto Sphere = MakeImplicitObjectPtr<TSphere<FReal, 3>>(FVec3(0), 3);
+			auto Sphere = MakeImplicitObjectPtr<Chaos::FSphere>(FVec3(0), 3);
 			Particle.SetGeometry(Sphere);
 		}
 
@@ -1523,7 +1523,7 @@ namespace ChaosTest {
 		FChaosEngineInterface::CreateActor(Params, Proxy);
 		auto& Particle = Proxy->GetGameThreadAPI();
 		{
-			auto Sphere = MakeImplicitObjectPtr<TSphere<FReal, 3>>(FVec3(0), 3);
+			auto Sphere = MakeImplicitObjectPtr<Chaos::FSphere>(FVec3(0), 3);
 			Particle.SetGeometry(Sphere);
 		}
 
@@ -1691,7 +1691,7 @@ namespace ChaosTest {
 		FChaosEngineInterface::CreateActor(Params, Proxy);
 		FRigidBodyHandle_External& Particle = Proxy->GetGameThreadAPI();
 		{
-			auto Sphere = MakeImplicitObjectPtr<TSphere<FReal, 3>>(FVec3(0), 3);
+			auto Sphere = MakeImplicitObjectPtr<Chaos::FSphere>(FVec3(0), 3);
 			Particle.SetGeometry(Sphere);
 		}
 
@@ -1753,7 +1753,7 @@ namespace ChaosTest {
 			FChaosEngineInterface::CreateActor(Params, Proxy);
 			auto& Particle = Proxy->GetGameThreadAPI();
 			{
-				auto Sphere = MakeImplicitObjectPtr<TSphere<FReal, 3>>(FVec3(0), 3);
+				auto Sphere = MakeImplicitObjectPtr<Chaos::FSphere>(FVec3(0), 3);
 				Particle.SetGeometry(Sphere);
 			}
 
@@ -1811,7 +1811,7 @@ namespace ChaosTest {
 		FChaosEngineInterface::CreateActor(Params, Proxy);
 		auto& Particle = Proxy->GetGameThreadAPI();
 		{
-			auto Sphere = MakeImplicitObjectPtr<TSphere<FReal, 3>>(FVec3(0), 3);
+			auto Sphere = MakeImplicitObjectPtr<Chaos::FSphere>(FVec3(0), 3);
 			Particle.SetGeometry(Sphere);
 		}
 
@@ -1882,7 +1882,7 @@ namespace ChaosTest {
 		FChaosEngineInterface::CreateActor(Params, Proxy);
 		auto& Particle = Proxy->GetGameThreadAPI();
 		{
-			auto Sphere = MakeImplicitObjectPtr<TSphere<FReal, 3>>(FVec3(0), 3);
+			auto Sphere = MakeImplicitObjectPtr<Chaos::FSphere>(FVec3(0), 3);
 			Particle.SetGeometry(Sphere);
 		}
 
@@ -1959,7 +1959,7 @@ namespace ChaosTest {
 		// was inside the geometry being used to query.
 		const FTransform QueryTM{ FVec3{0.0f, 0.0f, -110.0f} };
 		constexpr FReal SphereRadius = static_cast<FReal>(50.0);
-		SQ.Overlap(TSphere<FReal, 3>(FVec3(0.0f, 0.0f, 100.0f), SphereRadius), QueryTM, HitBuffer, FChaosQueryFilterData(), QueryCallback, FQueryDebugParams());
+		SQ.Overlap(Chaos::FSphere(FVec3(0.0f, 0.0f, 100.0f), SphereRadius), QueryTM, HitBuffer, FChaosQueryFilterData(), QueryCallback, FQueryDebugParams());
 
 		EXPECT_TRUE(HitBuffer.HasBlockingHit());
 	}
@@ -2051,7 +2051,7 @@ namespace ChaosTest {
 		FBlockAllQueryCallback QueryCallback;
 
 		const FReal SphereRadius = 25;
-		const TSphere<FReal, 3> QuerySphere(FVec3::ZeroVector, SphereRadius);
+		const Chaos::FSphere QuerySphere(FVec3::ZeroVector, SphereRadius);
 
 		// Sweep down over both shapes with an initial overlap, however center over shape 0 more.
 		FTransform QueryTM{ FVec3{-5, 0, 50} };
@@ -2116,7 +2116,7 @@ namespace ChaosTest {
 		FBlockAllQueryCallback QueryCallback;
 
 		const FReal SphereRadius = 25;
-		const TSphere<FReal, 3> QuerySphere(FVec3::ZeroVector, SphereRadius);
+		const Chaos::FSphere QuerySphere(FVec3::ZeroVector, SphereRadius);
 
 		// Sweep down over both shapes with an initial overlap, however center over shape 0 more.
 		FTransform QueryTM{ FVec3{-5, 0, 50} };
diff --git a/Engine/Source/Programs/HeadlessChaos/Private/HeadlessChaosTestRaycast.cpp b/Engine/Source/Programs/HeadlessChaos/Private/HeadlessChaosTestRaycast.cpp
--- a/Engine/Source/Programs/HeadlessChaos/Private/HeadlessChaosTestRaycast.cpp
+++ b/Engine/Source/Programs/HeadlessChaos/Private/HeadlessChaosTestRaycast.cpp
@@ -574,12 +574,12 @@ namespace ChaosTest
 		const FReal Thickness = 0.1;
 
 		FSpherePtr Sphere( new TSphere<FReal,3>(FVec3(1), 2));
-		TImplicitObjectScaled<TSphere<FReal, 3>> Unscaled(Sphere, FVec3(1));
-		TImplicitObjectScaled<TSphere<FReal, 3>> UnscaledThickened(Sphere, FVec3(1), Thickness);
-		TImplicitObjectScaled<TSphere<FReal, 3>> UniformScaled(Sphere, FVec3(2));
-		TImplicitObjectScaled<TSphere<FReal, 3>> UniformScaledThickened(Sphere, FVec3(2), Thickness);
-		TImplicitObjectScaled<TSphere<FReal, 3>> NonUniformScaled(Sphere, FVec3(2,1,1));
-		TImplicitObjectScaled<TSphere<FReal, 3>> NonUniformScaledThickened(Sphere, FVec3(2, 1, 1), Thickness);
+		TImplicitObjectScaled<Chaos::FSphere> Unscaled(Sphere, FVec3(1));
+		TImplicitObjectScaled<Chaos::FSphere> UnscaledThickened(Sphere, FVec3(1), Thickness);
+		TImplicitObjectScaled<Chaos::FSphere> UniformScaled(Sphere, FVec3(2));
+		TImplicitObjectScaled<Chaos::FSphere> UniformScaledThickened(Sphere, FVec3(2), Thickness);
+		TImplicitObjectScaled<Chaos::FSphere> NonUniformScaled(Sphere, FVec3(2,1,1));
+		TImplicitObjectScaled<Chaos::FSphere> NonUniformScaledThickened(Sphere, FVec3(2, 1, 1), Thickness);
 
 		//simple
 		bool bHit = Unscaled.Raycast(FVec3(1, 1, 8), FVec3(0, 0, -1), 8, 0, Time, Position, Normal, FaceIndex);
diff --git a/Engine/Source/Programs/HeadlessChaos/Private/HeadlessChaosTestRewind.cpp b/Engine/Source/Programs/HeadlessChaos/Private/HeadlessChaosTestRewind.cpp
--- a/Engine/Source/Programs/HeadlessChaos/Private/HeadlessChaosTestRewind.cpp
+++ b/Engine/Source/Programs/HeadlessChaos/Private/HeadlessChaosTestRewind.cpp
@@ -78,7 +78,7 @@ namespace ChaosTest {
 		{
 			TestEmpty([&Lambda, RewindHistorySize](auto* Solver, FReal SimDt, int32 Optimization)
 			{
-				auto Sphere = Chaos::FImplicitObjectPtr(new TSphere<FReal, 3>(FVec3(0), 10));
+					Chaos::FImplicitObjectPtr Sphere = Chaos::FImplicitObjectPtr(new Chaos::FSphere(FVec3(0), 10));
 
 				// Make particles
 					auto Proxy = FSingleParticlePhysicsProxy::Create(Chaos::FPBDRigidParticle::CreateParticle());
@@ -97,7 +97,7 @@ namespace ChaosTest {
 	{
 		TRewindHelper::TestEmpty([](auto* Solver, FReal SimDt, int32 Optimization)
 		{
-			auto Sphere = Chaos::FImplicitObjectPtr(new TSphere<FReal, 3>(FVec3(0), 10));
+			auto Sphere = Chaos::FImplicitObjectPtr(new Chaos::FSphere(FVec3(0), 10));
 			auto Box = Chaos::FImplicitObjectPtr(new TBox<FReal, 3>(FVec3(0), FVec3(1)));
 			auto Box2 = Chaos::FImplicitObjectPtr(new TBox<FReal, 3>(FVec3(2), FVec3(3)));
 
@@ -316,7 +316,7 @@ namespace ChaosTest {
 				Particle.SetX(FVec3(0, 0, 100));
 				Particle.SetV(FVec3(0, 0, -10));
 
-				auto SphereGeom = Chaos::FImplicitObjectPtr(new TSphere<FReal, 3>(FVec3(0), 10));
+				auto SphereGeom = Chaos::FImplicitObjectPtr(new Chaos::FSphere(FVec3(0), 10));
 				FSingleParticlePhysicsProxy* SecondSphere = nullptr;
 
 				auto FloorGeom = Chaos::FImplicitObjectPtr(new TBox<FReal, 3>(FVec3(-100, -100, 0), FVec3(100, 100, 30)));
@@ -382,7 +382,7 @@ namespace ChaosTest {
 				Particle.SetX(FVec3(0, 0, 100));
 				Particle.SetV(FVec3(0, 0, -10));
 
-				auto SphereGeom = Chaos::FImplicitObjectPtr(new TSphere<FReal, 3>(FVec3(0), 10));
+				auto SphereGeom = Chaos::FImplicitObjectPtr(new Chaos::FSphere(FVec3(0), 10));
 				FSingleParticlePhysicsProxy* SecondSphere = nullptr;
 
 				// Make particles
@@ -642,7 +642,7 @@ namespace ChaosTest {
 
 				const int32 LastGameStep = 32;
 				
-				auto SphereGeom = Chaos::FImplicitObjectPtr(new TSphere<FReal, 3>(FVec3(0), 10));
+				auto SphereGeom = Chaos::FImplicitObjectPtr(new Chaos::FSphere(FVec3(0), 10));
 
 				auto& Particle = Proxy->GetGameThreadAPI();
 				Particle.SetGravityEnabled(false);
@@ -1013,7 +1013,7 @@ namespace ChaosTest {
 					int32 ResimEndFrame = 10;
 				};
 
-				auto Sphere = Chaos::FImplicitObjectPtr(new TSphere<FReal, 3>(FVec3(0), 10));
+				auto Sphere = Chaos::FImplicitObjectPtr(new Chaos::FSphere(FVec3(0), 10));
 				auto Proxy = FSingleParticlePhysicsProxy::Create(Chaos::FPBDRigidParticle::CreateParticle());
 				Proxy->GetGameThreadAPI().SetGeometry(Sphere);
 				Solver->RegisterObject(Proxy);
@@ -1098,7 +1098,7 @@ namespace ChaosTest {
 
 				};
 
-				auto Sphere = Chaos::FImplicitObjectPtr(new TSphere<FReal, 3>(FVec3(0), 10));
+				auto Sphere = Chaos::FImplicitObjectPtr(new Chaos::FSphere(FVec3(0), 10));
 				auto Proxy = FSingleParticlePhysicsProxy::Create(Chaos::FPBDRigidParticle::CreateParticle());
 				Proxy->GetGameThreadAPI().SetGeometry(Sphere);
 				Solver->RegisterObject(Proxy);
@@ -2343,7 +2343,7 @@ namespace ChaosTest {
 	{
 		for (int Optimization = 0; Optimization < 2; ++Optimization)
 		{
-			auto Sphere = Chaos::FImplicitObjectPtr(new TSphere<FReal, 3>(FVec3(0), 10));
+			auto Sphere = Chaos::FImplicitObjectPtr(new Chaos::FSphere(FVec3(0), 10));
 
 			FChaosSolversModule* Module = FChaosSolversModule::GetModule();
 
@@ -2408,7 +2408,7 @@ namespace ChaosTest {
 	{
 		for (int Optimization = 0; Optimization < 2; ++Optimization)
 		{
-			auto Sphere = Chaos::FImplicitObjectPtr(new TSphere<FReal, 3>(FVec3(0), 10));
+			auto Sphere = Chaos::FImplicitObjectPtr(new Chaos::FSphere(FVec3(0), 10));
 
 			FChaosSolversModule* Module = FChaosSolversModule::GetModule();
 
@@ -2562,7 +2562,7 @@ namespace ChaosTest {
 	{
 		for (int Optimization = 0; Optimization < 2; ++Optimization)
 		{
-			auto Sphere = Chaos::FImplicitObjectPtr(new TSphere<FReal, 3>(FVec3(0), 10));
+			auto Sphere = Chaos::FImplicitObjectPtr(new Chaos::FSphere(FVec3(0), 10));
 
 			FChaosSolversModule* Module = FChaosSolversModule::GetModule();
 
@@ -2659,7 +2659,7 @@ namespace ChaosTest {
 	{
 		for (int Optimization = 0; Optimization < 2; ++Optimization)
 		{
-			auto Sphere = Chaos::FImplicitObjectPtr(new TSphere<FReal, 3>(FVec3(0), 10));
+			auto Sphere = Chaos::FImplicitObjectPtr(new Chaos::FSphere(FVec3(0), 10));
 
 			FChaosSolversModule* Module = FChaosSolversModule::GetModule();
 
@@ -2753,7 +2753,7 @@ namespace ChaosTest {
 			//We want to detect when sim results change
 			//Detecting output of position and velocity is expensive and hard to track
 			//Instead we need to rely on fast forward mechanism, this is still in progress
-			auto Sphere = Chaos::FImplicitObjectPtr(new TSphere<FReal, 3>(FVec3(0), 10));
+			auto Sphere = Chaos::FImplicitObjectPtr(new Chaos::FSphere(FVec3(0), 10));
 			auto Box = Chaos::FImplicitObjectPtr(new TBox<FReal, 3>(FVec3(-100, -100, -100), FVec3(100, 100, 0)));
 
 			FChaosSolversModule* Module = FChaosSolversModule::GetModule();
@@ -2852,7 +2852,7 @@ namespace ChaosTest {
 	{
 		for (int Optimization = 0; Optimization < 2; ++Optimization)
 		{
-			auto Sphere = Chaos::FImplicitObjectPtr(new TSphere<FReal, 3>(FVec3(0), 10));
+			auto Sphere = Chaos::FImplicitObjectPtr(new Chaos::FSphere(FVec3(0), 10));
 
 			FChaosSolversModule* Module = FChaosSolversModule::GetModule();
 
@@ -2945,7 +2945,7 @@ namespace ChaosTest {
 #if REWIND_DESYNC
 		for (int Optimization = 0; Optimization < 2; ++Optimization)
 		{
-			auto Sphere = Chaos::FImplicitObjectPtr(new TSphere<FReal, 3>(FVec3(0), 10));
+			auto Sphere = Chaos::FImplicitObjectPtr(new Chaos::FSphere(FVec3(0), 10));
 			auto Box = Chaos::FImplicitObjectPtr(new TBox<FReal, 3>(FVec3(-100, -100, -100), FVec3(100, 100, 0)));
 
 			FChaosSolversModule* Module = FChaosSolversModule::GetModule();
@@ -3036,7 +3036,7 @@ namespace ChaosTest {
 	{
 		for (int Optimization = 0; Optimization < 2; ++Optimization)
 		{
-			auto Sphere = Chaos::FImplicitObjectPtr(new TSphere<FReal, 3>(FVec3(0), 10));
+			auto Sphere = Chaos::FImplicitObjectPtr(new Chaos::FSphere(FVec3(0), 10));
 			auto Box = Chaos::FImplicitObjectPtr(new TBox<FReal, 3>(FVec3(-100, -100, -100), FVec3(100, 100, 0)));
 
 			FChaosSolversModule* Module = FChaosSolversModule::GetModule();
@@ -3130,7 +3130,7 @@ namespace ChaosTest {
 	{
 		for (int Optimization = 0; Optimization < 2; ++Optimization)
 		{
-			auto Sphere = Chaos::FImplicitObjectPtr(new TSphere<FReal, 3>(FVec3(0), 10));
+			auto Sphere = Chaos::FImplicitObjectPtr(new Chaos::FSphere(FVec3(0), 10));
 			auto Box = Chaos::FImplicitObjectPtr(new TBox<FReal, 3>(FVec3(-100, -100, -100), FVec3(100, 100, 0)));
 
 			FChaosSolversModule* Module = FChaosSolversModule::GetModule();
@@ -3481,7 +3481,7 @@ namespace ChaosTest {
 	{
 		for (int Optimization = 0; Optimization < 2; ++Optimization)
 		{
-			auto Sphere = Chaos::FImplicitObjectPtr(new TSphere<FReal, 3>(FVec3(0), 10));
+			auto Sphere = Chaos::FImplicitObjectPtr(new Chaos::FSphere(FVec3(0), 10));
 			auto Box = Chaos::FImplicitObjectPtr(new TBox<FReal, 3>(FVec3(-100, -100, -100), FVec3(100, 100, 0)));
 
 			FChaosSolversModule* Module = FChaosSolversModule::GetModule();
@@ -3590,7 +3590,7 @@ namespace ChaosTest {
 
 	GTEST_TEST(AllTraits, DISABLED_RewindTest_SoftDesyncFromSameIsland)
 	{
-		auto Sphere = Chaos::FImplicitObjectPtr(new TSphere<FReal, 3>(FVec3(0), 10));
+		auto Sphere = Chaos::FImplicitObjectPtr(new Chaos::FSphere(FVec3(0), 10));
 		auto Box = Chaos::FImplicitObjectPtr(new TBox<FReal, 3>(FVec3(-100, -100, -100), FVec3(100, 100, 0)));
 
 		FChaosSolversModule* Module = FChaosSolversModule::GetModule();
@@ -3703,7 +3703,7 @@ namespace ChaosTest {
 
 	GTEST_TEST(AllTraits, DISABLED_RewindTest_SoftDesyncFromSameIslandThenBackToInSync)
 	{
-		auto Sphere = Chaos::FImplicitObjectPtr(new TSphere<FReal, 3>(FVec3(0), 10));
+		auto Sphere = Chaos::FImplicitObjectPtr(new Chaos::FSphere(FVec3(0), 10));
 		auto Box = Chaos::FImplicitObjectPtr(new TBox<FReal, 3>(FVec3(-100, -100, -10), FVec3(100, 100, 0)));
 
 		FChaosSolversModule* Module = FChaosSolversModule::GetModule();
@@ -4146,7 +4146,7 @@ namespace ChaosTest {
 
 	GTEST_TEST(AllTraits, DeterministicSim_DoubleTickCollide)
 	{
-		auto Sphere = Chaos::FImplicitObjectPtr(new TSphere<FReal, 3>(FVec3(0), 50));
+		auto Sphere = Chaos::FImplicitObjectPtr(new Chaos::FSphere(FVec3(0), 50));
 
 		const auto InitLambda = [&Sphere](auto& Solver, auto)
 		{
@@ -4344,7 +4344,7 @@ namespace ChaosTest {
 			auto Proxy2 = FSingleParticlePhysicsProxy::Create(Chaos::FPBDRigidParticle::CreateParticle());
 			auto& Particle2 = Proxy2->GetGameThreadAPI();
 
-			auto Sphere2 = Chaos::FImplicitObjectPtr(new TSphere<FReal, 3>(FVec3(0), 10));
+			auto Sphere2 = Chaos::FImplicitObjectPtr(new Chaos::FSphere(FVec3(0), 10));
 			Particle2.SetGeometry(Sphere2);
 			Particle2.SetV(FVec3(0, 0, 1));
 			Particle2.SetGravityEnabled(false);
diff --git a/Engine/Source/Programs/HeadlessChaos/Private/HeadlessChaosTestSerialization.cpp b/Engine/Source/Programs/HeadlessChaos/Private/HeadlessChaosTestSerialization.cpp
--- a/Engine/Source/Programs/HeadlessChaos/Private/HeadlessChaosTestSerialization.cpp
+++ b/Engine/Source/Programs/HeadlessChaos/Private/HeadlessChaosTestSerialization.cpp
@@ -80,9 +80,9 @@ namespace ChaosTest
 	{
 
 		TArray<FSpherePtr> OriginalSpheres;
-		OriginalSpheres.Add(FSpherePtr(new TSphere<FReal, 3>(FVec3(), 1)));
-		OriginalSpheres.Add(FSpherePtr(new TSphere<FReal, 3>(FVec3(), 2)));
-		OriginalSpheres.Add(FSpherePtr(new TSphere<FReal, 3>(FVec3(), 3)));
+		OriginalSpheres.Add(FSpherePtr(new Chaos::FSphere(FVec3(), 1)));
+		OriginalSpheres.Add(FSpherePtr(new Chaos::FSphere(FVec3(), 2)));
+		OriginalSpheres.Add(FSpherePtr(new Chaos::FSphere(FVec3(), 3)));
 
 		TArray<uint8> Data;
 		{
@@ -103,7 +103,7 @@ namespace ChaosTest
 
 			for (int32 Idx = 0; Idx < SerializedSpheres.Num(); ++Idx)
 			{
-				EXPECT_TRUE(SerializedSpheres[Idx]->GetRadius() == OriginalSpheres[Idx]->GetRadius());
+				EXPECT_TRUE(SerializedSpheres[Idx]->GetRadiusf() == OriginalSpheres[Idx]->GetRadiusf());
 			}
 		}
 	}
@@ -111,10 +111,10 @@ namespace ChaosTest
 	void SharedObjectsSerialization()
 	{
 		TArray<FSpherePtr> OriginalSpheres;
-		FSpherePtr Sphere(new TSphere<FReal, 3>(FVec3(0), 1));
+		FSpherePtr Sphere(new Chaos::FSphere(FVec3(0), 1));
 		OriginalSpheres.Add(Sphere);
 		OriginalSpheres.Add(Sphere);
-		TSerializablePtr<TSphere<FReal, 3>> SerializableSphere(Sphere);
+		TSerializablePtr<Chaos::FSphere> SerializableSphere(Sphere);
 
 		TArray<uint8> Data;
 		{
@@ -127,7 +127,7 @@ namespace ChaosTest
 		
 		{
 			TArray<FSpherePtr> SerializedSpheres;
-			TSerializablePtr<TSphere<FReal, 3>> SerializedSphere;
+			TSerializablePtr<Chaos::FSphere> SerializedSphere;
 			{
 				FMemoryReader Ar(Data);
 				FChaosArchive Reader(Ar);
@@ -140,7 +140,7 @@ namespace ChaosTest
 
 				for (int32 Idx = 0; Idx < SerializedSpheres.Num(); ++Idx)
 				{
-					EXPECT_TRUE(SerializedSpheres[Idx]->GetRadius() == OriginalSpheres[Idx]->GetRadius());
+					EXPECT_TRUE(SerializedSpheres[Idx]->GetRadiusf() == OriginalSpheres[Idx]->GetRadiusf());
 				}
 
 				EXPECT_EQ(SerializedSpheres[0].GetReference(), SerializedSpheres[1].GetReference());
@@ -153,8 +153,8 @@ namespace ChaosTest
 	void GraphSerialization()
 	{
 		TArray<Chaos::FImplicitObjectPtr> OriginalSpheres;
-		OriginalSpheres.Emplace(new TSphere<FReal, 3>{ FVec3(1,2,3), 1 });
-		OriginalSpheres.Emplace(new TSphere<FReal, 3>{ FVec3(1,2,3), 2 });
+		OriginalSpheres.Emplace(new Chaos::FSphere{ FVec3(1,2,3), 1 });
+		OriginalSpheres.Emplace(new Chaos::FSphere{ FVec3(1,2,3), 2 });
 
 		TArray<Chaos::FImplicitObjectPtr> OriginalChildren;
 		OriginalChildren.Emplace(new TImplicitObjectTransformed<FReal, 3>(OriginalSpheres[0], FRigidTransform3::Identity));
@@ -198,8 +198,8 @@ namespace ChaosTest
 	void ObjectUnionSerialization()
 	{
 		TArray<Chaos::FImplicitObjectPtr> OriginalSpheres;
-		OriginalSpheres.Emplace(new TSphere<FReal, 3>(FVec3(1, 2, 3), 1));
-		OriginalSpheres.Emplace(new TSphere<FReal, 3>(FVec3(1, 2, 3), 2));
+		OriginalSpheres.Emplace(new Chaos::FSphere(FVec3(1, 2, 3), 1));
+		OriginalSpheres.Emplace(new Chaos::FSphere(FVec3(1, 2, 3), 2));
 
 		TArray<Chaos::FImplicitObjectPtr> OriginalChildren;
 		OriginalChildren.Emplace(new TImplicitObjectTransformed<FReal, 3>(OriginalSpheres[0], FRigidTransform3::Identity));
@@ -249,8 +249,8 @@ namespace ChaosTest
 	void ParticleSerialization()
 	{
 		TArray<Chaos::FImplicitObjectPtr> OriginalSpheres;
-		OriginalSpheres.Emplace(new TSphere<FReal, 3>(FVec3(1, 2, 3), 1));
-		OriginalSpheres.Emplace(new TSphere<FReal, 3>(FVec3(1, 2, 3), 2));
+		OriginalSpheres.Emplace(new Chaos::FSphere(FVec3(1, 2, 3), 1));
+		OriginalSpheres.Emplace(new Chaos::FSphere(FVec3(1, 2, 3), 2));
 
 		{
 			FGeometryParticles OriginalParticles;
@@ -329,8 +329,8 @@ namespace ChaosTest
 		TArray<uint8> Data;
 		{
 			TArray<FImplicitObjectPtr> OriginalSpheres;
-			OriginalSpheres.Emplace(new TSphere<FReal, 3>(FVec3(0, 0, 0), 1));
-			OriginalSpheres.Emplace(new TSphere<FReal, 3>(FVec3(0, 0, 0), 2));
+			OriginalSpheres.Emplace(new Chaos::FSphere(FVec3(0, 0, 0), 1));
+			OriginalSpheres.Emplace(new Chaos::FSphere(FVec3(0, 0, 0), 2));
 
 			FGeometryParticles OriginalParticles;
 			OriginalParticles.AddParticles(2);
@@ -424,9 +424,9 @@ namespace ChaosTest
 	{
 		TArray<uint8> Data;
 		TArray<FImplicitObjectPtr> Spheres;
-		Spheres.Emplace(new TSphere<FReal, 3>(FVec3(0, 0, 0), 1));
-		Spheres.Emplace(new TSphere<FReal, 3>(FVec3(0, 0, 0), 1));
-		Spheres.Emplace(new TSphere<FReal, 3>(FVec3(0, 0, 0), 1));
+		Spheres.Emplace(new Chaos::FSphere(FVec3(0, 0, 0), 1));
+		Spheres.Emplace(new Chaos::FSphere(FVec3(0, 0, 0), 1));
+		Spheres.Emplace(new Chaos::FSphere(FVec3(0, 0, 0), 1));
 
 		FGeometryParticles Particles;
 		Particles.AddParticles(3);
diff --git a/Engine/Source/Programs/HeadlessChaos/Private/HeadlessChaosTestShapePadding.cpp b/Engine/Source/Programs/HeadlessChaos/Private/HeadlessChaosTestShapePadding.cpp
--- a/Engine/Source/Programs/HeadlessChaos/Private/HeadlessChaosTestShapePadding.cpp
+++ b/Engine/Source/Programs/HeadlessChaos/Private/HeadlessChaosTestShapePadding.cpp
@@ -73,8 +73,8 @@ namespace ChaosTest {
 		// Boxes should have a margin, limited by the minimum box dimension
 		float ExpectedMargin0 = BoxImplicit0->ClampedMargin(Margin0);
 		float ExpectedMargin1 = BoxImplicit0->ClampedMargin(Margin1);
-		EXPECT_NEAR(BoxImplicit0->GetMargin(), ExpectedMargin0, Tolerance);
-		EXPECT_NEAR(BoxImplicit1->GetMargin(), ExpectedMargin1, Tolerance);
+		EXPECT_NEAR(BoxImplicit0->GetMarginf(), ExpectedMargin0, Tolerance);
+		EXPECT_NEAR(BoxImplicit1->GetMarginf(), ExpectedMargin1, Tolerance);
 
 		// Box Bounds should include margin
 		const FAABB3 BoxBounds0 = BoxImplicit0->BoundingBox();
