Add unit test for capsule capsule collision
Made sure recenent optimization didn't affect any of the results.

diff --git a/Engine/Source/Programs/HeadlessChaos/Private/HeadlessChaosApp.cpp b/Engine/Source/Programs/HeadlessChaos/Private/HeadlessChaosApp.cpp
--- a/Engine/Source/Programs/HeadlessChaos/Private/HeadlessChaosApp.cpp
+++ b/Engine/Source/Programs/HeadlessChaos/Private/HeadlessChaosApp.cpp
@@ -114,6 +114,7 @@ TEST(CollisionTests, Collisions) {
 	ChaosTest::CollisionBoxPlaneRestitution();
 	ChaosTest::CollisionCubeCubeRestitution();
 	ChaosTest::CollisionBoxToStaticBox();
+	ChaosTest::CollisionCapsuleCapsule();
 	ChaosTest::CollisionConvexConvex();
 
 	// @ todo: Make this work at some point
diff --git a/Engine/Source/Programs/HeadlessChaos/Private/HeadlessChaosTestCollisions.cpp b/Engine/Source/Programs/HeadlessChaos/Private/HeadlessChaosTestCollisions.cpp
--- a/Engine/Source/Programs/HeadlessChaos/Private/HeadlessChaosTestCollisions.cpp
+++ b/Engine/Source/Programs/HeadlessChaos/Private/HeadlessChaosTestCollisions.cpp
@@ -9,6 +9,7 @@
 #include "Chaos/PBDRigidsEvolution.h"
 #include "Chaos/PBDRigidParticles.h"
 #include "Chaos/CollisionResolutionTypes.h"
+#include "Chaos/CollisionOneShotManifoldsMiscShapes.h"
 #include "Chaos/Sphere.h"
 #include "Chaos/Utilities.h"
 #include "Modules/ModuleManager.h"
@@ -580,5 +581,184 @@ namespace ChaosTest {
 		EXPECT_FALSE(Box2->GetP().Equals(StartingPoint)); // moved
 		EXPECT_FALSE(Box2->GetQ().Equals(FQuat::Identity)); // and rotated
 	}
+
+	void CollisionCapsuleCapsule()
+	{
+		const FReal CullDistance = FReal(3);
+		{
+			FCapsule Capsule1(FVec3(-20, -10, 10), FVec3(20, -10, 10), FReal(20));
+			FCapsule Capsule2(FVec3(-20, -10, 10), FVec3(20, -10, 10), FReal(20));
+
+			FImplicitObjectPtr Implicit1 = MakeImplicitObjectPtr<FCapsule>(Capsule1);
+			FImplicitObjectPtr Implicit2 = MakeImplicitObjectPtr<FCapsule>(Capsule2);
+
+			FRigidTransform3 Transform1(FRigidTransform3::Identity);
+			FRigidTransform3 Transform2(FVec3(0.0, 0.0, 18), FRotation3(FQuat::Identity));
+			FPBDCollisionConstraint Constraint;
+
+			FParticleUniqueIndicesMultithreaded UniqueIndices;
+			FPBDRigidsSOAs RigidsSOA(UniqueIndices);
+			TArray<FPBDRigidParticleHandle*> Particles = RigidsSOA.CreateDynamicParticles(2);
+
+			FGeometryParticleHandle* Particle1 = Particles[0];
+			FGeometryParticleHandle* Particle2 = Particles[1];
+
+			Particle1->SetGeometry(Implicit1);
+			Particle2->SetGeometry(Implicit2);
+
+			const FShapeInstancePtr& Shape1 = Particle1->ShapeInstances()[0];
+			const FShapeInstancePtr& Shape2 = Particle2->ShapeInstances()[0];
+
+			FPBDCollisionConstraint::Make(Particle1,
+				Implicit1,
+				Shape1.Get(),
+				nullptr,
+				Transform1,
+				Particle1,
+				Implicit2,
+				Shape2.Get(),
+				nullptr,
+				Transform2,
+				CullDistance,
+				true,
+				EContactShapesType::CapsuleCapsule,
+				Constraint);
+
+			Collisions::UpdateConstraint(Constraint, Transform1, Transform2, 0.1);
+
+			const TArrayView<FManifoldPoint>& ManifoldPoints = Constraint.GetManifoldPoints();
+
+			EXPECT_EQ(ManifoldPoints.Num(), 2);
+
+			EXPECT_EQ(ManifoldPoints[0].ContactPoint.Phi, -22.0f); // Two radius 40 minus 18 of translation -> -22.0f
+			EXPECT_EQ(ManifoldPoints[0].ContactPoint.ShapeContactPoints[0], FVec3f(20, -10, 30)); // z = 10 + 20 radius -> 30.0f
+			EXPECT_EQ(ManifoldPoints[0].ContactPoint.ShapeContactPoints[1], FVec3f(20, -10, -10));
+
+			EXPECT_EQ(ManifoldPoints[1].ContactPoint.Phi, -18.0f);
+			EXPECT_EQ(ManifoldPoints[1].ContactPoint.ShapeContactPoints[0], FVec3f(-20, -10, 30));
+			EXPECT_EQ(ManifoldPoints[1].ContactPoint.ShapeContactPoints[1], FVec3f(-20, -10, 30));
+
+			{
+				Transform1.SetRotation(FRotation3::FromAxisAngle(FVec3(1.0, 0.0, 0.0), 0.1));
+				Collisions::UpdateConstraint(Constraint, Transform1, Transform2, 0.1);
+				const TArrayView<FManifoldPoint>& ManifoldPoints2 = Constraint.GetManifoldPoints();
+				EXPECT_EQ(ManifoldPoints2.Num(), 2);
+				EXPECT_LT(ManifoldPoints2[0].ContactPoint.Phi, -20.0f);
+				EXPECT_GT(ManifoldPoints2[0].ContactPoint.Phi, -21.0f);
+				EXPECT_LT(ManifoldPoints2[1].ContactPoint.Phi, -19.0f);
+				EXPECT_GT(ManifoldPoints2[1].ContactPoint.Phi, -20.0f);
+
+			}
+		}
+		{
+			FCapsule Capsule1(FVec3(-20, -10, 10), FVec3(20, -10, 10), FReal(18));
+			FCapsule Capsule2(FVec3(-20, -10, 10), FVec3(20, -10, 10), FReal(20));
+
+			FImplicitObjectPtr Implicit1 = MakeImplicitObjectPtr<FCapsule>(Capsule1);
+			FImplicitObjectPtr Implicit2 = MakeImplicitObjectPtr<FCapsule>(Capsule2);
+
+			FRigidTransform3 Transform1(FRigidTransform3::Identity);
+			FRigidTransform3 Transform2(FVec3(0.0, 0.0, 16), FRotation3(FQuat::Identity));
+			FPBDCollisionConstraint Constraint;
+
+			FParticleUniqueIndicesMultithreaded UniqueIndices;
+			FPBDRigidsSOAs RigidsSOA(UniqueIndices);
+			TArray<FPBDRigidParticleHandle*> Particles = RigidsSOA.CreateDynamicParticles(2);
+
+			FGeometryParticleHandle* Particle1 = Particles[0];
+			FGeometryParticleHandle* Particle2 = Particles[1];
+
+			Particle1->SetGeometry(Implicit1);
+			Particle2->SetGeometry(Implicit2);
+
+			const FShapeInstancePtr& Shape1 = Particle1->ShapeInstances()[0];
+			const FShapeInstancePtr& Shape2 = Particle2->ShapeInstances()[0];
+
+			FPBDCollisionConstraint::Make(Particle1,
+				Implicit1,
+				Shape1.Get(),
+				nullptr,
+				Transform1,
+				Particle1,
+				Implicit2,
+				Shape2.Get(),
+				nullptr,
+				Transform2,
+				CullDistance,
+				true,
+				EContactShapesType::CapsuleCapsule,
+				Constraint);
+
+			Collisions::UpdateConstraint(Constraint, Transform1, Transform2, 0.1);
+
+			const TArrayView<FManifoldPoint>& ManifoldPoints = Constraint.GetManifoldPoints();
+
+			EXPECT_EQ(ManifoldPoints.Num(), 2);
+
+			EXPECT_EQ(ManifoldPoints[0].ContactPoint.Phi, -22.0f);
+			EXPECT_EQ(ManifoldPoints[0].ContactPoint.ShapeContactPoints[0], FVec3f(20, -10, 28));
+			EXPECT_EQ(ManifoldPoints[0].ContactPoint.ShapeContactPoints[1], FVec3f(20, -10, -10));
+
+			EXPECT_EQ(ManifoldPoints[1].ContactPoint.Phi, -18.0f);
+			EXPECT_EQ(ManifoldPoints[1].ContactPoint.ShapeContactPoints[0], FVec3f(-20, -10, 28));
+			EXPECT_EQ(ManifoldPoints[1].ContactPoint.ShapeContactPoints[1], FVec3f(-20, -10, 30));
+		}
+
+		{
+			FCapsule Capsule1(FVec3(-20, -10, 10), FVec3(20, -10, 10), FReal(20));
+			FCapsule Capsule2(FVec3(-20, -10, 10), FVec3(20, -10, 10), FReal(18));
+
+			FImplicitObjectPtr Implicit1 = MakeImplicitObjectPtr<FCapsule>(Capsule1);
+			FImplicitObjectPtr Implicit2 = MakeImplicitObjectPtr<FCapsule>(Capsule2);
+
+			FRigidTransform3 Transform1(FRigidTransform3::Identity);
+			FRigidTransform3 Transform2(FVec3(0.0, 0.0, 16), FRotation3(FQuat::Identity));
+			FPBDCollisionConstraint Constraint;
+
+			FParticleUniqueIndicesMultithreaded UniqueIndices;
+			FPBDRigidsSOAs RigidsSOA(UniqueIndices);
+			TArray<FPBDRigidParticleHandle*> Particles = RigidsSOA.CreateDynamicParticles(2);
+
+			FGeometryParticleHandle* Particle1 = Particles[0];
+			FGeometryParticleHandle* Particle2 = Particles[1];
+
+			Particle1->SetGeometry(Implicit1);
+			Particle2->SetGeometry(Implicit2);
+
+			const FShapeInstancePtr& Shape1 = Particle1->ShapeInstances()[0];
+			const FShapeInstancePtr& Shape2 = Particle2->ShapeInstances()[0];
+
+			FPBDCollisionConstraint::Make(Particle1,
+				Implicit1,
+				Shape1.Get(),
+				nullptr,
+				Transform1,
+				Particle1,
+				Implicit2,
+				Shape2.Get(),
+				nullptr,
+				Transform2,
+				CullDistance,
+				true,
+				EContactShapesType::CapsuleCapsule,
+				Constraint);
+
+			Collisions::UpdateConstraint(Constraint, Transform1, Transform2, 0.1);
+
+			const TArrayView<FManifoldPoint>& ManifoldPoints = Constraint.GetManifoldPoints();
+
+			EXPECT_EQ(ManifoldPoints.Num(), 2);
+
+			EXPECT_EQ(ManifoldPoints[0].ContactPoint.Phi, -22.0f);
+			EXPECT_EQ(ManifoldPoints[0].ContactPoint.ShapeContactPoints[0], FVec3f(20, -10, 30));
+			EXPECT_EQ(ManifoldPoints[0].ContactPoint.ShapeContactPoints[1], FVec3f(20, -10, -8));
+
+			EXPECT_EQ(ManifoldPoints[1].ContactPoint.Phi, -18.0f);
+			EXPECT_EQ(ManifoldPoints[1].ContactPoint.ShapeContactPoints[0], FVec3f(-20, -10, -10));
+			EXPECT_EQ(ManifoldPoints[1].ContactPoint.ShapeContactPoints[1], FVec3f(-20, -10, -8));
+		}
+		
+	}
+
 }
 
diff --git a/Engine/Source/Programs/HeadlessChaos/Public/HeadlessChaosTestCollisions.h b/Engine/Source/Programs/HeadlessChaos/Public/HeadlessChaosTestCollisions.h
--- a/Engine/Source/Programs/HeadlessChaos/Public/HeadlessChaosTestCollisions.h
+++ b/Engine/Source/Programs/HeadlessChaos/Public/HeadlessChaosTestCollisions.h
@@ -27,4 +27,6 @@ namespace ChaosTest {
 	void CollisionCubeCubeRestitution();
 
 	void CollisionBoxToStaticBox();
+
+	void CollisionCapsuleCapsule();
 }
\ No newline at end of file
