Chaos - Add mass conditioning for character ground constraint

diff --git a/Engine/Source/Programs/HeadlessChaos/Private/Character/TestCharacterGroundConstraint.cpp b/Engine/Source/Programs/HeadlessChaos/Private/Character/TestCharacterGroundConstraint.cpp
--- a/Engine/Source/Programs/HeadlessChaos/Private/Character/TestCharacterGroundConstraint.cpp
+++ b/Engine/Source/Programs/HeadlessChaos/Private/Character/TestCharacterGroundConstraint.cpp
@@ -72,6 +72,16 @@ namespace ChaosTest
 
 		Constraint.ClearDirtyFlags();
 
+		Constraint.SetMaxCharacterGroundMassRatio(1.5);
+		EXPECT_EQ(Constraint.GetMaxCharacterGroundMassRatio(), 1.5);
+		EXPECT_TRUE(Constraint.IsDirty());
+		EXPECT_FALSE(Constraint.IsDirty(EChaosPropertyFlags::CharacterGroundConstraintDynamicData));
+		EXPECT_TRUE(Constraint.IsDirty(EChaosPropertyFlags::CharacterGroundConstraintSettings));
+		EXPECT_FALSE(Constraint.IsDirty(EChaosPropertyFlags::CharacterParticleProxy));
+		EXPECT_FALSE(Constraint.IsDirty(EChaosPropertyFlags::GroundParticleProxy));
+
+		Constraint.ClearDirtyFlags();
+
 		FSingleParticlePhysicsProxy* CharacterProxy = FSingleParticlePhysicsProxy::Create(FGeometryParticle::CreateParticle());
 		FSingleParticlePhysicsProxy* GroundProxy = FSingleParticlePhysicsProxy::Create(FGeometryParticle::CreateParticle());
 
@@ -99,6 +109,7 @@ namespace ChaosTest
 		EXPECT_FALSE(RemoteData.HasCharacterGroundConstraintSettings());
 
 		Constraint.SetTargetHeight(15.0);
+		Constraint.SetMaxCharacterGroundMassRatio(1.2);
 		Constraint.SyncRemoteData(Manager, 0, RemoteData);
 
 		ASSERT_TRUE(RemoteData.HasCharacterGroundConstraintSettings());
@@ -106,6 +117,7 @@ namespace ChaosTest
 		const FCharacterGroundConstraintSettings& Settings = RemoteData.GetCharacterGroundConstraintSettings(Manager, 0);
 
 		EXPECT_EQ(Settings.TargetHeight, 15.0);
+		EXPECT_EQ(Settings.MaxCharacterGroundMassRatio, 1.2);
 
 		Constraint.ClearDirtyFlags();
 		RemoteData.Clear(Manager, 0);
diff --git a/Engine/Source/Programs/HeadlessChaos/Private/Character/TestCharacterGroundConstraintSolver.cpp b/Engine/Source/Programs/HeadlessChaos/Private/Character/TestCharacterGroundConstraintSolver.cpp
--- a/Engine/Source/Programs/HeadlessChaos/Private/Character/TestCharacterGroundConstraintSolver.cpp
+++ b/Engine/Source/Programs/HeadlessChaos/Private/Character/TestCharacterGroundConstraintSolver.cpp
@@ -506,4 +506,50 @@ namespace ChaosTest
 		EXPECT_VECTOR_FLOAT_EQ(Solver.GetAngularImpulse(Dt), FVec3::ZeroVector);
 	}
 
+	// No overlap but projected overlap
+	// Solver should fix overlap
+	// Mass of ground shoul be adjusted so that max character/ground mass ratio is respected
+	TEST_F(CharacterGroundConstraintSolverTest, TestSingleBody_NormalImpulse_MassConditioning)
+	{
+		CharacterBody.SetX(FVec3(0.0, 0.0, 20.0));
+		CharacterBody.SetP(FVec3(0.0, 0.0, 10.0));
+		CharacterBody.SetV(FVec3(0.0, 0.0, -10.0 / Dt));
+
+		GroundBody.SetX(FVec3(0.0, 0.0, 0.0));
+		GroundBody.SetP(FVec3(0.0, 0.0, 0.0));
+		GroundBody.SetV(FVec3(0.0, 0.0, 0.0));
+
+		Data.GroundDistance = 20.0f;
+		Settings.TargetHeight = 15.0f;
+		Settings.RadialForceLimit = 0.0f;
+		Settings.SwingTorqueLimit = 0.0f;
+		Settings.TwistTorqueLimit = 0.0f;
+
+		// Set the mass of the ground to be half the mass of the character (i.e. ratio = 2)
+		// and set the max ratio to 1.5. With a ground mass of 100 the character mass
+		// should be adjusted to look like 150
+		Settings.MaxCharacterGroundMassRatio = 1.5f;
+		CharacterBody.SetInvM(1.0 / 200.0);
+		GroundBody.SetInvM(1.0 / 100.0);
+		GroundBody.SetInvILocal(FVec3(0.01, 0.01, 0.01));
+		const FReal ExpectedCharacterInvMass = 1.0f / 150.0f;
+
+		Solver.SetBodies(&CharacterBody, &GroundBody);
+		Solver.GatherInput(Dt, Settings, Data);
+		SolvePosition(1);
+
+		EXPECT_TRUE(CharacterBody.CP() == FSolverVec3::ZeroVector);
+		EXPECT_FALSE(CharacterBody.DP() == FSolverVec3::ZeroVector);
+
+		FVec3 ExpectedImpulse = (1.0 / ((GroundBody.InvM() + ExpectedCharacterInvMass) * Dt)) * FVec3(0.0, 0.0, 5.0);
+		EXPECT_VECTOR_FLOAT_EQ(Solver.GetLinearImpulse(Dt), ExpectedImpulse);
+		EXPECT_VECTOR_FLOAT_EQ(Solver.GetAngularImpulse(Dt), FVec3::ZeroVector);
+
+		const FVec3 ExpectedCharacterCorrectedP = FVec3(0.0, 0.0, 10.0) + ExpectedCharacterInvMass * ExpectedImpulse * Dt;
+		const FVec3 ExpectedGroundCorrectedP = FVec3(0.0, 0.0, 0.0) - FReal(GroundBody.InvM()) * ExpectedImpulse * Dt;
+		EXPECT_VECTOR_FLOAT_EQ(CharacterBody.CorrectedP(), ExpectedCharacterCorrectedP);
+		EXPECT_VECTOR_FLOAT_EQ(GroundBody.CorrectedP(), ExpectedGroundCorrectedP);
+		EXPECT_VECTOR_FLOAT_EQ(FVec3(CharacterBody.CorrectedP() - GroundBody.CorrectedP()), FVec3(0.0, 0.0, 15.0));
+	}
+
 } // namespace ChaosTest
\ No newline at end of file
diff --git a/Engine/Source/Runtime/Experimental/Chaos/Private/Chaos/Character/CharacterGroundConstraintSolver.cpp b/Engine/Source/Runtime/Experimental/Chaos/Private/Chaos/Character/CharacterGroundConstraintSolver.cpp
--- a/Engine/Source/Runtime/Experimental/Chaos/Private/Chaos/Character/CharacterGroundConstraintSolver.cpp
+++ b/Engine/Source/Runtime/Experimental/Chaos/Private/Chaos/Character/CharacterGroundConstraintSolver.cpp
@@ -89,6 +89,27 @@ namespace Chaos
 		{
 			ImpulseData.Reset();
 
+			//////////////////////////////////////////////////////////////////////////
+			// Mass conditioning
+
+			FReal CharacterInvM = BodyData.CharacterBody.InvM();
+			FSolverVec3 CharacterInvILocal = BodyData.CharacterBody.InvILocal();
+
+			FReal GroundInvM = 0.0;
+
+			const bool bTwoBody = BodyData.GroundBody.IsValid();
+			if (bTwoBody && Settings.MaxCharacterGroundMassRatio > 0.0f && BodyData.GroundBody.InvM() > 0.0f)
+			{
+				GroundInvM = BodyData.GroundBody.InvM();
+				const FReal Ratio = GroundInvM / CharacterInvM;
+				if (Ratio > Settings.MaxCharacterGroundMassRatio)
+				{
+					const FReal OldInvM = CharacterInvM;
+					CharacterInvM = GroundInvM / Settings.MaxCharacterGroundMassRatio;
+					CharacterInvILocal *= FSolverReal(CharacterInvM / OldInvM);
+				}
+			}
+
 			//////////////////////////////////////////////////////////////////////////
 			// Normal constraint
 
@@ -106,9 +127,8 @@ namespace Chaos
 			FVec3 InitialError = FVec3(PC_Init - PG_Init);
 			FVec3 ProjectedError = InitialError + CharacterLinearDisplacement;
 
-			const FReal CharacterInvM = BodyData.CharacterBody.InvM();
 			FReal TwoBodyEffectiveMass = 1.0 / CharacterInvM;
-			FReal GroundInvM = 0.0;
+			
 			FMatrix33 GroundInvI(0.0);
 			FVec3 GroundNormal = Data.GroundNormal;
 			FVec3 GroundOffset = FVec3(ForceInitToZero);
@@ -119,7 +139,6 @@ namespace Chaos
 			FVec3 Ground_P(ForceInitToZero);
 			FRotation3 Ground_Q(FRotation3::Identity);
 
-			const bool bTwoBody = BodyData.GroundBody.IsValid();
 			if (bTwoBody)
 			{
 				Ground_X = BodyData.GroundBody.X();
@@ -172,7 +191,7 @@ namespace Chaos
 
 			//////////////////////////////////////////////////////////////////////////
 			// Upright constraint
-			const Chaos::FMatrix33 CharacterInvInertia = Utilities::ComputeWorldSpaceInertia(Character_Q, BodyData.CharacterBody.InvILocal());
+			const Chaos::FMatrix33 CharacterInvInertia = Utilities::ComputeWorldSpaceInertia(Character_Q, CharacterInvILocal);
 			ConstraintData.CharacterInvI = FSolverMatrix33(CharacterInvInertia);
 			ConstraintData.VerticalAxis = FSolverVec3(Settings.VerticalAxis);
 			ConstraintData.CharacterVerticalAxis = FSolverVec3(Character_Q.RotateVector(Settings.VerticalAxis));
diff --git a/Engine/Source/Runtime/Experimental/Chaos/Public/Chaos/Character/CharacterGroundConstraint.h b/Engine/Source/Runtime/Experimental/Chaos/Public/Chaos/Character/CharacterGroundConstraint.h
--- a/Engine/Source/Runtime/Experimental/Chaos/Public/Chaos/Character/CharacterGroundConstraint.h
+++ b/Engine/Source/Runtime/Experimental/Chaos/Public/Chaos/Character/CharacterGroundConstraint.h
@@ -243,7 +243,6 @@ namespace Chaos
 		}
 		FReal GetRadialForceLimit() const { return ConstraintSettings.Read().RadialForceLimit; }
 
-
 		/// The maximum force that the character can apply in the plane defined by the ground normal
 		/// to hold the character in place while standing on an unwalkable incline (defined by max
 		/// walkable slope angle)
@@ -255,18 +254,6 @@ namespace Chaos
 		}
 		FReal GetFrictionForceLimit() const { return ConstraintSettings.Read().FrictionForceLimit; }
 
-
-		/// The maximum torque that the character can apply about the vertical axis to rotate the
-		/// character to its target facing direction (or hold the character in place if no target
-		/// delta facing is set).
-		UE_DEPRECATED(5.3, "Split into swing/twist torque limits. Use SetTwistTorqueLimit instead for torque limit about vertical axis.")
-		void SetTorqueLimit(const FReal& Value)
-		{
-			SetTwistTorqueLimit(Value);
-		}
-		UE_DEPRECATED(5.3, "Split into swing/twist torque limits. Use GetTwistTorqueLimit instead for torque limit about vertical axis.")
-		FReal GetTorqueLimit() const { return GetTwistTorqueLimit(); }
-
 		/// The maximum torque that the character can apply about the vertical axis to rotate the
 		/// character to its target facing direction (or hold the character in place if no target
 		/// delta facing is set).
@@ -296,6 +283,19 @@ namespace Chaos
 		}
 		void* GetUserData() const { return ConstraintSettings.Read().UserData; }
 
+		/// Maximum ratio between the character mass and the ground body mass as used by the constraint solver
+		/// E.g. a value of 2 means that the character can be at most twice as heavy as the ground body and
+		/// if it is over that the constraint solve will inflate the mass of the ground body to reach a 2:1 ratio
+		/// This is useful for stability when walking on a light mass.
+		/// A negative value disables this mass conditioning feature.
+		void SetMaxCharacterGroundMassRatio(const FReal& Value)
+		{
+			ConstraintSettings.Modify(true, DirtyFlags, Proxy, [&Value](FCharacterGroundConstraintSettings& Data) {
+				Data.MaxCharacterGroundMassRatio = Value;
+				});
+		}
+		FReal GetMaxCharacterGroundMassRatio() const { return ConstraintSettings.Read().MaxCharacterGroundMassRatio; }
+
 		//////////////////////////////////////////////////////////////////////////
 		// Output
 
diff --git a/Engine/Source/Runtime/Experimental/Chaos/Public/Chaos/Character/CharacterGroundConstraintSettings.h b/Engine/Source/Runtime/Experimental/Chaos/Public/Chaos/Character/CharacterGroundConstraintSettings.h
--- a/Engine/Source/Runtime/Experimental/Chaos/Public/Chaos/Character/CharacterGroundConstraintSettings.h
+++ b/Engine/Source/Runtime/Experimental/Chaos/Public/Chaos/Character/CharacterGroundConstraintSettings.h
@@ -24,6 +24,7 @@ namespace Chaos
 		FReal AssumedOnGroundHeight = 2.0;			/// Below this height the character is assumed to be on the ground and can apply force/torque to reach the target position and facing
 		FReal MotionTargetMassBias = 1.0;			/// Scaling factor applied to reduce the mass of the ground body in the linear motion target constraint
 		FReal RadialForceMotionTargetScaling = 1.0; /// Scale the radial force to always reach this fraction of the motion target
+		FReal MaxCharacterGroundMassRatio = -1.0f;	/// Maximum ratio for the mass of character mass to ground mass. A negative value indicates no mass conditioning is applied
 		void* UserData = nullptr;
 	};
 
