Added unit tests for chaos triangle meshes with negative scale and back-face culling.

diff --git a/Engine/Source/Programs/HeadlessChaos/Private/HeadlessChaosTestSweep.cpp b/Engine/Source/Programs/HeadlessChaos/Private/HeadlessChaosTestSweep.cpp
--- a/Engine/Source/Programs/HeadlessChaos/Private/HeadlessChaosTestSweep.cpp
+++ b/Engine/Source/Programs/HeadlessChaos/Private/HeadlessChaosTestSweep.cpp
@@ -7,6 +7,7 @@
 #include "Chaos/Convex.h"
 #include "Chaos/TriangleMeshImplicitObject.h"
 #include "Chaos/ImplicitObjectScaled.h"
+#include "Chaos/GeometryQueries.h"
 
 namespace ChaosTest
 {
@@ -1193,4 +1194,230 @@ namespace ChaosTest
 
 		EXPECT_EQ(Result.bResult, false);
 	}
+
+	// Some extra tests for negative scales
+	struct FScaledTriMeshSweepFixture : public FTriangleMeshBackfaceSweepFixture
+	{
+		static FMeshInitData CreateSimpleTriInitData()
+		{
+			FMeshInitData Result;
+			// Use a slightly different triangle mesh than above. This allows easy negation of the mesh without having to move the sweep's location.
+			Result.TrimeshParticles = FTriangleMeshImplicitObject::ParticlesType(
+				{
+					{-100, -100, 0},
+					{100, -100, 0},
+					{0, 100, 0},
+				});
+			Result.Indices =
+			{
+				{0, 1, 2},
+			};
+			Result.AutoBuildMaterials();
+			return Result;
+		}
+
+		FSweepResultData SweepSphere(const TImplicitObjectScaled<FTriangleMeshImplicitObject>& TriangleMesh, const FVec3& SweepStart, const FVec3& SweepDirection, const FReal SweepLength = 100)
+		{
+			const FReal SphereRadius = 25;
+			const Chaos::FSphere Sphere(FVec3::ZeroVector, SphereRadius);
+			const FRigidTransform3 StartTM(SweepStart, TRotation<FReal, 3>::Identity);
+			const FVec3 NormalizedSweepDirection = SweepDirection.GetSafeNormal();
+
+			FSweepResultData Result;
+			Result.bResult = SweepQuery(TriangleMesh, FRigidTransform3(), Sphere, StartTM, NormalizedSweepDirection, SweepLength, Result.TOI, Result.Position, Result.Normal, Result.FaceId, Result.FaceNormal, 0.0, true);
+			return Result;
+		}
+
+		FSweepResultData SweepSphere(FMeshInitData& InitData, const FVec3& SweepStart, const FVec3& SweepDirection, const FReal SweepLength = 100)
+		{
+			FTriangleMeshImplicitObject TriangleMesh(MoveTemp(InitData.TrimeshParticles), MoveTemp(InitData.Indices), MoveTemp(InitData.Materials), nullptr, nullptr, true);
+			FTriangleMeshBackfaceSweepFixture::SweepSphere(TriangleMesh, SweepStart, SweepDirection, SweepLength);
+			// Unfortunately, it is not possible to currently call FTriangleMeshImplicitObject::SweepGeom directly. 
+			// There's code inside of TImplicitObjectScaled that negates some inputs values (the direction) that is necessary for the tests to work (follows the editor path).
+			TImplicitObjectScaled<FTriangleMeshImplicitObject> ScaledMesh(&TriangleMesh, TriMeshScale);
+			return SweepSphere(ScaledMesh, SweepStart, SweepDirection, SweepLength);
+		}
+
+		FSweepResultData SweepSphereVsSimpleTri(const FVec3& SweepStart, const FVec3& SweepDirection, const FReal SweepLength = 100)
+		{
+			FMeshInitData InitData = CreateSimpleTriInitData();
+			return SweepSphere(InitData, SweepStart, SweepDirection, SweepLength);
+		}
+
+		FVec3 TriMeshScale = FVec3(1, 1, 1);
+	};
+
+	// Test sweeps along the original front face
+	TEST_F(FScaledTriMeshSweepFixture, TriMeshWithNegativeXAxis_SweepInFrontOpposedToNormal_HitIsExpected)
+	{
+		// Negating the x axis should flip the winding order, but up is still the front face
+		TriMeshScale = FVec3(-1, 1, 1);
+		const FSweepResultData Result = SweepSphereVsSimpleTri(FVec3(20, 20, 50), FVec3(0, 0, -1));
+
+		EXPECT_EQ(Result.bResult, true);
+		EXPECT_NEAR(Result.TOI, 25, KINDA_SMALL_NUMBER);
+		EXPECT_VECTOR_NEAR(Result.Position, FVec3(20, 20, 0), KINDA_SMALL_NUMBER);
+		EXPECT_VECTOR_NEAR(Result.Normal, FVec3(0, 0, 1), KINDA_SMALL_NUMBER);
+	}
+
+	TEST_F(FScaledTriMeshSweepFixture, TriMeshWithNegativeYAxis_SweepInFrontOpposedToNormal_HitIsExpected)
+	{
+		// Negating the y axis should flip the winding order, but up is still the front face
+		TriMeshScale = FVec3(1, -1, 1);
+		const FSweepResultData Result = SweepSphereVsSimpleTri(FVec3(20, 20, 50), FVec3(0, 0, -1));
+
+		EXPECT_EQ(Result.bResult, true);
+		EXPECT_NEAR(Result.TOI, 25, KINDA_SMALL_NUMBER);
+		EXPECT_VECTOR_NEAR(Result.Position, FVec3(20, 20, 0), KINDA_SMALL_NUMBER);
+		EXPECT_VECTOR_NEAR(Result.Normal, FVec3(0, 0, 1), KINDA_SMALL_NUMBER);
+	}
+
+	TEST_F(FScaledTriMeshSweepFixture, TriMeshWithNegativeZAxis_SweepInBackAlongToNormal_HitIsNotExpected)
+	{
+		// Negating the z axis should cause down to be the front face
+		TriMeshScale = FVec3(1, 1, -1);
+		const FSweepResultData Result = SweepSphereVsSimpleTri(FVec3(20, 20, 50), FVec3(0, 0, -1));
+
+		EXPECT_EQ(Result.bResult, false);
+	}
+
+	TEST_F(FScaledTriMeshSweepFixture, TriMeshWithNegativeXYAxis_SweepInFrontOpposedToNormal_HitIsExpected)
+	{
+		TriMeshScale = FVec3(-1, -1, 1);
+		const FSweepResultData Result = SweepSphereVsSimpleTri(FVec3(20, 20, 50), FVec3(0, 0, -1));
+
+		EXPECT_EQ(Result.bResult, true);
+		EXPECT_NEAR(Result.TOI, 25, KINDA_SMALL_NUMBER);
+		EXPECT_VECTOR_NEAR(Result.Position, FVec3(20, 20, 0), KINDA_SMALL_NUMBER);
+		EXPECT_VECTOR_NEAR(Result.Normal, FVec3(0, 0, 1), KINDA_SMALL_NUMBER);
+	}
+
+	TEST_F(FScaledTriMeshSweepFixture, TriMeshWithNegativeXZAxis_SweepInFrontOpposedToNormal_HitIsNotExpected)
+	{
+		TriMeshScale = FVec3(-1, 1, -1);
+		const FSweepResultData Result = SweepSphereVsSimpleTri(FVec3(20, 20, 50), FVec3(0, 0, -1));
+
+		EXPECT_EQ(Result.bResult, false);
+	}
+
+	TEST_F(FScaledTriMeshSweepFixture, TriMeshWithNegativeYZAxis_SweepInFrontOpposedToNormal_HitIsNotExpected)
+	{
+		TriMeshScale = FVec3(1, -1, -1);
+		const FSweepResultData Result = SweepSphereVsSimpleTri(FVec3(20, 20, 50), FVec3(0, 0, -1));
+
+		EXPECT_EQ(Result.bResult, false);
+	}
+
+	TEST_F(FScaledTriMeshSweepFixture, TriMeshWithNegativeXYZAxis_SweepInFrontOpposedToNormal_HitIsNotExpected)
+	{
+		TriMeshScale = FVec3(-1, -1, -1);
+		const FSweepResultData Result = SweepSphereVsSimpleTri(FVec3(20, 20, 50), FVec3(0, 0, -1));
+
+		EXPECT_EQ(Result.bResult, false);
+	}
+
+	// Test sweeps along the original back face
+	TEST_F(FScaledTriMeshSweepFixture, TriMeshWithNegativeXAxis_SweepOnBackAlongNormal_HitIsNotExpected)
+	{
+		// Negating the x axis should flip the winding order, but up is still the front face
+		TriMeshScale = FVec3(-1, 1, 1);
+		const FSweepResultData Result = SweepSphereVsSimpleTri(FVec3(20, 20, -50), FVec3(0, 0, 1));
+
+		EXPECT_EQ(Result.bResult, false);
+	}
+
+	TEST_F(FScaledTriMeshSweepFixture, TriMeshWithNegativeYAxis_SweepOnBackAlongNormal_HitIsNotExpected)
+	{
+		// Negating the y axis should flip the winding order, but up is still the front face
+		TriMeshScale = FVec3(1, -1, 1);
+		const FSweepResultData Result = SweepSphereVsSimpleTri(FVec3(20, 20, -50), FVec3(0, 0, 1));
+
+		EXPECT_EQ(Result.bResult, false);
+	}
+
+	TEST_F(FScaledTriMeshSweepFixture, TriMeshWithNegativeZAxis_SweepOnBackAlongNormal_HitIsExpected)
+	{
+		// Negating the z axis should cause down to be the front face
+		TriMeshScale = FVec3(1, 1, -1);
+		const FSweepResultData Result = SweepSphereVsSimpleTri(FVec3(20, 20, -50), FVec3(0, 0, 1));
+
+		EXPECT_EQ(Result.bResult, true);
+		EXPECT_NEAR(Result.TOI, 25, KINDA_SMALL_NUMBER);
+		EXPECT_VECTOR_NEAR(Result.Position, FVec3(20, 20, 0), KINDA_SMALL_NUMBER);
+		EXPECT_VECTOR_NEAR(Result.Normal, FVec3(0, 0, -1), KINDA_SMALL_NUMBER);
+	}
+
+	TEST_F(FScaledTriMeshSweepFixture, TriMeshWithNegativeXYAxis_SweepOnBackAlongNormal_HitIsNotExpected)
+	{
+		TriMeshScale = FVec3(-1, -1, 1);
+		const FSweepResultData Result = SweepSphereVsSimpleTri(FVec3(20, 20, -50), FVec3(0, 0, 1));
+
+		EXPECT_EQ(Result.bResult, false);
+	}
+
+	TEST_F(FScaledTriMeshSweepFixture, TriMeshWithNegativeXZAxis_SweepOnBackAlongNormal_HitIsExpected)
+	{
+		TriMeshScale = FVec3(-1, 1, -1);
+		const FSweepResultData Result = SweepSphereVsSimpleTri(FVec3(20, 20, -50), FVec3(0, 0, 1));
+
+		EXPECT_EQ(Result.bResult, true);
+		EXPECT_NEAR(Result.TOI, 25, KINDA_SMALL_NUMBER);
+		EXPECT_VECTOR_NEAR(Result.Position, FVec3(20, 20, 0), KINDA_SMALL_NUMBER);
+		EXPECT_VECTOR_NEAR(Result.Normal, FVec3(0, 0, -1), KINDA_SMALL_NUMBER);
+	}
+
+	TEST_F(FScaledTriMeshSweepFixture, TriMeshWithNegativeYZAxis_SweepOnBackAlongNormal_HitIsExpected)
+	{
+		TriMeshScale = FVec3(1, -1, -1);
+		const FSweepResultData Result = SweepSphereVsSimpleTri(FVec3(20, 20, -50), FVec3(0, 0, 1));
+
+		EXPECT_EQ(Result.bResult, true);
+		EXPECT_NEAR(Result.TOI, 25, KINDA_SMALL_NUMBER);
+		EXPECT_VECTOR_NEAR(Result.Position, FVec3(20, 20, 0), KINDA_SMALL_NUMBER);
+		EXPECT_VECTOR_NEAR(Result.Normal, FVec3(0, 0, -1), KINDA_SMALL_NUMBER);
+	}
+
+	TEST_F(FScaledTriMeshSweepFixture, TriMeshWithNegativeXYZAxis_SweepOnBackAlongNormal_HitIsExpected)
+	{
+		TriMeshScale = FVec3(-1, -1, -1);
+		const FSweepResultData Result = SweepSphereVsSimpleTri(FVec3(20, 20, -50), FVec3(0, 0, 1));
+
+		EXPECT_EQ(Result.bResult, true);
+		EXPECT_NEAR(Result.TOI, 25, KINDA_SMALL_NUMBER);
+		EXPECT_VECTOR_NEAR(Result.Position, FVec3(20, 20, 0), KINDA_SMALL_NUMBER);
+		EXPECT_VECTOR_NEAR(Result.Normal, FVec3(0, 0, -1), KINDA_SMALL_NUMBER);
+	}
+
+	// Test parallel sweeps
+	TEST_F(FScaledTriMeshSweepFixture, TriMeshWithNegativeXAxis_SphereWithInitialOverlap_SweepOpposedToNormalWithinParallelEpsilon_HitIsExpected)
+	{
+		TriMeshScale = FVec3(-1, 1, 1);
+		FSweepResultData Result = SweepSphereVsSimpleTri(FVec3(0, 0, 10), FVec3(1, 0, ValueWithinParallelEpsilon));
+
+		EXPECT_EQ(Result.bResult, true);
+	}
+
+	TEST_F(FScaledTriMeshSweepFixture, TriMeshWithNegativeZAxis_SphereWithInitialOverlap_SweepOpposedToNormalWithinParallelEpsilon_HitIsExpected)
+	{
+		TriMeshScale = FVec3(1, 1, -1);
+		FSweepResultData Result = SweepSphereVsSimpleTri(FVec3(0, 0, 10), FVec3(1, 0, ValueWithinParallelEpsilon));
+
+		EXPECT_EQ(Result.bResult, true);
+	}
+
+	TEST_F(FScaledTriMeshSweepFixture, TriMeshWithNegativeXAxis_SphereWithInitialOverlap_SweepAlongNormalWithinParallelEpsilon_HitIsExpected)
+	{
+		TriMeshScale = FVec3(-1, 1, 1);
+		FSweepResultData Result = SweepSphereVsSimpleTri(FVec3(0, 0, 10), FVec3(1, 0, -ValueWithinParallelEpsilon));
+
+		EXPECT_EQ(Result.bResult, true);
+	}
+
+	TEST_F(FScaledTriMeshSweepFixture, TriMeshWithNegativeZAxis_SphereWithInitialOverlap_SweepAlongNormalWithinParallelEpsilon_HitIsExpected)
+	{
+		TriMeshScale = FVec3(1, 1, -1);
+		FSweepResultData Result = SweepSphereVsSimpleTri(FVec3(0, 0, 10), FVec3(1, 0, -ValueWithinParallelEpsilon));
+
+		EXPECT_EQ(Result.bResult, true);
+	}
 }
