Fix collision bug in EPA
An epsilon check was required when checking visible triangles from the new added vertex.

diff --git a/Engine/Source/Programs/HeadlessChaos/Private/HeadlessChaosApp.cpp b/Engine/Source/Programs/HeadlessChaos/Private/HeadlessChaosApp.cpp
--- a/Engine/Source/Programs/HeadlessChaos/Private/HeadlessChaosApp.cpp
+++ b/Engine/Source/Programs/HeadlessChaos/Private/HeadlessChaosApp.cpp
@@ -245,6 +245,7 @@ TEST(GJK, GJKRaycastTests) {
 	ChaosTest::GJKBoxCapsuleSweep();
 	ChaosTest::GJKBoxBoxSweep();
 	ChaosTest::GJKCapsuleConvexInitialOverlapSweep();
+	ChaosTest::GJKConvexCapsule();
 	SUCCEED();
 }
 
diff --git a/Engine/Source/Programs/HeadlessChaos/Private/HeadlessChaosTestGJK.cpp b/Engine/Source/Programs/HeadlessChaos/Private/HeadlessChaosTestGJK.cpp
--- a/Engine/Source/Programs/HeadlessChaos/Private/HeadlessChaosTestGJK.cpp
+++ b/Engine/Source/Programs/HeadlessChaos/Private/HeadlessChaosTestGJK.cpp
@@ -621,6 +621,65 @@ namespace ChaosTest
 	}
 
 
+	void GJKConvexCapsule()
+	{
+		// The following tests are coming form a bug found in EPAComputeVisibilityBorder. The fix was adding an epsilon test, instead of a hard zero comparison.
+		TArray<FConvex::FVec3Type> HullParticles;
+		HullParticles.SetNum(12);
+		HullParticles[0] = { 100.000000, -460.000000, -100.000000 };
+		HullParticles[1] = { 100.000000, -460.000000, 100.000000 };
+		HullParticles[2] = { 400.000000, 0.00000000, 100.000000 };
+		HullParticles[3] = { 400.000000, 0.00000000, -100.000000 };
+		HullParticles[4] = { -100.000000, 390.000000, 100.000000 };
+		HullParticles[5] = { -500.000000, 0.00000000, 100.000000 };
+		HullParticles[6] = { -500.000000, 0.00000000, -100.000000 };
+		HullParticles[7] = { -100.000000, 390.000000, -100.000000 };
+		HullParticles[8] = { -100.000000, -460.000000, 100.000000 };
+		HullParticles[9] = { 100.000000, 390.000000, 100.000000 };
+		HullParticles[10] = { 100.000000, 390.000000, -100.000000 };
+		HullParticles[11] = { -100.000000, -460.000000, -100.000000 };
+		FConvex Convex(HullParticles, 0.0);
+		TImplicitObjectScaled<FConvex> ImplicitScaledConvex(&Convex, FVec3(2.00000000, 2.00000000, 2.00000000));
+		FCapsule Capsule(FVec3(0.00000000, 0.00000000, -55.0000000), FVec3(0.00000000, 0.00000000, 55.00000000), 35.000000);
+		
+		FReal Thickness = 0;
+		bool bComputeMTD = true;
+
+		FVec3 LocalPosition, LocalNormal;
+		FReal OutTime = -1;
+		{
+			FRigidTransform3 BToAFullTM(FVec3(-23.634607186206267, 18.980813439935446, 2.2124981121014571),
+				FRotation3(FQuat(0.0000000000000000, 0.0000000000000000, -0.57785636374940896, 0.81613848265739242)));
+			FVec3 LocalDir(0.39419760716989327, -0.91902570906429382, 0.0000000000000000);
+			FVec3 Offset(14.467136106672115, 2.3920753243291983, -2.2124981121014571);
+			FReal Length = 6.9428286552429199;
+			bool bIsHitting = GJKRaycast2<FReal>(ImplicitScaledConvex, Capsule, BToAFullTM, LocalDir, Length, OutTime, LocalPosition, LocalNormal, Thickness, bComputeMTD, Offset, Thickness);
+			EXPECT_TRUE(bIsHitting);
+			EXPECT_LT(OutTime, 0.0);
+		}
+		{
+			FRigidTransform3 BToAFullTM(FVec3(27.711238420990412, -2.1483584435627563, 2.0126037597656250),
+				FRotation3(FQuat(0.0000000000000000, 0.0000000000000000, 0.12793015790484025, 0.99178317927783100)));
+			FVec3 LocalDir(0.96726777078778936, 0.25375796307645249, 0.0000000000000000);
+			FVec3 Offset(-27.711238420990412, 2.1483584435627563, -2.0126037597656250);
+			FReal Length = 15.497749328613281;
+			bool bIsHitting = GJKRaycast2<FReal>(ImplicitScaledConvex, Capsule, BToAFullTM, LocalDir, Length, OutTime, LocalPosition, LocalNormal, Thickness, bComputeMTD, Offset, Thickness);
+			EXPECT_TRUE(bIsHitting);
+			EXPECT_LT(OutTime, 0.0);
+		}
+		{
+			FRigidTransform3 BToAFullTM(FVec3(-22.297271080315113, 30.859426572671509, 2.2125010768795619),
+			FRotation3(FQuat(0.0000000000000000, 0.0000000000000000, 0.076283974503952301, 0.99708613230446663)));
+			FVec3 LocalDir(0.59685429204739793, 0.091864661873730533, 0.79707334589105316);
+			FReal Length = 38.664058685302734;
+			FVec3 Offset(22.297271080315113, -30.859426572671509, -2.2125010768795619);
+			bool bIsHitting = GJKRaycast2<FReal>(ImplicitScaledConvex, Capsule, BToAFullTM, LocalDir, Length, OutTime, LocalPosition, LocalNormal, Thickness, bComputeMTD, Offset, Thickness);
+			EXPECT_TRUE(bIsHitting);
+			EXPECT_LT(OutTime, 0.0);
+		}
+	}
+
+
 	void GJKSphereScaledSphereTest()
 	{
 		Chaos::FSphere A(FVec3(10, 0, 0), 5);
diff --git a/Engine/Source/Programs/HeadlessChaos/Public/HeadlessChaosTestGJK.h b/Engine/Source/Programs/HeadlessChaos/Public/HeadlessChaosTestGJK.h
--- a/Engine/Source/Programs/HeadlessChaos/Public/HeadlessChaosTestGJK.h
+++ b/Engine/Source/Programs/HeadlessChaos/Public/HeadlessChaosTestGJK.h
@@ -43,4 +43,6 @@ namespace ChaosTest {
 	void GJKCapsuleConvexInitialOverlapSweep();
 
 	void GJKLargeDistanceCapsuleSweep();
+
+	void GJKConvexCapsule();
 }
\ No newline at end of file
diff --git a/Engine/Source/Runtime/Experimental/Chaos/Public/Chaos/EPAVectorized.h b/Engine/Source/Runtime/Experimental/Chaos/Public/Chaos/EPAVectorized.h
--- a/Engine/Source/Runtime/Experimental/Chaos/Public/Chaos/EPAVectorized.h
+++ b/Engine/Source/Runtime/Experimental/Chaos/Public/Chaos/EPAVectorized.h
@@ -318,6 +318,7 @@ bool VectorInitializeEPA(TArray<VectorRegister4Float>& VertsA, TArray<VectorRegi
 
 FORCEINLINE_DEBUGGABLE void VectorEPAComputeVisibilityBorder(TEPAWorkingArray<VectorTEPAEntry>& Entries, int32 EntryIdx, const VectorRegister4Float& W, TEPAWorkingArray<FEPAFloodEntry>& OutBorderEdges, TEPAWorkingArray<FEPAFloodEntry>& ToVisitStack)
 {
+	constexpr VectorRegister4Float Eps = MakeVectorRegisterFloatConstant(1.e-4f, 1.e-4f, 1.e-4f, 1.e-4f);
 	{
 		VectorTEPAEntry& Entry = Entries[EntryIdx];
 		for (int i = 0; i < 3; ++i)
@@ -335,8 +336,7 @@ FORCEINLINE_DEBUGGABLE void VectorEPAComputeVisibilityBorder(TEPAWorkingArray<Ve
 		VectorTEPAEntry& Entry = Entries[FloodEntry.EntryIdx];
 		if (!Entry.bObsolete)
 		{
-			
-			if (VectorMaskBits(VectorCompareGT(VectorZero(), Entry.DistanceToPlane(W))))
+			if (VectorMaskBits(VectorCompareLE(Entry.DistanceToPlane(W), Eps)))
 			{
 				//W can't see this triangle so mark the edge as a border
 				OutBorderEdges.Add(FloodEntry);
