Fixed a sweep bug where unions with scaled sub-objects would incorrectly run the "sweep as ray" optimization and hit an ensure.

diff --git a/Engine/Source/Programs/HeadlessChaos/Private/HeadlessChaosTestSweep.cpp b/Engine/Source/Programs/HeadlessChaos/Private/HeadlessChaosTestSweep.cpp
--- a/Engine/Source/Programs/HeadlessChaos/Private/HeadlessChaosTestSweep.cpp
+++ b/Engine/Source/Programs/HeadlessChaos/Private/HeadlessChaosTestSweep.cpp
@@ -885,7 +885,37 @@ namespace ChaosTest
 		EXPECT_NEAR(WorldPosition.Z, 303.77856, KINDA_SMALL_NUMBER);
 	}
 
-	struct FTriangleMeshSweepFixture : public testing::Test
+	struct FSphereSweepFixture : public testing::Test
+	{
+		struct FSweepResultData
+		{
+			bool bResult;
+			FReal TOI;
+			FVec3 Position, Normal, FaceNormal;
+			int32 FaceId;
+		};
+
+		FSweepResultData SweepQuerySphere(const FImplicitObject& TestObject, const FRigidTransform3& TestObjectTransform, const FVec3& SweepStart, const FVec3& SweepDirection, const FReal SweepLength = 100)
+		{
+			const Chaos::FSphere Sphere(FVec3::ZeroVector, SphereRadius);
+			const FRigidTransform3 StartTM(SweepStart, TRotation<FReal, 3>::Identity);
+			const FVec3 NormalizedSweepDirection = SweepDirection.GetSafeNormal();
+
+			FSweepResultData Result;
+			Result.bResult = SweepQuery(TestObject, TestObjectTransform, Sphere, StartTM, NormalizedSweepDirection, SweepLength, Result.TOI, Result.Position, Result.Normal, Result.FaceId, Result.FaceNormal, 0.0, bComputeMTD);
+			return Result;
+		}
+
+		FSweepResultData SweepQuerySphere(const FImplicitObject& TestObject, const FVec3& SweepStart, const FVec3& SweepDirection, const FReal SweepLength = 100)
+		{
+			return SweepQuerySphere(TestObject, FRigidTransform3(), SweepStart, SweepDirection, SweepLength);
+		}
+
+		FReal SphereRadius = 25;
+		bool bComputeMTD = true;
+	};
+
+	struct FTriangleMeshSweepFixture : public FSphereSweepFixture
 	{
 		struct FMeshInitData
 		{
@@ -903,14 +933,6 @@ namespace ChaosTest
 			}
 		};
 
-		struct FSweepResultData
-		{
-			bool bResult;
-			FReal TOI;
-			FVec3 Position, Normal, FaceNormal;
-			int32 FaceId;
-		};
-
 		static FMeshInitData CreateSimpleTriInitData()
 		{
 			FMeshInitData Result;
@@ -949,13 +971,12 @@ namespace ChaosTest
 
 		FSweepResultData SweepSphere(const FTriangleMeshImplicitObject& TriangleMesh, const FVec3& SweepStart, const FVec3& SweepDirection, const FReal SweepLength = 100)
 		{
-			const FReal SphereRadius = 25;
 			const Chaos::FSphere Sphere(FVec3::ZeroVector, SphereRadius);
 			const FRigidTransform3 StartTM(SweepStart, TRotation<FReal, 3>::Identity);
 			const FVec3 NormalizedSweepDirection = SweepDirection.GetSafeNormal();
 
 			FSweepResultData Result;
-			Result.bResult = TriangleMesh.SweepGeom(Sphere, StartTM, NormalizedSweepDirection, SweepLength, Result.TOI, Result.Position, Result.Normal, Result.FaceId, Result.FaceNormal, 0.0, true);
+			Result.bResult = TriangleMesh.SweepGeom(Sphere, StartTM, NormalizedSweepDirection, SweepLength, Result.TOI, Result.Position, Result.Normal, Result.FaceId, Result.FaceNormal, 0.0, bComputeMTD);
 			return Result;
 		}
 
@@ -1218,7 +1239,6 @@ namespace ChaosTest
 
 		FSweepResultData SweepSphere(const TImplicitObjectScaled<FTriangleMeshImplicitObject>& TriangleMesh, const FVec3& SweepStart, const FVec3& SweepDirection, const FReal SweepLength = 100)
 		{
-			const FReal SphereRadius = 25;
 			const Chaos::FSphere Sphere(FVec3::ZeroVector, SphereRadius);
 			const FRigidTransform3 StartTM(SweepStart, TRotation<FReal, 3>::Identity);
 			const FVec3 NormalizedSweepDirection = SweepDirection.GetSafeNormal();
@@ -1420,4 +1440,56 @@ namespace ChaosTest
 
 		EXPECT_EQ(Result.bResult, true);
 	}
+
+	TEST_F(FSphereSweepFixture, TransformedUnionContainingNonUniformScaledObject_SweepSphere_ResultsAreExpected)
+	{
+		// This test is to catch a bug where a Transformed union with a scaled object would 
+		// ensure because it attempted to go down the sweep as raycast path when there was a non-uniform scale.
+		bComputeMTD = false;
+		SphereRadius = 1;
+
+		
+		// Create a union that has a scaled and non-scaled object
+		FImplicitObjectPtr SphereObjPtr = new FImplicitSphere3(FVec3(0, 0, 5), 1);
+		FImplicitObjectPtr ScaledBoxObjPtr = new TImplicitObjectScaled<FImplicitBox3>(new FImplicitBox3(FVec3(-1), FVec3(1)), FVec3(2, 1, 1));
+		TArray<FImplicitObjectPtr> Objects{ ScaledBoxObjPtr, SphereObjPtr };
+		FImplicitObjectPtr UnionObjectPtr = new FImplicitObjectUnion(MoveTemp(Objects));
+		// Build a transform around the union (this is necessary to produce the bug)
+		FRigidTransform3 ObjTransform(FVec3::Zero(), FRotation3::Identity, FVec3(1));
+		TImplicitObjectTransformed<FReal, 3> TestObject(UnionObjectPtr, ObjTransform);
+
+		// Test a few configurations to make sure we hit and don't hit where expected
+		FSweepResultData Result;
+
+		// To the left of the box
+		Result = SweepQuerySphere(TestObject, FVec3(-3.1, 20, 0), FVec3(0, -1, 0));
+		EXPECT_EQ(Result.bResult, false);
+
+		// Should hit left edge of the box
+		Result = SweepQuerySphere(TestObject, FVec3(-2, 20, 0), FVec3(0, -1, 0));
+		EXPECT_EQ(Result.bResult, true);
+		EXPECT_NEAR(Result.TOI, 18, KINDA_SMALL_NUMBER);
+
+		// Should hit right edge of the box
+		Result = SweepQuerySphere(TestObject, FVec3(2, 20, 0), FVec3(0, -1, 0));
+		EXPECT_EQ(Result.bResult, true);
+		EXPECT_NEAR(Result.TOI, 18, KINDA_SMALL_NUMBER);
+
+		// To the right of the box
+		Result = SweepQuerySphere(TestObject, FVec3(3.1, 20, 0), FVec3(0, -1, 0));
+		EXPECT_EQ(Result.bResult, false);
+
+		// Center of the sphere
+		Result = SweepQuerySphere(TestObject, FVec3(0, 20, 5), FVec3(0, -1, 0));
+		EXPECT_EQ(Result.bResult, true);
+		EXPECT_NEAR(Result.TOI, 18, KINDA_SMALL_NUMBER);
+
+		// To the left of the sphere
+		Result = SweepQuerySphere(TestObject, FVec3(-2.1, 20, 5), FVec3(0, -1, 0));
+		EXPECT_EQ(Result.bResult, false);
+
+		// To the right of the sphere
+		Result = SweepQuerySphere(TestObject, FVec3(2.1, 20, 5), FVec3(0, -1, 0));
+		EXPECT_EQ(Result.bResult, false);
+	}
 }
diff --git a/Engine/Source/Runtime/Experimental/Chaos/Public/Chaos/GeometryQueries.h b/Engine/Source/Runtime/Experimental/Chaos/Public/Chaos/GeometryQueries.h
--- a/Engine/Source/Runtime/Experimental/Chaos/Public/Chaos/GeometryQueries.h
+++ b/Engine/Source/Runtime/Experimental/Chaos/Public/Chaos/GeometryQueries.h
@@ -152,7 +152,6 @@ namespace Chaos
 		}
 	}
 
-	// @todo(chaos): This does not handle Unions
 	template <typename SweptGeometry>
 	bool SweepQuery(const FImplicitObject& A, const FRigidTransform3& ATM, const SweptGeometry& B, const FRigidTransform3& BTM, const FVec3& Dir, const FReal Length, FReal& OutTime, FVec3& OutPosition, FVec3& OutNormal, int32& OutFaceIndex, FVec3& OutFaceNormal, const FReal Thickness, const bool bComputeMTD)
 	{
@@ -178,7 +177,9 @@ namespace Chaos
 		const FVec3 LocalDir = ATM.InverseTransformVectorNoScale(Dir);
 
 		// We can optimize some sweeps as a raycast. Sweeps can have zero length and most raycast functions don't handle this.
-		bool bSweepAsRaycast = BType == ImplicitObjectType::Sphere && !bComputeMTD && Length != 0;
+		// Do not allow unions to do this optimization. A union could contain a non-uniform scaled object which would break the optimization.
+		// Each internal object in the union will do the optimization if possible.
+		bool bSweepAsRaycast = BType == ImplicitObjectType::Sphere && !bComputeMTD && Length != 0 && (AType != ImplicitObjectType::Union && AType != ImplicitObjectType::UnionClustered);
 		if (bSweepAsRaycast && IsScaled(AType))
 		{
 			const auto& Scaled = TImplicitObjectScaledGeneric<FReal, 3>::AsScaledChecked(A);
@@ -233,7 +234,6 @@ namespace Chaos
 		}
 		else
 		{
-			//todo: pass bComputeMTD into these functions
 			switch (AType)
 			{
 				case ImplicitObjectType::HeightField:
