Fix HeadlessChaos compile + run when using clang

diff --git a/Engine/Source/Programs/HeadlessChaos/Private/GeometryCollection/GeometryCollectionTestClustering.cpp b/Engine/Source/Programs/HeadlessChaos/Private/GeometryCollection/GeometryCollectionTestClustering.cpp
--- a/Engine/Source/Programs/HeadlessChaos/Private/GeometryCollection/GeometryCollectionTestClustering.cpp
+++ b/Engine/Source/Programs/HeadlessChaos/Private/GeometryCollection/GeometryCollectionTestClustering.cpp
@@ -707,7 +707,7 @@ namespace GeometryCollectionTest
 			EXPECT_TRUE( (Collection1_PreReleaseTM[Idx].GetTranslation()-Collection1_InitialTM[Idx].GetTranslation()).Size()<KINDA_SMALL_NUMBER);
 		}
 
-		UnitTest.Solver->GetEvolution()->GetRigidClustering().DeactivateClusterParticle({ Collection1->PhysObject->GetParticle_Internal(3) });
+		UnitTest.Solver->GetEvolution()->GetRigidClustering().DeactivateClusterParticle(Collection1->PhysObject->GetParticle_Internal(3));
 		UnitTest.Advance();
 
 		EXPECT_EQ(SovlerParticleHandles.Size(), 4);
@@ -719,7 +719,7 @@ namespace GeometryCollectionTest
 			EXPECT_TRUE((Collection1_PostReleaseTM[Idx].GetTranslation() - Collection1_InitialTM[Idx].GetTranslation()).Size() < KINDA_SMALL_NUMBER);
 		}
 
-		UnitTest.Solver->GetEvolution()->GetRigidClustering().DeactivateClusterParticle({ Collection1->PhysObject->GetParticle_Internal(2) });
+		UnitTest.Solver->GetEvolution()->GetRigidClustering().DeactivateClusterParticle(Collection1->PhysObject->GetParticle_Internal(2));
 		UnitTest.Advance();
 
 		EXPECT_EQ(SovlerParticleHandles.Size(), 4);
diff --git a/Engine/Source/Programs/HeadlessChaos/Private/HeadlessChaosApp.cpp b/Engine/Source/Programs/HeadlessChaos/Private/HeadlessChaosApp.cpp
--- a/Engine/Source/Programs/HeadlessChaos/Private/HeadlessChaosApp.cpp
+++ b/Engine/Source/Programs/HeadlessChaos/Private/HeadlessChaosApp.cpp
@@ -541,7 +541,7 @@ class UEGTestPrinter : public ::testing::EmptyTestEventListener
     }
 };
 
-INT32_MAIN_INT32_ARGC_TCHAR_ARGV()
+int32 tchar_main(int32 ArgC, TCHAR* ArgV[])
 {
 	UE::Geometry::ExactPredicates::GlobalInit();
 
@@ -565,4 +565,42 @@ INT32_MAIN_INT32_ARGC_TCHAR_ARGV()
 	return 0;
 }
 
+// We need to provide a non-wide main otherwise we default to one inside gtest
+// which doesn't correctly initialise the tests.
+int32 main(int32 ArgC, ANSICHAR* Utf8ArgV[])
+{
+	// Allocated storage for TCHAR arguments
+	TCHAR** LocalArgs = new TCHAR* [ArgC];
+
+	// Gtest will directly modify the args array we pass in to the InitGoogleTest call
+	// so we need to pass it a copy of the pointers otherwise we lose track and leak
+	// the allocations for the args
+	TCHAR** PassthroughArgs = new TCHAR* [ArgC];
+
+	for(int32 ArgIndex = 0; ArgIndex < ArgC; ArgIndex++)
+	{
+		FUTF8ToTCHAR ConvertFromUtf8(Utf8ArgV[ArgIndex]);
+		LocalArgs[ArgIndex] = new TCHAR[ConvertFromUtf8.Length() + 1];
+
+		// Copy the pointer to give to Gtest
+		PassthroughArgs[ArgIndex] = LocalArgs[ArgIndex];
+
+		// Fill in the arg
+		FCString::Strncpy(LocalArgs[ArgIndex], ConvertFromUtf8.Get(), ConvertFromUtf8.Length() + 1);
+	}
+
+	// Run our actual test main
+	int32 Result = tchar_main(ArgC, PassthroughArgs);
+
+	for(int32 ArgIndex = 0; ArgIndex < ArgC; ArgIndex++)
+	{
+		delete[] LocalArgs[ArgIndex];
+	}
+
+	delete[] PassthroughArgs;
+	delete[] LocalArgs;
+
+	return Result;
+}
+
 #undef LOCTEXT_NAMESPACE
diff --git a/Engine/Source/Programs/HeadlessChaos/Private/HeadlessChaosTestRewind.cpp b/Engine/Source/Programs/HeadlessChaos/Private/HeadlessChaosTestRewind.cpp
--- a/Engine/Source/Programs/HeadlessChaos/Private/HeadlessChaosTestRewind.cpp
+++ b/Engine/Source/Programs/HeadlessChaos/Private/HeadlessChaosTestRewind.cpp
@@ -1835,7 +1835,7 @@ namespace ChaosTest {
 				const int32 NumPhysSteps = FMath::TruncToInt(LastGameStep / SimDt);
 
 				Solver->SetRewindCallback(MakeUnique<FRewindCallback>(NumPhysSteps, SimDt));
-				FSimCallbackHelper* SimCallback = Solver->CreateAndRegisterSimCallbackObject_External<FSimCallbackHelper>();
+				FSimCallbackHelper* SimCallback = Solver->template CreateAndRegisterSimCallbackObject_External<FSimCallbackHelper>();
 
 				{
 					auto& Particle = Proxy->GetGameThreadAPI();
@@ -1914,7 +1914,7 @@ namespace ChaosTest {
 					}
 				};
 
-				FSimCallbackHelper* SimCallback = Solver->CreateAndRegisterSimCallbackObject_External<FSimCallbackHelper>();
+				FSimCallbackHelper* SimCallback = Solver->template CreateAndRegisterSimCallbackObject_External<FSimCallbackHelper>();
 				SimCallback->Solver = Solver;
 				Solver->SetRewindCallback(MakeUnique<FRewindCallback>(SimCallback));
 
@@ -1995,7 +1995,7 @@ namespace ChaosTest {
 				auto RewindCallback = UniqueRewindCallback.Get();
 
 				Solver->SetRewindCallback(MoveTemp(UniqueRewindCallback));
-				FSimCallbackHelper* SimCallback = Solver->CreateAndRegisterSimCallbackObject_External<FSimCallbackHelper>();
+				FSimCallbackHelper* SimCallback = Solver->template CreateAndRegisterSimCallbackObject_External<FSimCallbackHelper>();
 
 				{
 					auto& Particle = Proxy->GetGameThreadAPI();
diff --git a/Engine/Source/Programs/HeadlessChaos/Public/GeometryCollection/GeometryCollectionTestFramework.h b/Engine/Source/Programs/HeadlessChaos/Public/GeometryCollection/GeometryCollectionTestFramework.h
--- a/Engine/Source/Programs/HeadlessChaos/Public/GeometryCollection/GeometryCollectionTestFramework.h
+++ b/Engine/Source/Programs/HeadlessChaos/Public/GeometryCollection/GeometryCollectionTestFramework.h
@@ -114,6 +114,10 @@ namespace GeometryCollectionTest
 		static WrapperBase* Init(const CreationParameters Params = CreationParameters());
 	};
 
+	template<> WrapperBase* TNewSimulationObject<GeometryType::GeometryCollectionWithSingleRigid>::Init(const CreationParameters);
+	template<> WrapperBase* TNewSimulationObject<GeometryType::RigidFloor>::Init(const CreationParameters);
+	template<> WrapperBase* TNewSimulationObject<GeometryType::GeometryCollectionWithSuppliedRestCollection>::Init(const CreationParameters);
+
 	struct FrameworkParameters
 	{
 		FrameworkParameters() : Dt(1/60.) {}
diff --git a/Engine/Source/Runtime/Experimental/Chaos/Public/Chaos/RigidParticles.h b/Engine/Source/Runtime/Experimental/Chaos/Public/Chaos/RigidParticles.h
--- a/Engine/Source/Runtime/Experimental/Chaos/Public/Chaos/RigidParticles.h
+++ b/Engine/Source/Runtime/Experimental/Chaos/Public/Chaos/RigidParticles.h
@@ -64,13 +64,13 @@ static constexpr int8 ObjectStateBitCount = NumBitsNeeded((int8)EObjectStateType
 // NOTE: not a member class for easier natvis
 struct FRigidParticleCoreData
 {
-	int32 CollisionGroup;							// 4 bytes
-	uint32 CollisionConstraintFlags;				// 4 bytes
-	FRigidParticleControlFlags ControlFlags;		// 2 bytes
-	FRigidParticleTransientFlags TransientFlags;	// 1 byte
-	EObjectStateType ObjectState;					// 1 byte
-	EObjectStateType PreObjectState;				// 1 byte
-	bool bDisabled;									// 1 byte
+	int32 CollisionGroup							= INDEX_NONE;						// 4 bytes
+	uint32 CollisionConstraintFlags					= 0;								// 4 bytes
+	FRigidParticleControlFlags ControlFlags			= {};								// 2 bytes
+	FRigidParticleTransientFlags TransientFlags		= {};								// 1 byte
+	EObjectStateType ObjectState					= EObjectStateType::Uninitialized;	// 1 byte
+	EObjectStateType PreObjectState					= EObjectStateType::Uninitialized;	// 1 byte
+	bool bDisabled									= false;							// 1 byte
 };
 
 template<class T, int d>
