Partial island sleeping (constraint-distance & level-based)

diff --git a/Engine/Source/Programs/HeadlessChaos/Private/HeadlessChaosTestGraphEvolution.cpp b/Engine/Source/Programs/HeadlessChaos/Private/HeadlessChaosTestGraphEvolution.cpp
--- a/Engine/Source/Programs/HeadlessChaos/Private/HeadlessChaosTestGraphEvolution.cpp
+++ b/Engine/Source/Programs/HeadlessChaos/Private/HeadlessChaosTestGraphEvolution.cpp
@@ -9,21 +9,27 @@
 #include "Chaos/PBDRigidsEvolutionGBF.h"
 #include "Chaos/Utilities.h"
 #include "Modules/ModuleManager.h"
+#include "HAL/ConsoleManager.h"
 
 namespace ChaosTest
 {
 	using namespace Chaos;
 
+	class GraphEvolutionTests : public ::testing::TestWithParam<bool>
+	{
+	};
+
 	class FGraphEvolutionTest
 	{
 	public:
-		FGraphEvolutionTest(const int32 NumParticles)
+		FGraphEvolutionTest(const int32 NumParticles, const bool bPartialSleeping)
 			: UniqueIndices()
 			, Particles(UniqueIndices)
 			, PhysicalMaterials()
 			, Evolution(Particles, PhysicalMaterials)
 			, Constraints()
 			, TickCount(0)
+			, CVarPartialSleeping()
 		{
 			IslandManager = &Evolution.GetIslandManager();
 
@@ -39,6 +45,11 @@ namespace ChaosTest
 			{
 				Evolution.EnableParticle(ParticleHandle);
 			}
+
+			CVarPartialSleeping = IConsoleManager::Get().FindConsoleVariable(TEXT("p.Chaos.Solver.Sleep.PartialIslandSleep"), false);
+			check(CVarPartialSleeping && CVarPartialSleeping->IsVariableBool());
+
+			CVarPartialSleeping->Set(bPartialSleeping);
 		}
 
 		// Connect all the particles in a chain
@@ -101,12 +112,14 @@ namespace ChaosTest
 		Private::FPBDIslandManager* IslandManager;
 
 		int32 TickCount;
+
+		IConsoleVariable* CVarPartialSleeping;
 	};
 
 	// Veryify that the Null Constraint mockup is working as intended. We can create the container and constraints, and they are correctly bound to the evolution
-	GTEST_TEST(GraphEvolutionTests, TestNullConstraint)
+	TEST_P(GraphEvolutionTests, TestNullConstraint)
 	{
-		FGraphEvolutionTest Test(4);
+		FGraphEvolutionTest Test(4, GetParam());
 		Test.MakeChain();
 
 		Test.Advance();
@@ -122,9 +135,9 @@ namespace ChaosTest
 	//		Ak - Bd
 	// =>	{}
 	//
-	GTEST_TEST(GraphEvolutionTests, TestConstraintGraph_KinematicDynamic_Remove)
+	TEST_P(GraphEvolutionTests, TestConstraintGraph_KinematicDynamic_Remove)
 	{
-		FGraphEvolutionTest Test(2);
+		FGraphEvolutionTest Test(2, GetParam());
 		Test.MakeChain();
 
 		Test.Evolution.SetParticleObjectState(Test.ParticleHandles[0], EObjectStateType::Kinematic);
@@ -158,9 +171,9 @@ namespace ChaosTest
 	// =>	Ak - Bk
 	// =>	{}
 	//
-	GTEST_TEST(GraphEvolutionTests, TestConstraintGraph_KinematicDynamic_Remove2)
+	TEST_P(GraphEvolutionTests, TestConstraintGraph_KinematicDynamic_Remove2)
 	{
-		FGraphEvolutionTest Test(2);
+		FGraphEvolutionTest Test(2, GetParam());
 		Test.MakeChain();
 
 		Test.Evolution.SetParticleObjectState(Test.ParticleHandles[0], EObjectStateType::Kinematic);
@@ -195,9 +208,9 @@ namespace ChaosTest
 	// This tests a bug where we were not waking a particle if we removed all other particles 
 	// from its island. We now defer island destruction to after sleep handling.
 	//
-	GTEST_TEST(GraphEvolutionTests, TestConstraintGraph_KinematicDynamic_RemoveKinematic)
+	TEST_P(GraphEvolutionTests, TestConstraintGraph_KinematicDynamic_RemoveKinematic)
 	{
-		FGraphEvolutionTest Test(2);
+		FGraphEvolutionTest Test(2, GetParam());
 		Test.MakeChain();
 
 		Test.Evolution.SetParticleObjectState(Test.ParticleHandles[0], EObjectStateType::Kinematic);
@@ -236,7 +249,7 @@ namespace ChaosTest
 		EXPECT_FALSE(Test.ParticleHandles[0]->IsInConstraintGraph());
 		EXPECT_FALSE(Test.ParticleHandles[1]->IsInConstraintGraph());
 		EXPECT_FALSE(Test.ConstraintHandles[0]->IsInConstraintGraph());
-		
+
 		// B should be awake
 		EXPECT_FALSE(Test.ParticleHandles[1]->IsSleeping());
 	}
@@ -247,9 +260,9 @@ namespace ChaosTest
 	//		Ad - Bd - Cd - Dd 
 	// =>	Ad - Bk   Bk - Cd - Dd
 	//
-	GTEST_TEST(GraphEvolutionTests, TestConstraintGraph_ToKinematic)
+	TEST_P(GraphEvolutionTests, TestConstraintGraph_ToKinematic)
 	{
-		FGraphEvolutionTest Test(4);
+		FGraphEvolutionTest Test(4, GetParam());
 		Test.MakeChain();
 
 		Test.Advance();
@@ -281,9 +294,9 @@ namespace ChaosTest
 	// =>	Ad   Bd
 	// =>	Ad - Bd
 	//
-	GTEST_TEST(GraphEvolutionTests, TestConstraintGraph_Invalidate)
+	TEST_P(GraphEvolutionTests, TestConstraintGraph_Invalidate)
 	{
-		FGraphEvolutionTest Test(2);
+		FGraphEvolutionTest Test(2, GetParam());
 		Test.MakeChain();
 
 		Test.Advance();
@@ -311,9 +324,9 @@ namespace ChaosTest
 	}
 
 	// An isolated stationary particle with no gravity should go to sleep
-	GTEST_TEST(GraphEvolutionTests, TestConstraintGraph_ParticleSleep_Isolated)
+	TEST_P(GraphEvolutionTests, TestConstraintGraph_ParticleSleep_Isolated)
 	{
-		FGraphEvolutionTest Test(1);
+		FGraphEvolutionTest Test(1, GetParam());
 		Test.Evolution.GetGravityForces().SetAcceleration(FVec3(0), 0);
 
 		// Make all the particles sleep
@@ -326,9 +339,9 @@ namespace ChaosTest
 
 	// Wait for all particles to go to sleep naturally (i.e., as part of the tick and not by explicitly setting the state)
 	// then check that the islands are preserved.
-	GTEST_TEST(GraphEvolutionTests, TestConstraintGraph_ParticleSleep_Natural)
+	TEST_P(GraphEvolutionTests, TestConstraintGraph_ParticleSleep_Natural)
 	{
-		FGraphEvolutionTest Test(4);
+		FGraphEvolutionTest Test(4, GetParam());
 		Test.MakeChain();
 
 		Test.Advance();
@@ -351,9 +364,9 @@ namespace ChaosTest
 	// Sleep all particles in the scene and ensure that the island manager puts the island to sleep
 	// but retains all the constraints and particles in the island.
 	//
-	GTEST_TEST(GraphEvolutionTests, TestConstraintGraph_ParticleSleep_Manual)
+	TEST_P(GraphEvolutionTests, TestConstraintGraph_ParticleSleep_Manual)
 	{
-		FGraphEvolutionTest Test(4);
+		FGraphEvolutionTest Test(4, GetParam());
 		Test.MakeChain();
 
 		Test.Advance();
@@ -385,9 +398,9 @@ namespace ChaosTest
 	//		A-B-C-D 
 	// =>	A-B  C-D
 	//
-	GTEST_TEST(GraphEvolutionTests, TestConstraintGraph_KinematicKinematic)
+	TEST_P(GraphEvolutionTests, TestConstraintGraph_KinematicKinematic)
 	{
-		FGraphEvolutionTest Test(4);
+		FGraphEvolutionTest Test(4, GetParam());
 		Test.MakeChain();
 
 		Test.Advance();
@@ -421,9 +434,9 @@ namespace ChaosTest
 	// =>	As - Bk   Ck - Ds
 	// =>	Ad - Bk   Ck - Ds
 	//
-	GTEST_TEST(GraphEvolutionTests, TestConstraintGraph_KinematicKinematic_Sleeping)
+	TEST_P(GraphEvolutionTests, TestConstraintGraph_KinematicKinematic_Sleeping)
 	{
-		FGraphEvolutionTest Test(4);
+		FGraphEvolutionTest Test(4, GetParam());
 		Test.MakeChain();
 
 		Test.Advance();
@@ -480,9 +493,9 @@ namespace ChaosTest
 	//  \   |   /	   =>     \   |   /
 	//      Ak		              Ad
 	//
-	GTEST_TEST(GraphEvolutionTests, TestConstraintGraph_KinematicToDynamic)
+	TEST_P(GraphEvolutionTests, TestConstraintGraph_KinematicToDynamic)
 	{
-		FGraphEvolutionTest Test(4);
+		FGraphEvolutionTest Test(4, GetParam());
 		Test.MakeFloor();
 
 		Test.Advance();
@@ -507,6 +520,47 @@ namespace ChaosTest
 		EXPECT_EQ(Test.IslandManager->FindParticleIslands(Test.ParticleHandles[3]).Num(), 1);
 	}
 
+	// 3 objects sat on the floor asleep. Move the kinematic floor by setting its transform.
+	// This tests what happens when a non-moving kinematic in multiple islands starts moving.
+	// 
+	// (d=dynamic, s=sleeping, k=kinematic, km=kinematic, moving)
+	// Bs   Cs   Ds          Bd   Cd   Dd
+	//  \   |   /	   =>     \   |   /
+	//      Ak		              Akm
+	//
+	TEST_P(GraphEvolutionTests, TestConstraintGraph_MoveKinematicFloor)
+	{
+		FGraphEvolutionTest Test(4, GetParam());
+		Test.MakeFloor();
+
+		Test.AdvanceUntilSleeping();
+
+		// Each particle in its own island (kinematic will be in all 3)
+		EXPECT_EQ(Test.IslandManager->GetNumIslands(), 3);
+		EXPECT_EQ(Test.IslandManager->FindParticleIslands(Test.ParticleHandles[0]).Num(), 3);
+		EXPECT_EQ(Test.IslandManager->FindParticleIslands(Test.ParticleHandles[1]).Num(), 1);
+		EXPECT_EQ(Test.IslandManager->FindParticleIslands(Test.ParticleHandles[2]).Num(), 1);
+		EXPECT_EQ(Test.IslandManager->FindParticleIslands(Test.ParticleHandles[3]).Num(), 1);
+
+		// Teleport a particle by setting its transform explicitly
+		const bool bIsTeleport = true;
+		Test.Evolution.SetParticleKinematicTarget(Test.ParticleHandles[0], FKinematicTarget::MakePositionTarget(FVec3(0, 3, 0), FRotation3()));
+
+		Test.Advance();
+
+		// Each particle in its own island (kinematic will be in all 3)
+		EXPECT_EQ(Test.IslandManager->GetNumIslands(), 3);
+		EXPECT_EQ(Test.IslandManager->FindParticleIslands(Test.ParticleHandles[0]).Num(), 3);
+		EXPECT_EQ(Test.IslandManager->FindParticleIslands(Test.ParticleHandles[1]).Num(), 1);
+		EXPECT_EQ(Test.IslandManager->FindParticleIslands(Test.ParticleHandles[2]).Num(), 1);
+		EXPECT_EQ(Test.IslandManager->FindParticleIslands(Test.ParticleHandles[3]).Num(), 1);
+
+		// All particles awake
+		EXPECT_FALSE(FConstGenericParticleHandle(Test.ParticleHandles[1])->Sleeping());
+		EXPECT_FALSE(FConstGenericParticleHandle(Test.ParticleHandles[2])->Sleeping());
+		EXPECT_FALSE(FConstGenericParticleHandle(Test.ParticleHandles[3])->Sleeping());
+	}
+
 	// 3 objects sat on the floor asleep. Make the floor dynamic.
 	// This tests what happens when a kinematic in multiple sleeping islands gets converted to a dynamic.
 	// 
@@ -515,9 +569,9 @@ namespace ChaosTest
 	//  \   |   /	   =>     \   |   /
 	//      Ak		              Ad
 	//
-	GTEST_TEST(GraphEvolutionTests, TestConstraintGraph_KinematicToDynamic_WithSleep1)
+	TEST_P(GraphEvolutionTests, TestConstraintGraph_KinematicToDynamic_WithSleep1)
 	{
-		FGraphEvolutionTest Test(4);
+		FGraphEvolutionTest Test(4, GetParam());
 		Test.MakeFloor();
 
 		Test.AdvanceUntilSleeping();
@@ -558,9 +612,9 @@ namespace ChaosTest
 	//  \   |   /      =>     \   |   /
 	//      Ak                    As
 	//
-	GTEST_TEST(GraphEvolutionTests, TestConstraintGraph_KinematicToDynamic_WithSleep2)
+	TEST_P(GraphEvolutionTests, TestConstraintGraph_KinematicToDynamic_WithSleep2)
 	{
-		FGraphEvolutionTest Test(4);
+		FGraphEvolutionTest Test(4, GetParam());
 		Test.MakeFloor();
 
 		Test.AdvanceUntilSleeping();
@@ -596,14 +650,19 @@ namespace ChaosTest
 	// 3 Objects sat on the floor, 2 asleep and 1 awake. Make the floor dynamic and asleep.
 	// In this case we should get 1 awake island and all particles should wake.
 	// 
-	// (d=dynamic, s=sleeping, k=kinematic)
+	// Island sleeping: (d=dynamic, s=sleeping, k=kinematic)
 	// Bs   Cs   Dd          Bd   Cd   Dd
 	//  \   |   /      =>     \   |   /
 	//      Ak	                  Ad
 	//
-	GTEST_TEST(GraphEvolutionTests, TestConstraintGraph_KinematicToDynamic_WithSleep3)
+	// Partial island sleeping: (d=dynamic, s=sleeping, k=kinematic)
+	// Bs   Cs   Dd          Bs   Cs   Dd
+	//  \   |   /      =>     \   |   /
+	//      Ak	                  As
+	//
+	TEST_P(GraphEvolutionTests, TestConstraintGraph_KinematicToDynamic_WithSleep3)
 	{
-		FGraphEvolutionTest Test(4);
+		FGraphEvolutionTest Test(4, GetParam());
 		Test.MakeFloor();
 
 		Test.AdvanceUntilSleeping();
@@ -629,19 +688,38 @@ namespace ChaosTest
 
 		Test.Advance();
 
-		// All particles in one awake island (D was awake so it would wake the island)
-		EXPECT_EQ(Test.IslandManager->GetNumIslands(), 1);
-		EXPECT_FALSE(Test.IslandManager->GetIsland(0)->IsSleeping());
-		EXPECT_EQ(Test.IslandManager->FindParticleIslands(Test.ParticleHandles[0]).Num(), 1);
-		EXPECT_EQ(Test.IslandManager->FindParticleIslands(Test.ParticleHandles[1]).Num(), 1);
-		EXPECT_EQ(Test.IslandManager->FindParticleIslands(Test.ParticleHandles[2]).Num(), 1);
-		EXPECT_EQ(Test.IslandManager->FindParticleIslands(Test.ParticleHandles[3]).Num(), 1);
+		if (!Test.CVarPartialSleeping->GetBool()) // island sleeping
+		{
+			// All particles in one awake island (D was awake so it would wake the island)
+			EXPECT_EQ(Test.IslandManager->GetNumIslands(), 1);
+			EXPECT_FALSE(Test.IslandManager->GetIsland(0)->IsSleeping());
+			EXPECT_EQ(Test.IslandManager->FindParticleIslands(Test.ParticleHandles[0]).Num(), 1);
+			EXPECT_EQ(Test.IslandManager->FindParticleIslands(Test.ParticleHandles[1]).Num(), 1);
+			EXPECT_EQ(Test.IslandManager->FindParticleIslands(Test.ParticleHandles[2]).Num(), 1);
+			EXPECT_EQ(Test.IslandManager->FindParticleIslands(Test.ParticleHandles[3]).Num(), 1);
+
+			// All particles awake
+			EXPECT_FALSE(FConstGenericParticleHandle(Test.ParticleHandles[0])->Sleeping());
+			EXPECT_FALSE(FConstGenericParticleHandle(Test.ParticleHandles[1])->Sleeping());
+			EXPECT_FALSE(FConstGenericParticleHandle(Test.ParticleHandles[2])->Sleeping());
+			EXPECT_FALSE(FConstGenericParticleHandle(Test.ParticleHandles[3])->Sleeping());
+		}
+		else // partial island sleeping
+		{
+			// All particles are asleep except for D which remains awake
+			EXPECT_EQ(Test.IslandManager->GetNumIslands(), 1);
+			EXPECT_FALSE(Test.IslandManager->GetIsland(0)->IsSleeping());
+			EXPECT_EQ(Test.IslandManager->FindParticleIslands(Test.ParticleHandles[0]).Num(), 1);
+			EXPECT_EQ(Test.IslandManager->FindParticleIslands(Test.ParticleHandles[1]).Num(), 1);
+			EXPECT_EQ(Test.IslandManager->FindParticleIslands(Test.ParticleHandles[2]).Num(), 1);
+			EXPECT_EQ(Test.IslandManager->FindParticleIslands(Test.ParticleHandles[3]).Num(), 1);
 
-		// All particles awake
-		EXPECT_FALSE(FConstGenericParticleHandle(Test.ParticleHandles[0])->Sleeping());
-		EXPECT_FALSE(FConstGenericParticleHandle(Test.ParticleHandles[1])->Sleeping());
-		EXPECT_FALSE(FConstGenericParticleHandle(Test.ParticleHandles[2])->Sleeping());
-		EXPECT_FALSE(FConstGenericParticleHandle(Test.ParticleHandles[3])->Sleeping());
+			// All particles except for D asleep
+			EXPECT_TRUE(FConstGenericParticleHandle(Test.ParticleHandles[0])->Sleeping());
+			EXPECT_TRUE(FConstGenericParticleHandle(Test.ParticleHandles[1])->Sleeping());
+			EXPECT_TRUE(FConstGenericParticleHandle(Test.ParticleHandles[2])->Sleeping());
+			EXPECT_FALSE(FConstGenericParticleHandle(Test.ParticleHandles[3])->Sleeping());
+		}
 	}
 
 	// 3 Objects sat on the floor, 2 asleep and 1 awake. Make the floor dynamic and asleep.
@@ -649,14 +727,19 @@ namespace ChaosTest
 	// be sure we weren't just lucky above (when we make a kinematic into a dynamic we add it to one
 	// of the islands it is in. This is testing that this is ok).
 	// 
-	// (d=dynamic, s=sleeping, k=kinematic)
+	// Island sleeping (d=dynamic, s=sleeping, k=kinematic)
 	// Bd   Cs   Ds          Bd   Cd   Dd
 	//  \   |   /      =>     \   |   /
 	//      Ak	                  Ad
 	//
-	GTEST_TEST(GraphEvolutionTests, TestConstraintGraph_KinematicToDynamic_WithSleep4)
+	// Partial island sleeping (d=dynamic, s=sleeping, k=kinematic)
+	// Bd   Cs   Ds          Bd   Cs   Ds
+	//  \   |   /      =>     \   |   /
+	//      Ak	                  As
+	//
+	TEST_P(GraphEvolutionTests, TestConstraintGraph_KinematicToDynamic_WithSleep4)
 	{
-		FGraphEvolutionTest Test(4);
+		FGraphEvolutionTest Test(4, GetParam());
 		Test.MakeFloor();
 
 		Test.AdvanceUntilSleeping();
@@ -682,19 +765,38 @@ namespace ChaosTest
 
 		Test.Advance();
 
-		// All particles in one awake island (B was awake so it would wake the island, even though A was set to sleeping)
-		EXPECT_EQ(Test.IslandManager->GetNumIslands(), 1);
-		EXPECT_FALSE(Test.IslandManager->GetIsland(0)->IsSleeping());
-		EXPECT_EQ(Test.IslandManager->FindParticleIslands(Test.ParticleHandles[0]).Num(), 1);
-		EXPECT_EQ(Test.IslandManager->FindParticleIslands(Test.ParticleHandles[1]).Num(), 1);
-		EXPECT_EQ(Test.IslandManager->FindParticleIslands(Test.ParticleHandles[2]).Num(), 1);
-		EXPECT_EQ(Test.IslandManager->FindParticleIslands(Test.ParticleHandles[3]).Num(), 1);
+		if (!Test.CVarPartialSleeping->GetBool()) // island sleeping
+		{
+			// All particles in one awake island (B was awake so it would wake the island)
+			EXPECT_EQ(Test.IslandManager->GetNumIslands(), 1);
+			EXPECT_FALSE(Test.IslandManager->GetIsland(0)->IsSleeping());
+			EXPECT_EQ(Test.IslandManager->FindParticleIslands(Test.ParticleHandles[0]).Num(), 1);
+			EXPECT_EQ(Test.IslandManager->FindParticleIslands(Test.ParticleHandles[1]).Num(), 1);
+			EXPECT_EQ(Test.IslandManager->FindParticleIslands(Test.ParticleHandles[2]).Num(), 1);
+			EXPECT_EQ(Test.IslandManager->FindParticleIslands(Test.ParticleHandles[3]).Num(), 1);
+
+			// All particles awake
+			EXPECT_FALSE(FConstGenericParticleHandle(Test.ParticleHandles[0])->Sleeping());
+			EXPECT_FALSE(FConstGenericParticleHandle(Test.ParticleHandles[1])->Sleeping());
+			EXPECT_FALSE(FConstGenericParticleHandle(Test.ParticleHandles[2])->Sleeping());
+			EXPECT_FALSE(FConstGenericParticleHandle(Test.ParticleHandles[3])->Sleeping());
+		}
+		else // partial island sleeping
+		{
+			// All particles are asleep except for B which remains awake
+			EXPECT_EQ(Test.IslandManager->GetNumIslands(), 1);
+			EXPECT_FALSE(Test.IslandManager->GetIsland(0)->IsSleeping());
+			EXPECT_EQ(Test.IslandManager->FindParticleIslands(Test.ParticleHandles[0]).Num(), 1);
+			EXPECT_EQ(Test.IslandManager->FindParticleIslands(Test.ParticleHandles[1]).Num(), 1);
+			EXPECT_EQ(Test.IslandManager->FindParticleIslands(Test.ParticleHandles[2]).Num(), 1);
+			EXPECT_EQ(Test.IslandManager->FindParticleIslands(Test.ParticleHandles[3]).Num(), 1);
 
-		// All particles awake
-		EXPECT_FALSE(FConstGenericParticleHandle(Test.ParticleHandles[0])->Sleeping());
-		EXPECT_FALSE(FConstGenericParticleHandle(Test.ParticleHandles[1])->Sleeping());
-		EXPECT_FALSE(FConstGenericParticleHandle(Test.ParticleHandles[2])->Sleeping());
-		EXPECT_FALSE(FConstGenericParticleHandle(Test.ParticleHandles[3])->Sleeping());
+			// All particles except for B asleep
+			EXPECT_TRUE(FConstGenericParticleHandle(Test.ParticleHandles[0])->Sleeping());
+			EXPECT_FALSE(FConstGenericParticleHandle(Test.ParticleHandles[1])->Sleeping());
+			EXPECT_TRUE(FConstGenericParticleHandle(Test.ParticleHandles[2])->Sleeping());
+			EXPECT_TRUE(FConstGenericParticleHandle(Test.ParticleHandles[3])->Sleeping());
+		}
 	}
 
 	// Start with an island containing 4 awake particles connected in a chain. Then sleep the island
@@ -703,9 +805,9 @@ namespace ChaosTest
 	//		(d=dynamic, s=sleeping, k=kinematic)
 	//		Ad - Bd - Cd - Dd  =>	As - Bs - Cs - Ds
 	//
-	GTEST_TEST(GraphEvolutionTests, TestConstraintGraph_SleepIsland)
+	TEST_P(GraphEvolutionTests, TestConstraintGraph_SleepIsland)
 	{
-		FGraphEvolutionTest Test(4);
+		FGraphEvolutionTest Test(4, GetParam());
 		Test.MakeChain();
 
 		Test.Advance();
@@ -753,7 +855,7 @@ namespace ChaosTest
 	// explicitly put to sleep. However that bug was a result of how collision constraints are 
 	// destroyed (i.e., when they are not updated this tick) and the NullConstraints don't have that same 
 	// behaviour. We need a unit testing constraint that can reproduce that behaviour. 
-	GTEST_TEST(GraphEvolutionTests, DISABLED_TestConstraintGraph_SleepIsland_Collisions)
+	TEST_P(GraphEvolutionTests, DISABLED_TestConstraintGraph_SleepIsland_Collisions)
 	{
 	}
 
@@ -763,9 +865,9 @@ namespace ChaosTest
 	//		(d=dynamic, s=sleeping, k=kinematic)
 	//		Ak - Bs    =>    Ak - Bs
 	//
-	GTEST_TEST(GraphEvolutionTests, TestConstraintGraph_SleepingKinematicConstraint)
+	TEST_P(GraphEvolutionTests, TestConstraintGraph_SleepingKinematicConstraint)
 	{
-		FGraphEvolutionTest Test(2);
+		FGraphEvolutionTest Test(2, GetParam());
 
 		// Make A kinematic, B sleeping
 		Test.Evolution.SetParticleObjectState(Test.ParticleHandles[0], EObjectStateType::Kinematic);
@@ -794,9 +896,9 @@ namespace ChaosTest
 	//		(d=dynamic, s=sleeping, k=kinematic)
 	//		Ak  Bs    =>    Ak - Bs
 	//
-	GTEST_TEST(GraphEvolutionTests, TestConstraintGraph_SleepingKinematicConstraint2)
+	TEST_P(GraphEvolutionTests, TestConstraintGraph_SleepingKinematicConstraint2)
 	{
-		FGraphEvolutionTest Test(2);
+		FGraphEvolutionTest Test(2, GetParam());
 
 		// Make A kinematic, B sleeping
 		Test.Evolution.SetParticleObjectState(Test.ParticleHandles[0], EObjectStateType::Kinematic);
@@ -826,9 +928,9 @@ namespace ChaosTest
 	//		(d=dynamic, s=sleeping, k=kinematic)
 	//		As  Bs    =>    As - Bs
 	//
-	GTEST_TEST(GraphEvolutionTests, TestConstraintGraph_SleepingSleepingConstraint)
+	TEST_P(GraphEvolutionTests, TestConstraintGraph_SleepingSleepingConstraint)
 	{
-		FGraphEvolutionTest Test(2);
+		FGraphEvolutionTest Test(2, GetParam());
 
 		// Make A and B sleeping
 		Test.Evolution.SetParticleObjectState(Test.ParticleHandles[0], EObjectStateType::Sleeping);
@@ -861,9 +963,9 @@ namespace ChaosTest
 	//		Ak - Bs - Cs  =>    Ak - Bs - Cs
 	//		                     ^--------^
 	//
-	GTEST_TEST(GraphEvolutionTests, TestConstraintGraph_SleepingKinematicConstraint_SameIsland)
+	TEST_P(GraphEvolutionTests, TestConstraintGraph_SleepingKinematicConstraint_SameIsland)
 	{
-		FGraphEvolutionTest Test(3);
+		FGraphEvolutionTest Test(3, GetParam());
 
 		// Chains the particles and make the first one kinematic
 		Test.Evolution.SetParticleObjectState(Test.ParticleHandles[0], EObjectStateType::Kinematic);
@@ -898,9 +1000,9 @@ namespace ChaosTest
 	//		(d=dynamic, s=sleeping, k=kinematic)
 	//		As - Bs   Cs - Ds  =>    As - Bs - Cs - Ds
 	//
-	GTEST_TEST(GraphEvolutionTests, TestConstraintGraph_SleepingSleepingConstraint_MergeIslands)
+	TEST_P(GraphEvolutionTests, TestConstraintGraph_SleepingSleepingConstraint_MergeIslands)
 	{
-		FGraphEvolutionTest Test(4);
+		FGraphEvolutionTest Test(4, GetParam());
 
 		// Add constraints A-B and C-D
 		Test.ConstraintHandles.Add(Test.Constraints.AddConstraint({ Test.ParticleHandles[0], Test.ParticleHandles[1] }));
@@ -942,10 +1044,10 @@ namespace ChaosTest
 	// sleep counts to the graph when we add a constraint to them. We could fix this
 	// but probably not worth worrying about
 	//
-	GTEST_TEST(GraphEvolutionTests, DISABLED_TestConstraintGraph_SleepingSleepingConstraint_Timing_Isolated)
+	TEST_P(GraphEvolutionTests, DISABLED_TestConstraintGraph_SleepingSleepingConstraint_Timing_Isolated)
 	{
 		// Count how many frames it takes the simulation to sleep
-		FGraphEvolutionTest SleepTest(2);
+		FGraphEvolutionTest SleepTest(2, GetParam());
 		SleepTest.AdvanceUntilSleeping();
 		const int32 TicksToSleep = SleepTest.TickCount;
 
@@ -953,7 +1055,7 @@ namespace ChaosTest
 		// Verify that adding a constraint on that tick leaves the scene as expected
 		for (int32 SleepRelativeTickCount = -1; SleepRelativeTickCount < 2; ++SleepRelativeTickCount)
 		{
-			FGraphEvolutionTest Test(2);
+			FGraphEvolutionTest Test(2, GetParam());
 			for (int32 Frame = 0; Frame < TicksToSleep + SleepRelativeTickCount; ++Frame)
 			{
 				Test.Advance();
@@ -990,12 +1092,12 @@ namespace ChaosTest
 	// would fail because we retain SleepCounter when merging islands (but not when adding isloated
 	// particles that have their own sleep counter)
 	//
-	GTEST_TEST(GraphEvolutionTests, TestConstraintGraph_SleepingSleepingConstraint_Timing)
+	TEST_P(GraphEvolutionTests, TestConstraintGraph_SleepingSleepingConstraint_Timing)
 	{
 		// Count how many frames it takes the simulation to sleep
 		int32 TicksToSleep = 0;
 		{
-			FGraphEvolutionTest SleepTest(4);
+			FGraphEvolutionTest SleepTest(4, GetParam());
 			SleepTest.ConstraintHandles.Add(SleepTest.Constraints.AddConstraint({ SleepTest.ParticleHandles[0], SleepTest.ParticleHandles[1] }));
 			SleepTest.ConstraintHandles.Add(SleepTest.Constraints.AddConstraint({ SleepTest.ParticleHandles[2], SleepTest.ParticleHandles[3] }));
 
@@ -1007,7 +1109,7 @@ namespace ChaosTest
 		// Verify that adding a constraint on that tick leaves the scene as expected
 		for (int32 SleepRelativeTickCount = -1; SleepRelativeTickCount < 2; ++SleepRelativeTickCount)
 		{
-			FGraphEvolutionTest Test(4);
+			FGraphEvolutionTest Test(4, GetParam());
 			Test.ConstraintHandles.Add(Test.Constraints.AddConstraint({ Test.ParticleHandles[0], Test.ParticleHandles[1] }));
 			Test.ConstraintHandles.Add(Test.Constraints.AddConstraint({ Test.ParticleHandles[2], Test.ParticleHandles[3] }));
 
@@ -1067,9 +1169,9 @@ namespace ChaosTest
 	//
 	// NOTE: the transition to sleep is by a user call, not the automatic sleep-when-not-moving system
 	//
-	GTEST_TEST(GraphEvolutionTests, TestConstraintGraph_SleepingSleepingConstraint_Timing2)
+	TEST_P(GraphEvolutionTests, TestConstraintGraph_SleepingSleepingConstraint_Timing2)
 	{
-		FGraphEvolutionTest Test(2);
+		FGraphEvolutionTest Test(2, GetParam());
 
 		// Make A kinematic, B dynamic
 		Test.Evolution.SetParticleObjectState(Test.ParticleHandles[0], EObjectStateType::Kinematic);
@@ -1104,9 +1206,9 @@ namespace ChaosTest
 	// 
 	// NOTE: the transition to sleep is by a user call, not the automatic sleep-when-not-moving system
 	//
-	GTEST_TEST(GraphEvolutionTests, TestConstraintGraph_SleepingSleepingConstraint_Timing3)
+	TEST_P(GraphEvolutionTests, TestConstraintGraph_SleepingSleepingConstraint_Timing3)
 	{
-		FGraphEvolutionTest Test(2);
+		FGraphEvolutionTest Test(2, GetParam());
 
 		// Make A, B dynamic
 		Test.Evolution.SetParticleObjectState(Test.ParticleHandles[0], EObjectStateType::Dynamic);
@@ -1136,10 +1238,10 @@ namespace ChaosTest
 	}
 
 	// Test isolated particles are not present in the graph
-	GTEST_TEST(GraphEvolutionTests, TestConstraintGraph_KinematicRemoveFromGraph)
+	TEST_P(GraphEvolutionTests, TestConstraintGraph_KinematicRemoveFromGraph)
 	{
 		// Create a scene with 3 dynamic particles
-		FGraphEvolutionTest Test(3);
+		FGraphEvolutionTest Test(3, GetParam());
 
 		Test.Advance();
 
@@ -1165,9 +1267,9 @@ namespace ChaosTest
 	// If a kinematic is being animated by velocity or by setting a target
 	// position the island should wake but only if the target velocity is
 	// non-zero or the target transform is different from the identity
-	GTEST_TEST(GraphEvolutionTests, TestConstraintGraph_KinematicWakeIslandConditions)
+	TEST_P(GraphEvolutionTests, TestConstraintGraph_KinematicWakeIslandConditions)
 	{
-		FGraphEvolutionTest Test(4);
+		FGraphEvolutionTest Test(4, GetParam());
 		Test.MakeChain();
 
 		// Set the root of the chain to be kinematic and the rest to be sleeping
@@ -1292,9 +1394,9 @@ namespace ChaosTest
 	//		Ad - Bd   Cd - Dd
 	// =>	Ad - Bd - Cd   Dd
 	//
-	GTEST_TEST(GraphEvolutionTests, TestConstraintGraph_IslandMerge_EnableDisable)
+	TEST_P(GraphEvolutionTests, TestConstraintGraph_IslandMerge_EnableDisable)
 	{
-		FGraphEvolutionTest Test(4);
+		FGraphEvolutionTest Test(4, GetParam());
 
 		// Create constraints in a chain but disable the middle constraint so we have two islands {A-B} and {C-D}
 		Test.MakeChain();
@@ -1363,7 +1465,7 @@ namespace ChaosTest
 
 		// Now do the same as above on a new array, but call SortFreeList before reusing it
 		TSparseArray<int32> Values2;
-		
+
 		Values2.Add(0);
 		Values2.Add(1);
 		Values2.Add(2);
@@ -1380,4 +1482,679 @@ namespace ChaosTest
 		EXPECT_EQ(Values2[1], 1);
 		EXPECT_EQ(Values2[2], 2);
 	}
+
+
+	// Test validating that a sleeping island wakes up when teleporting one of its particles.
+	// Teleporting means explicitly updating the transform of the particle.
+	TEST_P(GraphEvolutionTests, TestConstraintGraph_TeleportSleeping)
+	{
+		FGraphEvolutionTest Test(3, GetParam());
+		Test.MakeChain();
+		Test.Evolution.SetParticleObjectState(Test.ParticleHandles[0], EObjectStateType::Kinematic);
+
+		Test.AdvanceUntilSleeping();
+
+		// Should have 1 island and it should be asleep
+		EXPECT_EQ(Test.IslandManager->GetNumIslands(), 1);
+		EXPECT_TRUE(Test.IslandManager->GetIsland(0)->IsSleeping());
+		EXPECT_TRUE(Test.ParticleHandles[0]->IsInConstraintGraph());
+		EXPECT_TRUE(Test.ParticleHandles[1]->IsInConstraintGraph());
+		EXPECT_TRUE(Test.ParticleHandles[2]->IsInConstraintGraph());
+		EXPECT_TRUE(FConstGenericParticleHandle(Test.ParticleHandles[0])->IsKinematic());
+		EXPECT_TRUE(FConstGenericParticleHandle(Test.ParticleHandles[1])->Sleeping());
+		EXPECT_TRUE(FConstGenericParticleHandle(Test.ParticleHandles[2])->Sleeping());
+
+		EXPECT_TRUE(Test.ConstraintHandles[0]->IsInConstraintGraph());
+		EXPECT_TRUE(Test.ConstraintHandles[1]->IsInConstraintGraph());
+		EXPECT_TRUE(Test.ConstraintHandles[0]->IsSleeping());
+		EXPECT_TRUE(Test.ConstraintHandles[1]->IsSleeping());
+
+		// Teleport a particle by setting its transform explicitly
+		const bool bIsTeleport = true;
+		const bool bWakeUp = true;
+		Test.Evolution.SetParticleTransform(Test.ParticleHandles[2], FVec3(0, 3, 0), FRotation3(), bIsTeleport, bWakeUp);
+
+		Test.Advance();
+
+		// Should wake up the entire island
+		EXPECT_EQ(Test.IslandManager->GetNumIslands(), 1);
+		EXPECT_FALSE(Test.IslandManager->GetIsland(0)->IsSleeping());
+		EXPECT_FALSE(FConstGenericParticleHandle(Test.ParticleHandles[1])->Sleeping());
+		EXPECT_FALSE(FConstGenericParticleHandle(Test.ParticleHandles[2])->Sleeping());
+		EXPECT_FALSE(Test.ConstraintHandles[0]->IsSleeping());
+		EXPECT_FALSE(Test.ConstraintHandles[1]->IsSleeping());
+	}
+
+	// Test validating that a sleeping island wakes up when adding an impulse, acceleration, etc. to one particle.
+	// NOTE: This code emulates what happens if dynamics data is updated in @PushToPhysicsStateImp (SingleParticlePhysicsProxy.cpp).
+	TEST_P(GraphEvolutionTests, TestConstraintGraph_SetVelocityOfSleeping)
+	{
+		FGraphEvolutionTest Test(3, GetParam());
+		Test.MakeChain();
+		Test.Evolution.SetParticleObjectState(Test.ParticleHandles[0], EObjectStateType::Kinematic);
+
+		Test.AdvanceUntilSleeping();
+
+		// Should have 1 island and it should be asleep
+		EXPECT_EQ(Test.IslandManager->GetNumIslands(), 1);
+		EXPECT_TRUE(Test.IslandManager->GetIsland(0)->IsSleeping());
+		EXPECT_TRUE(Test.ParticleHandles[0]->IsInConstraintGraph());
+		EXPECT_TRUE(Test.ParticleHandles[1]->IsInConstraintGraph());
+		EXPECT_TRUE(Test.ParticleHandles[2]->IsInConstraintGraph());
+		EXPECT_TRUE(FConstGenericParticleHandle(Test.ParticleHandles[0])->IsKinematic());
+		EXPECT_TRUE(FConstGenericParticleHandle(Test.ParticleHandles[1])->Sleeping());
+		EXPECT_TRUE(FConstGenericParticleHandle(Test.ParticleHandles[2])->Sleeping());
+
+		EXPECT_TRUE(Test.ConstraintHandles[0]->IsInConstraintGraph());
+		EXPECT_TRUE(Test.ConstraintHandles[1]->IsInConstraintGraph());
+		EXPECT_TRUE(Test.ConstraintHandles[0]->IsSleeping());
+		EXPECT_TRUE(Test.ConstraintHandles[1]->IsSleeping());
+
+		// Explicitly set a non-zero velocity
+		FParticleDynamics Dynamics;
+		Dynamics.SetAcceleration(FVec3(0, 0, 0));
+		Dynamics.SetAngularAcceleration(FVec3(0, 0, 0));
+		Dynamics.SetLinearImpulseVelocity(FVec3(50, 50, 0));
+		Dynamics.SetAngularImpulseVelocity(FVec3(0, 0, 0));
+
+		// This emulates what happens in @PushToPhysicsStateImp (SingleParticlePhysicsProxy.cpp).
+		Test.ParticleHandles[2]->SetDynamics(Dynamics);
+		Test.Evolution.ResetVSmoothFromForces(*Test.ParticleHandles[2]);
+		Test.Evolution.SetParticleObjectState(Test.ParticleHandles[2], EObjectStateType::Dynamic);
+
+		Test.Advance();
+
+		// Should wake up the entire island
+		EXPECT_EQ(Test.IslandManager->GetNumIslands(), 1);
+		EXPECT_FALSE(Test.IslandManager->GetIsland(0)->IsSleeping());
+		EXPECT_FALSE(FConstGenericParticleHandle(Test.ParticleHandles[1])->Sleeping());
+		EXPECT_FALSE(FConstGenericParticleHandle(Test.ParticleHandles[2])->Sleeping());
+		EXPECT_FALSE(Test.ConstraintHandles[0]->IsSleeping());
+		EXPECT_FALSE(Test.ConstraintHandles[1]->IsSleeping());
+	}
+
+	// Test validating the wake-up propagation throughout a sleeping island if adding an awake particle to the top of the stack.
+	// 1) For island sleeping, the entire island should wake up in a single tick.
+	// (d=dynamic, s=sleeping, k=kinematic)
+	//		Ak - Bs - Cs - Ds
+	// =>	Ak - Bd - Cd - Dd - Ed
+	// 2) For partial island sleeping, the awake particle will be added to the island and only wake the particle closest to it. 
+	// (d=dynamic, s=sleeping, k=kinematic)
+	//		Ak - Bs - Cs - Ds
+	// =>	Ak - Bs - Cs - Dd - Ed
+	//
+	TEST_P(GraphEvolutionTests, TestConstraintGraph_WakeUpPropagation_AddToStackTop)
+	{
+		FGraphEvolutionTest Test(4, GetParam());
+
+		if (Test.CVarPartialSleeping->GetBool())
+		{
+			Test.MakeChain();
+			Test.Evolution.SetParticleObjectState(Test.ParticleHandles[0], EObjectStateType::Kinematic);
+
+			Test.AdvanceUntilSleeping();
+
+			// Should have 1 island and it should be asleep
+			EXPECT_EQ(Test.IslandManager->GetNumIslands(), 1);
+			EXPECT_TRUE(Test.IslandManager->GetIsland(0)->IsSleeping());
+			EXPECT_TRUE(Test.ParticleHandles[0]->IsInConstraintGraph());
+			EXPECT_TRUE(Test.ParticleHandles[1]->IsInConstraintGraph());
+			EXPECT_TRUE(Test.ParticleHandles[2]->IsInConstraintGraph());
+			EXPECT_TRUE(Test.ParticleHandles[3]->IsInConstraintGraph());
+			EXPECT_TRUE(FConstGenericParticleHandle(Test.ParticleHandles[0])->IsKinematic());
+			EXPECT_TRUE(FConstGenericParticleHandle(Test.ParticleHandles[1])->Sleeping());
+			EXPECT_TRUE(FConstGenericParticleHandle(Test.ParticleHandles[2])->Sleeping());
+			EXPECT_TRUE(FConstGenericParticleHandle(Test.ParticleHandles[3])->Sleeping());
+
+			EXPECT_TRUE(Test.ConstraintHandles[0]->IsInConstraintGraph());
+			EXPECT_TRUE(Test.ConstraintHandles[1]->IsInConstraintGraph());
+			EXPECT_TRUE(Test.ConstraintHandles[2]->IsInConstraintGraph());
+			EXPECT_TRUE(Test.ConstraintHandles[0]->IsSleeping());
+			EXPECT_TRUE(Test.ConstraintHandles[1]->IsSleeping());
+			EXPECT_TRUE(Test.ConstraintHandles[2]->IsSleeping());
+
+			// Add another particle and connect it to the top of the chain
+			Test.ParticleHandles.Add(Test.Evolution.CreateDynamicParticles(1)[0]);
+			Test.Evolution.EnableParticle(Test.ParticleHandles.Last());
+			Test.ConstraintHandles.Add(Test.Constraints.AddConstraint({ Test.ParticleHandles.Last(1), Test.ParticleHandles.Last() }));
+
+			// NOTE: Set a non-zero velocity to make the particle move and trigger a wake-up event.
+			FParticleDynamics Dynamics;
+			Dynamics.SetAcceleration(FVec3(0, 0, 0));
+			Dynamics.SetAngularAcceleration(FVec3(0, 0, 0));
+			Dynamics.SetLinearImpulseVelocity(FVec3(50, 50, 0));
+			Dynamics.SetAngularImpulseVelocity(FVec3(0, 0, 0));
+			Test.ParticleHandles.Last()->SetDynamics(Dynamics);
+			Test.Evolution.ResetVSmoothFromForces(*Test.ParticleHandles.Last());
+
+			Test.Advance();
+
+			EXPECT_EQ(Test.IslandManager->GetNumParticles(), 5);
+			EXPECT_EQ(Test.IslandManager->GetNumConstraints(), 4);
+
+			// Will wake up the top particle of the original stack
+			EXPECT_EQ(Test.IslandManager->GetNumIslands(), 1);
+			EXPECT_FALSE(Test.IslandManager->GetIsland(0)->IsSleeping());
+			EXPECT_TRUE(FConstGenericParticleHandle(Test.ParticleHandles[1])->Sleeping());
+			EXPECT_TRUE(FConstGenericParticleHandle(Test.ParticleHandles[2])->Sleeping());
+			EXPECT_FALSE(FConstGenericParticleHandle(Test.ParticleHandles[3])->Sleeping());
+			EXPECT_FALSE(FConstGenericParticleHandle(Test.ParticleHandles[4])->Sleeping());
+			EXPECT_TRUE(Test.ConstraintHandles[0]->IsSleeping());
+			EXPECT_TRUE(Test.ConstraintHandles[1]->IsSleeping());
+			EXPECT_FALSE(Test.ConstraintHandles[2]->IsSleeping());
+			EXPECT_FALSE(Test.ConstraintHandles[3]->IsSleeping());
+		}
+	}
+
+	// Test validating the wake-up propagation throughout a sleeping island if adding an awake particle to the top of the stack.
+	// Partial island sleeping only: the awake particle will be added to the island and only wake the particle closest to it. 
+	// (d=dynamic, s=sleeping, k=kinematic)
+	//		Ak - Bs - Cs - Ds
+	// =>	Ak - Bs - Cd - Dd	
+	//                |
+	//                Ed
+	TEST_P(GraphEvolutionTests, TestConstraintGraph_WakeUpPropagation_AddToStackCenter)
+	{
+		FGraphEvolutionTest Test(4, GetParam());
+
+		if (Test.CVarPartialSleeping->GetBool())
+		{
+			Test.MakeChain();
+			Test.Evolution.SetParticleObjectState(Test.ParticleHandles[0], EObjectStateType::Kinematic);
+
+			Test.AdvanceUntilSleeping();
+
+			// Should have 1 island and it should be asleep
+			EXPECT_EQ(Test.IslandManager->GetNumIslands(), 1);
+			EXPECT_TRUE(Test.IslandManager->GetIsland(0)->IsSleeping());
+			EXPECT_TRUE(Test.ParticleHandles[0]->IsInConstraintGraph());
+			EXPECT_TRUE(Test.ParticleHandles[1]->IsInConstraintGraph());
+			EXPECT_TRUE(Test.ParticleHandles[2]->IsInConstraintGraph());
+			EXPECT_TRUE(Test.ParticleHandles[3]->IsInConstraintGraph());
+			EXPECT_TRUE(FConstGenericParticleHandle(Test.ParticleHandles[0])->IsKinematic());
+			EXPECT_TRUE(FConstGenericParticleHandle(Test.ParticleHandles[1])->Sleeping());
+			EXPECT_TRUE(FConstGenericParticleHandle(Test.ParticleHandles[2])->Sleeping());
+			EXPECT_TRUE(FConstGenericParticleHandle(Test.ParticleHandles[3])->Sleeping());
+
+			EXPECT_TRUE(Test.ConstraintHandles[0]->IsInConstraintGraph());
+			EXPECT_TRUE(Test.ConstraintHandles[1]->IsInConstraintGraph());
+			EXPECT_TRUE(Test.ConstraintHandles[2]->IsInConstraintGraph());
+			EXPECT_TRUE(Test.ConstraintHandles[0]->IsSleeping());
+			EXPECT_TRUE(Test.ConstraintHandles[1]->IsSleeping());
+			EXPECT_TRUE(Test.ConstraintHandles[2]->IsSleeping());
+
+			// Add another particle and connect it to the top of the chain
+			Test.ParticleHandles.Add(Test.Evolution.CreateDynamicParticles(1)[0]);
+			Test.Evolution.EnableParticle(Test.ParticleHandles.Last());
+			Test.ConstraintHandles.Add(Test.Constraints.AddConstraint({ Test.ParticleHandles.Last(2), Test.ParticleHandles.Last() }));
+
+			// NOTE: Set a non-zero velocity to make the particle move and trigger a wake-up event.
+			FParticleDynamics Dynamics;
+			Dynamics.SetAcceleration(FVec3(0, 0, 0));
+			Dynamics.SetAngularAcceleration(FVec3(0, 0, 0));
+			Dynamics.SetLinearImpulseVelocity(FVec3(50, 50, 0));
+			Dynamics.SetAngularImpulseVelocity(FVec3(0, 0, 0));
+			Test.ParticleHandles.Last()->SetDynamics(Dynamics);
+			Test.Evolution.ResetVSmoothFromForces(*Test.ParticleHandles.Last());
+
+			Test.Advance();
+
+			EXPECT_EQ(Test.IslandManager->GetNumParticles(), 5);
+			EXPECT_EQ(Test.IslandManager->GetNumConstraints(), 4);
+
+			// Will wake up the top particle of the original stack
+			EXPECT_EQ(Test.IslandManager->GetNumIslands(), 1);
+			EXPECT_FALSE(Test.IslandManager->GetIsland(0)->IsSleeping());
+			EXPECT_TRUE(FConstGenericParticleHandle(Test.ParticleHandles[1])->Sleeping());
+			EXPECT_FALSE(FConstGenericParticleHandle(Test.ParticleHandles[2])->Sleeping());
+			EXPECT_FALSE(FConstGenericParticleHandle(Test.ParticleHandles[3])->Sleeping());
+			EXPECT_FALSE(FConstGenericParticleHandle(Test.ParticleHandles[4])->Sleeping());
+			EXPECT_TRUE(Test.ConstraintHandles[0]->IsSleeping());
+			EXPECT_FALSE(Test.ConstraintHandles[1]->IsSleeping());
+			EXPECT_FALSE(Test.ConstraintHandles[2]->IsSleeping());
+			EXPECT_FALSE(Test.ConstraintHandles[3]->IsSleeping());
+		}
+	}
+
+	// Test validating the wake-up propagation for a triangular brick wall with lateral spacing.
+	// We add another particle near the top of the wall (partial island sleeping only).
+	//				    Gs
+	//				  /	  \
+	//		        Es     Fs
+	//             / \	  / \
+	//            Bs   Cs   Ds
+	//             \   |    /
+	//                 Ak
+	//  =>
+	//			 Hd     Gd
+	//			  \	  /	  \
+	//		        Ed     Fs
+	//             / \	  / \
+	//            Bs   Cs   Ds
+	//             \   |    /
+	//                 Ak
+	TEST_P(GraphEvolutionTests, TestConstraintGraph_WakeUpPropagation_AddToWallWithSpacing)
+	{
+		FGraphEvolutionTest Test(7, GetParam());
+		if (Test.CVarPartialSleeping->GetBool())
+		{
+			Test.Evolution.SetParticleObjectState(Test.ParticleHandles[0], EObjectStateType::Kinematic);
+			Test.ConstraintHandles.Add(Test.Constraints.AddConstraint({ Test.ParticleHandles[0], Test.ParticleHandles[1] }));
+			Test.ConstraintHandles.Add(Test.Constraints.AddConstraint({ Test.ParticleHandles[0], Test.ParticleHandles[2] }));
+			Test.ConstraintHandles.Add(Test.Constraints.AddConstraint({ Test.ParticleHandles[0], Test.ParticleHandles[3] }));
+			Test.ConstraintHandles.Add(Test.Constraints.AddConstraint({ Test.ParticleHandles[1], Test.ParticleHandles[4] }));
+			Test.ConstraintHandles.Add(Test.Constraints.AddConstraint({ Test.ParticleHandles[2], Test.ParticleHandles[4] }));
+			Test.ConstraintHandles.Add(Test.Constraints.AddConstraint({ Test.ParticleHandles[2], Test.ParticleHandles[5] }));
+			Test.ConstraintHandles.Add(Test.Constraints.AddConstraint({ Test.ParticleHandles[3], Test.ParticleHandles[5] }));
+			Test.ConstraintHandles.Add(Test.Constraints.AddConstraint({ Test.ParticleHandles[4], Test.ParticleHandles[6] }));
+			Test.ConstraintHandles.Add(Test.Constraints.AddConstraint({ Test.ParticleHandles[5], Test.ParticleHandles[6] }));
+
+			Test.AdvanceUntilSleeping();
+
+			// Should have 1 island and it should be asleep
+			EXPECT_EQ(Test.IslandManager->GetNumIslands(), 1);
+			EXPECT_TRUE(Test.IslandManager->GetIsland(0)->IsSleeping());
+			EXPECT_TRUE(FConstGenericParticleHandle(Test.ParticleHandles[0])->IsKinematic());
+			EXPECT_TRUE(FConstGenericParticleHandle(Test.ParticleHandles[1])->Sleeping());
+			EXPECT_TRUE(FConstGenericParticleHandle(Test.ParticleHandles[2])->Sleeping());
+			EXPECT_TRUE(FConstGenericParticleHandle(Test.ParticleHandles[3])->Sleeping());
+			EXPECT_TRUE(FConstGenericParticleHandle(Test.ParticleHandles[4])->Sleeping());
+			EXPECT_TRUE(FConstGenericParticleHandle(Test.ParticleHandles[5])->Sleeping());
+			EXPECT_TRUE(FConstGenericParticleHandle(Test.ParticleHandles[6])->Sleeping());
+			EXPECT_TRUE(Test.ConstraintHandles[0]->IsSleeping());
+			EXPECT_TRUE(Test.ConstraintHandles[1]->IsSleeping());
+			EXPECT_TRUE(Test.ConstraintHandles[2]->IsSleeping());
+			EXPECT_TRUE(Test.ConstraintHandles[3]->IsSleeping());
+			EXPECT_TRUE(Test.ConstraintHandles[4]->IsSleeping());
+			EXPECT_TRUE(Test.ConstraintHandles[5]->IsSleeping());
+			EXPECT_TRUE(Test.ConstraintHandles[6]->IsSleeping());
+			EXPECT_TRUE(Test.ConstraintHandles[7]->IsSleeping());
+			EXPECT_TRUE(Test.ConstraintHandles[8]->IsSleeping());
+
+			// Add another particle near the top of the wall
+			Test.ParticleHandles.Add(Test.Evolution.CreateDynamicParticles(1)[0]);
+			Test.Evolution.EnableParticle(Test.ParticleHandles.Last());
+			Test.ConstraintHandles.Add(Test.Constraints.AddConstraint({ Test.ParticleHandles.Last(3), Test.ParticleHandles.Last() }));
+
+			// NOTE: Set a non-zero velocity to make the particle move and trigger a wake-up event.
+			FParticleDynamics Dynamics;
+			Dynamics.SetAcceleration(FVec3(0, 0, 0));
+			Dynamics.SetAngularAcceleration(FVec3(0, 0, 0));
+			Dynamics.SetLinearImpulseVelocity(FVec3(50, 50, 0));
+			Dynamics.SetAngularImpulseVelocity(FVec3(0, 0, 0));
+			Test.ParticleHandles.Last()->SetDynamics(Dynamics);
+			Test.Evolution.ResetVSmoothFromForces(*Test.ParticleHandles.Last());
+
+			Test.Advance();
+
+			EXPECT_EQ(Test.IslandManager->GetNumParticles(), 8);
+			EXPECT_EQ(Test.IslandManager->GetNumConstraints(), 10);
+
+			// Will wake up the two particles connected by the new constraint and the top of the wall
+			EXPECT_EQ(Test.IslandManager->GetNumIslands(), 1);
+			EXPECT_FALSE(Test.IslandManager->GetIsland(0)->IsSleeping());
+			EXPECT_TRUE(FConstGenericParticleHandle(Test.ParticleHandles[1])->Sleeping());
+			EXPECT_TRUE(FConstGenericParticleHandle(Test.ParticleHandles[2])->Sleeping());
+			EXPECT_TRUE(FConstGenericParticleHandle(Test.ParticleHandles[3])->Sleeping());
+			EXPECT_FALSE(FConstGenericParticleHandle(Test.ParticleHandles[4])->Sleeping());
+			EXPECT_TRUE(FConstGenericParticleHandle(Test.ParticleHandles[5])->Sleeping());
+			EXPECT_FALSE(FConstGenericParticleHandle(Test.ParticleHandles[6])->Sleeping());
+			EXPECT_FALSE(FConstGenericParticleHandle(Test.ParticleHandles[7])->Sleeping());
+			EXPECT_TRUE(Test.ConstraintHandles[0]->IsSleeping());
+			EXPECT_TRUE(Test.ConstraintHandles[1]->IsSleeping());
+			EXPECT_TRUE(Test.ConstraintHandles[2]->IsSleeping());
+			EXPECT_FALSE(Test.ConstraintHandles[3]->IsSleeping());
+			EXPECT_FALSE(Test.ConstraintHandles[4]->IsSleeping());
+			EXPECT_TRUE(Test.ConstraintHandles[5]->IsSleeping());
+			EXPECT_TRUE(Test.ConstraintHandles[6]->IsSleeping());
+			EXPECT_FALSE(Test.ConstraintHandles[7]->IsSleeping());
+			EXPECT_FALSE(Test.ConstraintHandles[8]->IsSleeping());
+			EXPECT_FALSE(Test.ConstraintHandles[9]->IsSleeping());
+		}
+	}
+
+	// Test validating the wake-up propagation for a triangular brick wall without lateral spacing.
+	// We add another particle near the top of the wall (partial island sleeping only).
+	//				    Gs
+	//				  /	  \
+	//		        Es  -  Fs
+	//             / \	  / \
+	//            Bs - Cs - Ds
+	//             \   |    /
+	//                 Ak
+	//  =>
+	//			 Hd     Gd
+	//			  \	  /	  \
+	//		        Ed  -  Fs
+	//             / \	  / \
+	//            Bs - Cs - Ds
+	//             \   |    /
+	//                 Ak
+	TEST_P(GraphEvolutionTests, TestConstraintGraph_WakeUpPropagation_AddToWallWithoutSpacing)
+	{
+		FGraphEvolutionTest Test(7, GetParam());
+		if (Test.CVarPartialSleeping->GetBool())
+		{
+			Test.Evolution.SetParticleObjectState(Test.ParticleHandles[0], EObjectStateType::Kinematic);
+			Test.ConstraintHandles.Add(Test.Constraints.AddConstraint({ Test.ParticleHandles[0], Test.ParticleHandles[1] }));
+			Test.ConstraintHandles.Add(Test.Constraints.AddConstraint({ Test.ParticleHandles[0], Test.ParticleHandles[2] }));
+			Test.ConstraintHandles.Add(Test.Constraints.AddConstraint({ Test.ParticleHandles[0], Test.ParticleHandles[3] }));
+			Test.ConstraintHandles.Add(Test.Constraints.AddConstraint({ Test.ParticleHandles[1], Test.ParticleHandles[2] }));
+			Test.ConstraintHandles.Add(Test.Constraints.AddConstraint({ Test.ParticleHandles[2], Test.ParticleHandles[3] }));
+			Test.ConstraintHandles.Add(Test.Constraints.AddConstraint({ Test.ParticleHandles[1], Test.ParticleHandles[4] }));
+			Test.ConstraintHandles.Add(Test.Constraints.AddConstraint({ Test.ParticleHandles[2], Test.ParticleHandles[4] }));
+			Test.ConstraintHandles.Add(Test.Constraints.AddConstraint({ Test.ParticleHandles[2], Test.ParticleHandles[5] }));
+			Test.ConstraintHandles.Add(Test.Constraints.AddConstraint({ Test.ParticleHandles[3], Test.ParticleHandles[5] }));
+			Test.ConstraintHandles.Add(Test.Constraints.AddConstraint({ Test.ParticleHandles[4], Test.ParticleHandles[5] }));
+			Test.ConstraintHandles.Add(Test.Constraints.AddConstraint({ Test.ParticleHandles[4], Test.ParticleHandles[6] }));
+			Test.ConstraintHandles.Add(Test.Constraints.AddConstraint({ Test.ParticleHandles[5], Test.ParticleHandles[6] }));
+
+			Test.AdvanceUntilSleeping();
+
+			// Should have 1 island and it should be asleep
+			EXPECT_EQ(Test.IslandManager->GetNumIslands(), 1);
+			EXPECT_TRUE(Test.IslandManager->GetIsland(0)->IsSleeping());
+			EXPECT_TRUE(FConstGenericParticleHandle(Test.ParticleHandles[0])->IsKinematic());
+			EXPECT_TRUE(FConstGenericParticleHandle(Test.ParticleHandles[1])->Sleeping());
+			EXPECT_TRUE(FConstGenericParticleHandle(Test.ParticleHandles[2])->Sleeping());
+			EXPECT_TRUE(FConstGenericParticleHandle(Test.ParticleHandles[3])->Sleeping());
+			EXPECT_TRUE(FConstGenericParticleHandle(Test.ParticleHandles[4])->Sleeping());
+			EXPECT_TRUE(FConstGenericParticleHandle(Test.ParticleHandles[5])->Sleeping());
+			EXPECT_TRUE(FConstGenericParticleHandle(Test.ParticleHandles[6])->Sleeping());
+			EXPECT_TRUE(Test.ConstraintHandles[0]->IsSleeping());
+			EXPECT_TRUE(Test.ConstraintHandles[1]->IsSleeping());
+			EXPECT_TRUE(Test.ConstraintHandles[2]->IsSleeping());
+			EXPECT_TRUE(Test.ConstraintHandles[3]->IsSleeping());
+			EXPECT_TRUE(Test.ConstraintHandles[4]->IsSleeping());
+			EXPECT_TRUE(Test.ConstraintHandles[5]->IsSleeping());
+			EXPECT_TRUE(Test.ConstraintHandles[6]->IsSleeping());
+			EXPECT_TRUE(Test.ConstraintHandles[7]->IsSleeping());
+			EXPECT_TRUE(Test.ConstraintHandles[8]->IsSleeping());
+			EXPECT_TRUE(Test.ConstraintHandles[9]->IsSleeping());
+			EXPECT_TRUE(Test.ConstraintHandles[10]->IsSleeping());
+			EXPECT_TRUE(Test.ConstraintHandles[11]->IsSleeping());
+
+			// Add another particle near the top of the wall
+			Test.ParticleHandles.Add(Test.Evolution.CreateDynamicParticles(1)[0]);
+			Test.Evolution.EnableParticle(Test.ParticleHandles.Last());
+			Test.ConstraintHandles.Add(Test.Constraints.AddConstraint({ Test.ParticleHandles.Last(3), Test.ParticleHandles.Last() }));
+
+			// NOTE: Set a non-zero velocity to make the particle move and trigger a wake-up event.
+			FParticleDynamics Dynamics;
+			Dynamics.SetAcceleration(FVec3(0, 0, 0));
+			Dynamics.SetAngularAcceleration(FVec3(0, 0, 0));
+			Dynamics.SetLinearImpulseVelocity(FVec3(50, 50, 0));
+			Dynamics.SetAngularImpulseVelocity(FVec3(0, 0, 0));
+			Test.ParticleHandles.Last()->SetDynamics(Dynamics);
+			Test.Evolution.ResetVSmoothFromForces(*Test.ParticleHandles.Last());
+
+			Test.Advance();
+
+			EXPECT_EQ(Test.IslandManager->GetNumParticles(), 8);
+			EXPECT_EQ(Test.IslandManager->GetNumConstraints(), 13);
+
+			// Will wake up the two particles connected by the new constraint and the top of the wall
+			EXPECT_EQ(Test.IslandManager->GetNumIslands(), 1);
+			EXPECT_FALSE(Test.IslandManager->GetIsland(0)->IsSleeping());
+			EXPECT_TRUE(FConstGenericParticleHandle(Test.ParticleHandles[0])->IsKinematic());
+			EXPECT_TRUE(FConstGenericParticleHandle(Test.ParticleHandles[1])->Sleeping());
+			EXPECT_TRUE(FConstGenericParticleHandle(Test.ParticleHandles[2])->Sleeping());
+			EXPECT_TRUE(FConstGenericParticleHandle(Test.ParticleHandles[3])->Sleeping());
+			EXPECT_FALSE(FConstGenericParticleHandle(Test.ParticleHandles[4])->Sleeping());
+			EXPECT_TRUE(FConstGenericParticleHandle(Test.ParticleHandles[5])->Sleeping());
+			EXPECT_FALSE(FConstGenericParticleHandle(Test.ParticleHandles[6])->Sleeping());
+			EXPECT_FALSE(FConstGenericParticleHandle(Test.ParticleHandles[7])->Sleeping());
+			EXPECT_TRUE(Test.ConstraintHandles[0]->IsSleeping());
+			EXPECT_TRUE(Test.ConstraintHandles[1]->IsSleeping());
+			EXPECT_TRUE(Test.ConstraintHandles[2]->IsSleeping());
+			EXPECT_TRUE(Test.ConstraintHandles[3]->IsSleeping());
+			EXPECT_TRUE(Test.ConstraintHandles[4]->IsSleeping());
+			EXPECT_FALSE(Test.ConstraintHandles[5]->IsSleeping());
+			EXPECT_FALSE(Test.ConstraintHandles[6]->IsSleeping());
+			EXPECT_TRUE(Test.ConstraintHandles[7]->IsSleeping());
+			EXPECT_TRUE(Test.ConstraintHandles[8]->IsSleeping());
+			EXPECT_FALSE(Test.ConstraintHandles[9]->IsSleeping());
+			EXPECT_FALSE(Test.ConstraintHandles[10]->IsSleeping());
+			EXPECT_FALSE(Test.ConstraintHandles[11]->IsSleeping());
+			EXPECT_FALSE(Test.ConstraintHandles[12]->IsSleeping());
+		}
+	}
+
+	// Test validating the wake-up propagation for a rectangular block wall without lateral spacing.
+	// We add another constraint to connect two particles diagonally.
+	//			  Hs - Ks - Ls
+	//			  |    |    |
+	//		      Es - Fs - Gs
+	//            |    |    |
+	//            Bs - Cs - Ds
+	//             \   |    /
+	//                 Ak
+	//  =>
+	//			  Hs - Ks - Ls
+	//			  |  \ |    |
+	//		      Es - Fs - Gs
+	//            |    |    |
+	//            Bs - Cs - Ds
+	//             \   |    /
+	//                 Ak
+	TEST_P(GraphEvolutionTests, TestConstraintGraph_WakeUpPropagation_AddToBlockWithoutSpacingSleeping)
+	{
+		FGraphEvolutionTest Test(10, GetParam());
+		Test.Evolution.SetParticleObjectState(Test.ParticleHandles[0], EObjectStateType::Kinematic);
+
+		// Constraints with the ground
+		Test.ConstraintHandles.Add(Test.Constraints.AddConstraint({ Test.ParticleHandles[0], Test.ParticleHandles[1] }));
+		Test.ConstraintHandles.Add(Test.Constraints.AddConstraint({ Test.ParticleHandles[0], Test.ParticleHandles[2] }));
+		Test.ConstraintHandles.Add(Test.Constraints.AddConstraint({ Test.ParticleHandles[0], Test.ParticleHandles[3] }));
+		// Horizontal constraints
+		Test.ConstraintHandles.Add(Test.Constraints.AddConstraint({ Test.ParticleHandles[1], Test.ParticleHandles[2] }));
+		Test.ConstraintHandles.Add(Test.Constraints.AddConstraint({ Test.ParticleHandles[2], Test.ParticleHandles[3] }));
+		// Vertical constraints
+		Test.ConstraintHandles.Add(Test.Constraints.AddConstraint({ Test.ParticleHandles[1], Test.ParticleHandles[4] }));
+		Test.ConstraintHandles.Add(Test.Constraints.AddConstraint({ Test.ParticleHandles[2], Test.ParticleHandles[5] }));
+		Test.ConstraintHandles.Add(Test.Constraints.AddConstraint({ Test.ParticleHandles[3], Test.ParticleHandles[6] }));
+		// Horizontal constraints
+		Test.ConstraintHandles.Add(Test.Constraints.AddConstraint({ Test.ParticleHandles[4], Test.ParticleHandles[5] }));
+		Test.ConstraintHandles.Add(Test.Constraints.AddConstraint({ Test.ParticleHandles[5], Test.ParticleHandles[6] }));
+		// Vertical constraints
+		Test.ConstraintHandles.Add(Test.Constraints.AddConstraint({ Test.ParticleHandles[4], Test.ParticleHandles[7] }));
+		Test.ConstraintHandles.Add(Test.Constraints.AddConstraint({ Test.ParticleHandles[5], Test.ParticleHandles[8] }));
+		Test.ConstraintHandles.Add(Test.Constraints.AddConstraint({ Test.ParticleHandles[6], Test.ParticleHandles[9] }));
+		// Horizontal constraints
+		Test.ConstraintHandles.Add(Test.Constraints.AddConstraint({ Test.ParticleHandles[7], Test.ParticleHandles[8] }));
+		Test.ConstraintHandles.Add(Test.Constraints.AddConstraint({ Test.ParticleHandles[8], Test.ParticleHandles[9] }));
+
+		Test.AdvanceUntilSleeping();
+
+		// Should have 1 island and it should be asleep
+		EXPECT_EQ(Test.IslandManager->GetNumIslands(), 1);
+		EXPECT_TRUE(Test.IslandManager->GetIsland(0)->IsSleeping());
+		EXPECT_TRUE(FConstGenericParticleHandle(Test.ParticleHandles[0])->IsKinematic());
+		EXPECT_TRUE(FConstGenericParticleHandle(Test.ParticleHandles[1])->Sleeping());
+		EXPECT_TRUE(FConstGenericParticleHandle(Test.ParticleHandles[2])->Sleeping());
+		EXPECT_TRUE(FConstGenericParticleHandle(Test.ParticleHandles[3])->Sleeping());
+		EXPECT_TRUE(FConstGenericParticleHandle(Test.ParticleHandles[4])->Sleeping());
+		EXPECT_TRUE(FConstGenericParticleHandle(Test.ParticleHandles[5])->Sleeping());
+		EXPECT_TRUE(FConstGenericParticleHandle(Test.ParticleHandles[6])->Sleeping());
+		EXPECT_TRUE(FConstGenericParticleHandle(Test.ParticleHandles[7])->Sleeping());
+		EXPECT_TRUE(FConstGenericParticleHandle(Test.ParticleHandles[8])->Sleeping());
+		EXPECT_TRUE(FConstGenericParticleHandle(Test.ParticleHandles[9])->Sleeping());
+		EXPECT_TRUE(Test.ConstraintHandles[0]->IsSleeping());
+		EXPECT_TRUE(Test.ConstraintHandles[1]->IsSleeping());
+		EXPECT_TRUE(Test.ConstraintHandles[2]->IsSleeping());
+		EXPECT_TRUE(Test.ConstraintHandles[3]->IsSleeping());
+		EXPECT_TRUE(Test.ConstraintHandles[4]->IsSleeping());
+		EXPECT_TRUE(Test.ConstraintHandles[5]->IsSleeping());
+		EXPECT_TRUE(Test.ConstraintHandles[6]->IsSleeping());
+		EXPECT_TRUE(Test.ConstraintHandles[7]->IsSleeping());
+		EXPECT_TRUE(Test.ConstraintHandles[8]->IsSleeping());
+		EXPECT_TRUE(Test.ConstraintHandles[9]->IsSleeping());
+		EXPECT_TRUE(Test.ConstraintHandles[10]->IsSleeping());
+		EXPECT_TRUE(Test.ConstraintHandles[11]->IsSleeping());
+		EXPECT_TRUE(Test.ConstraintHandles[12]->IsSleeping());
+		EXPECT_TRUE(Test.ConstraintHandles[13]->IsSleeping());
+		EXPECT_TRUE(Test.ConstraintHandles[14]->IsSleeping());
+
+		// Add another constraint to diagonally connect to particles in the block
+		Test.ConstraintHandles.Add(Test.Constraints.AddConstraint({ Test.ParticleHandles[5], Test.ParticleHandles[7] }));
+
+		Test.Advance();
+
+		EXPECT_EQ(Test.IslandManager->GetNumParticles(), 10);
+		EXPECT_EQ(Test.IslandManager->GetNumConstraints(), 16);
+
+		EXPECT_EQ(Test.IslandManager->GetNumIslands(), 1);
+		EXPECT_TRUE(Test.IslandManager->GetIsland(0)->IsSleeping());
+		EXPECT_TRUE(FConstGenericParticleHandle(Test.ParticleHandles[0])->IsKinematic());
+		EXPECT_TRUE(FConstGenericParticleHandle(Test.ParticleHandles[1])->Sleeping());
+		EXPECT_TRUE(FConstGenericParticleHandle(Test.ParticleHandles[2])->Sleeping());
+		EXPECT_TRUE(FConstGenericParticleHandle(Test.ParticleHandles[3])->Sleeping());
+		EXPECT_TRUE(FConstGenericParticleHandle(Test.ParticleHandles[4])->Sleeping());
+		EXPECT_TRUE(FConstGenericParticleHandle(Test.ParticleHandles[5])->Sleeping());
+		EXPECT_TRUE(FConstGenericParticleHandle(Test.ParticleHandles[6])->Sleeping());
+		EXPECT_TRUE(FConstGenericParticleHandle(Test.ParticleHandles[7])->Sleeping());
+		EXPECT_TRUE(FConstGenericParticleHandle(Test.ParticleHandles[8])->Sleeping());
+		EXPECT_TRUE(FConstGenericParticleHandle(Test.ParticleHandles[9])->Sleeping());
+		EXPECT_TRUE(Test.ConstraintHandles[0]->IsSleeping());
+		EXPECT_TRUE(Test.ConstraintHandles[1]->IsSleeping());
+		EXPECT_TRUE(Test.ConstraintHandles[2]->IsSleeping());
+		EXPECT_TRUE(Test.ConstraintHandles[3]->IsSleeping());
+		EXPECT_TRUE(Test.ConstraintHandles[4]->IsSleeping());
+		EXPECT_TRUE(Test.ConstraintHandles[5]->IsSleeping());
+		EXPECT_TRUE(Test.ConstraintHandles[6]->IsSleeping());
+		EXPECT_TRUE(Test.ConstraintHandles[7]->IsSleeping());
+		EXPECT_TRUE(Test.ConstraintHandles[8]->IsSleeping());
+		EXPECT_TRUE(Test.ConstraintHandles[9]->IsSleeping());
+		EXPECT_TRUE(Test.ConstraintHandles[10]->IsSleeping());
+		EXPECT_TRUE(Test.ConstraintHandles[11]->IsSleeping());
+		EXPECT_TRUE(Test.ConstraintHandles[12]->IsSleeping());
+		EXPECT_TRUE(Test.ConstraintHandles[13]->IsSleeping());
+		EXPECT_TRUE(Test.ConstraintHandles[14]->IsSleeping());
+		EXPECT_TRUE(Test.ConstraintHandles[15]->IsSleeping());
+	}
+
+	// Test validating the wake-up propagation for a rectangular block wall without lateral spacing.
+	// We add another constraint to connect two particles diagonally and wake one of the particles.
+	//			  Hs - Ks - Ls
+	//			  |    |    |
+	//		      Es - Fs - Gs
+	//            |    |    |
+	//            Bs - Cs - Ds
+	//             \   |    /
+	//                 Ak
+	//  =>
+	//			  Hd - Kd - Ld
+	//			  |  \ |    |
+	//		      Ed - Fd - Gd
+	//            |    |    |
+	//            Bd - Cd - Dd
+	//             \   |    /
+	//                 Ak
+	TEST_P(GraphEvolutionTests, TestConstraintGraph_WakeUpPropagation_AddToBlockWithoutSpacingAwake)
+	{
+		FGraphEvolutionTest Test(10, GetParam());
+		Test.Evolution.SetParticleObjectState(Test.ParticleHandles[0], EObjectStateType::Kinematic);
+
+		// Constraints with the ground
+		Test.ConstraintHandles.Add(Test.Constraints.AddConstraint({ Test.ParticleHandles[0], Test.ParticleHandles[1] }));
+		Test.ConstraintHandles.Add(Test.Constraints.AddConstraint({ Test.ParticleHandles[0], Test.ParticleHandles[2] }));
+		Test.ConstraintHandles.Add(Test.Constraints.AddConstraint({ Test.ParticleHandles[0], Test.ParticleHandles[3] }));
+		// Horizontal constraints
+		Test.ConstraintHandles.Add(Test.Constraints.AddConstraint({ Test.ParticleHandles[1], Test.ParticleHandles[2] }));
+		Test.ConstraintHandles.Add(Test.Constraints.AddConstraint({ Test.ParticleHandles[2], Test.ParticleHandles[3] }));
+		// Vertical constraints
+		Test.ConstraintHandles.Add(Test.Constraints.AddConstraint({ Test.ParticleHandles[1], Test.ParticleHandles[4] }));
+		Test.ConstraintHandles.Add(Test.Constraints.AddConstraint({ Test.ParticleHandles[2], Test.ParticleHandles[5] }));
+		Test.ConstraintHandles.Add(Test.Constraints.AddConstraint({ Test.ParticleHandles[3], Test.ParticleHandles[6] }));
+		// Horizontal constraints
+		Test.ConstraintHandles.Add(Test.Constraints.AddConstraint({ Test.ParticleHandles[4], Test.ParticleHandles[5] }));
+		Test.ConstraintHandles.Add(Test.Constraints.AddConstraint({ Test.ParticleHandles[5], Test.ParticleHandles[6] }));
+		// Vertical constraints
+		Test.ConstraintHandles.Add(Test.Constraints.AddConstraint({ Test.ParticleHandles[4], Test.ParticleHandles[7] }));
+		Test.ConstraintHandles.Add(Test.Constraints.AddConstraint({ Test.ParticleHandles[5], Test.ParticleHandles[8] }));
+		Test.ConstraintHandles.Add(Test.Constraints.AddConstraint({ Test.ParticleHandles[6], Test.ParticleHandles[9] }));
+		// Horizontal constraints
+		Test.ConstraintHandles.Add(Test.Constraints.AddConstraint({ Test.ParticleHandles[7], Test.ParticleHandles[8] }));
+		Test.ConstraintHandles.Add(Test.Constraints.AddConstraint({ Test.ParticleHandles[8], Test.ParticleHandles[9] }));
+
+		Test.AdvanceUntilSleeping();
+
+		// Should have 1 island and it should be asleep
+		EXPECT_EQ(Test.IslandManager->GetNumIslands(), 1);
+		EXPECT_TRUE(Test.IslandManager->GetIsland(0)->IsSleeping());
+		EXPECT_TRUE(FConstGenericParticleHandle(Test.ParticleHandles[0])->IsKinematic());
+		EXPECT_TRUE(FConstGenericParticleHandle(Test.ParticleHandles[1])->Sleeping());
+		EXPECT_TRUE(FConstGenericParticleHandle(Test.ParticleHandles[2])->Sleeping());
+		EXPECT_TRUE(FConstGenericParticleHandle(Test.ParticleHandles[3])->Sleeping());
+		EXPECT_TRUE(FConstGenericParticleHandle(Test.ParticleHandles[4])->Sleeping());
+		EXPECT_TRUE(FConstGenericParticleHandle(Test.ParticleHandles[5])->Sleeping());
+		EXPECT_TRUE(FConstGenericParticleHandle(Test.ParticleHandles[6])->Sleeping());
+		EXPECT_TRUE(FConstGenericParticleHandle(Test.ParticleHandles[7])->Sleeping());
+		EXPECT_TRUE(FConstGenericParticleHandle(Test.ParticleHandles[8])->Sleeping());
+		EXPECT_TRUE(FConstGenericParticleHandle(Test.ParticleHandles[9])->Sleeping());
+		EXPECT_TRUE(Test.ConstraintHandles[0]->IsSleeping());
+		EXPECT_TRUE(Test.ConstraintHandles[1]->IsSleeping());
+		EXPECT_TRUE(Test.ConstraintHandles[2]->IsSleeping());
+		EXPECT_TRUE(Test.ConstraintHandles[3]->IsSleeping());
+		EXPECT_TRUE(Test.ConstraintHandles[4]->IsSleeping());
+		EXPECT_TRUE(Test.ConstraintHandles[5]->IsSleeping());
+		EXPECT_TRUE(Test.ConstraintHandles[6]->IsSleeping());
+		EXPECT_TRUE(Test.ConstraintHandles[7]->IsSleeping());
+		EXPECT_TRUE(Test.ConstraintHandles[8]->IsSleeping());
+		EXPECT_TRUE(Test.ConstraintHandles[9]->IsSleeping());
+		EXPECT_TRUE(Test.ConstraintHandles[10]->IsSleeping());
+		EXPECT_TRUE(Test.ConstraintHandles[11]->IsSleeping());
+		EXPECT_TRUE(Test.ConstraintHandles[12]->IsSleeping());
+		EXPECT_TRUE(Test.ConstraintHandles[13]->IsSleeping());
+		EXPECT_TRUE(Test.ConstraintHandles[14]->IsSleeping());
+
+		// Add another constraint to diagonally connect to particles in the block
+		Test.ConstraintHandles.Add(Test.Constraints.AddConstraint({ Test.ParticleHandles[5], Test.ParticleHandles[7] }));
+		Test.Evolution.SetParticleObjectState(Test.ParticleHandles[7], EObjectStateType::Dynamic);
+
+		// NOTE: Set a non-zero velocity to make the particle move and trigger a wake-up event.
+		FParticleDynamics Dynamics;
+		Dynamics.SetAcceleration(FVec3(0, 0, 0));
+		Dynamics.SetAngularAcceleration(FVec3(0, 0, 0));
+		Dynamics.SetLinearImpulseVelocity(FVec3(50, 50, 0));
+		Dynamics.SetAngularImpulseVelocity(FVec3(0, 0, 0));
+		Test.ParticleHandles[7]->SetDynamics(Dynamics);
+		Test.Evolution.ResetVSmoothFromForces(*Test.ParticleHandles[7]);
+
+		Test.Advance();
+
+		EXPECT_EQ(Test.IslandManager->GetNumParticles(), 10);
+		EXPECT_EQ(Test.IslandManager->GetNumConstraints(), 16);
+
+		EXPECT_EQ(Test.IslandManager->GetNumIslands(), 1);
+		EXPECT_FALSE(Test.IslandManager->GetIsland(0)->IsSleeping());
+		EXPECT_TRUE(FConstGenericParticleHandle(Test.ParticleHandles[0])->IsKinematic());
+		EXPECT_FALSE(FConstGenericParticleHandle(Test.ParticleHandles[1])->Sleeping());
+		EXPECT_FALSE(FConstGenericParticleHandle(Test.ParticleHandles[2])->Sleeping());
+		EXPECT_FALSE(FConstGenericParticleHandle(Test.ParticleHandles[3])->Sleeping());
+		EXPECT_FALSE(FConstGenericParticleHandle(Test.ParticleHandles[4])->Sleeping());
+		EXPECT_FALSE(FConstGenericParticleHandle(Test.ParticleHandles[5])->Sleeping());
+		EXPECT_FALSE(FConstGenericParticleHandle(Test.ParticleHandles[6])->Sleeping());
+		EXPECT_FALSE(FConstGenericParticleHandle(Test.ParticleHandles[7])->Sleeping());
+		EXPECT_FALSE(FConstGenericParticleHandle(Test.ParticleHandles[8])->Sleeping());
+		EXPECT_FALSE(FConstGenericParticleHandle(Test.ParticleHandles[9])->Sleeping());
+		EXPECT_FALSE(Test.ConstraintHandles[0]->IsSleeping());
+		EXPECT_FALSE(Test.ConstraintHandles[1]->IsSleeping());
+		EXPECT_FALSE(Test.ConstraintHandles[2]->IsSleeping());
+		EXPECT_FALSE(Test.ConstraintHandles[3]->IsSleeping());
+		EXPECT_FALSE(Test.ConstraintHandles[4]->IsSleeping());
+		EXPECT_FALSE(Test.ConstraintHandles[5]->IsSleeping());
+		EXPECT_FALSE(Test.ConstraintHandles[6]->IsSleeping());
+		EXPECT_FALSE(Test.ConstraintHandles[7]->IsSleeping());
+		EXPECT_FALSE(Test.ConstraintHandles[8]->IsSleeping());
+		EXPECT_FALSE(Test.ConstraintHandles[9]->IsSleeping());
+		EXPECT_FALSE(Test.ConstraintHandles[10]->IsSleeping());
+		EXPECT_FALSE(Test.ConstraintHandles[11]->IsSleeping());
+		EXPECT_FALSE(Test.ConstraintHandles[12]->IsSleeping());
+		EXPECT_FALSE(Test.ConstraintHandles[13]->IsSleeping());
+		EXPECT_FALSE(Test.ConstraintHandles[14]->IsSleeping());
+		EXPECT_FALSE(Test.ConstraintHandles[15]->IsSleeping());
+	}
+
+	INSTANTIATE_TEST_SUITE_P(, GraphEvolutionTests,	testing::Bool());
 }
diff --git a/Engine/Source/Runtime/Experimental/Chaos/Private/Chaos/Character/CharacterGroundConstraintContainerSolver.cpp b/Engine/Source/Runtime/Experimental/Chaos/Private/Chaos/Character/CharacterGroundConstraintContainerSolver.cpp
--- a/Engine/Source/Runtime/Experimental/Chaos/Private/Chaos/Character/CharacterGroundConstraintContainerSolver.cpp
+++ b/Engine/Source/Runtime/Experimental/Chaos/Private/Chaos/Character/CharacterGroundConstraintContainerSolver.cpp
@@ -39,10 +39,16 @@ namespace Chaos
 		{
 			for (Private::FPBDIslandConstraint* IslandConstraint : IslandConstraints)
 			{
-				// We will only ever be given constraints from our container (asserts in non-shipping)
-				FCharacterGroundConstraintHandle* Constraint = IslandConstraint->GetConstraint()->AsUnsafe<FCharacterGroundConstraintHandle>();
+				FConstraintHandle* Constraint = IslandConstraint->GetConstraint();
+				
+				// Filter out sleeping constraints in any partially sleeping island
+				if (!Constraint->IsSleeping())
+				{
+					// We will only ever be given constraints from our container (asserts in non-shipping)
+					FCharacterGroundConstraintHandle* CharacterGroundConstraint = Constraint->AsUnsafe<FCharacterGroundConstraintHandle>();
 
-				AddConstraint(Constraint);
+					AddConstraint(CharacterGroundConstraint);
+				}
 			}
 		}
 
diff --git a/Engine/Source/Runtime/Experimental/Chaos/Private/Chaos/Collision/PBDCollisionContainerSolver.cpp b/Engine/Source/Runtime/Experimental/Chaos/Private/Chaos/Collision/PBDCollisionContainerSolver.cpp
--- a/Engine/Source/Runtime/Experimental/Chaos/Private/Chaos/Collision/PBDCollisionContainerSolver.cpp
+++ b/Engine/Source/Runtime/Experimental/Chaos/Private/Chaos/Collision/PBDCollisionContainerSolver.cpp
@@ -537,10 +537,16 @@ namespace Chaos
 	{
 		for (Private::FPBDIslandConstraint* IslandConstraint : IslandConstraints)
 		{
-			// We will only ever be given constraints from our container (NOTE: AsUnsafe asserts in dev)
-			FPBDCollisionConstraint& Constraint = IslandConstraint->GetConstraint()->AsUnsafe<FPBDCollisionConstraintHandle>()->GetContact();
+			FConstraintHandle* Constraint = IslandConstraint->GetConstraint();
 
-			AddConstraint(Constraint);
+			// Filter out sleeping constraints in any partially sleeping island
+			if (!Constraint->IsSleeping())
+			{
+				// We will only ever be given constraints from our container (NOTE: AsUnsafe asserts in dev)
+				FPBDCollisionConstraint& CollisionConstraint = Constraint->AsUnsafe<FPBDCollisionConstraintHandle>()->GetContact();
+
+				AddConstraint(CollisionConstraint);
+			}
 		}
 	}
 
diff --git a/Engine/Source/Runtime/Experimental/Chaos/Private/Chaos/Collision/PBDCollisionContainerSolverJacobi.cpp b/Engine/Source/Runtime/Experimental/Chaos/Private/Chaos/Collision/PBDCollisionContainerSolverJacobi.cpp
--- a/Engine/Source/Runtime/Experimental/Chaos/Private/Chaos/Collision/PBDCollisionContainerSolverJacobi.cpp
+++ b/Engine/Source/Runtime/Experimental/Chaos/Private/Chaos/Collision/PBDCollisionContainerSolverJacobi.cpp
@@ -306,10 +306,16 @@ namespace Chaos
 		{
 			for (Private::FPBDIslandConstraint* IslandConstraint : IslandConstraints)
 			{
-				// We will only ever be given constraints from our container (asserts in non-shipping)
-				FPBDCollisionConstraint& Constraint = IslandConstraint->GetConstraint()->AsUnsafe<FPBDCollisionConstraintHandle>()->GetContact();
+				FConstraintHandle* Constraint = IslandConstraint->GetConstraint();
 
-				AddConstraint(Constraint);
+				// Filter out sleeping constraints in any partially sleeping island
+				if (!Constraint->IsSleeping())
+				{
+					// We will only ever be given constraints from our container (asserts in non-shipping)
+					FPBDCollisionConstraint& CollisionConstraint = Constraint->AsUnsafe<FPBDCollisionConstraintHandle>()->GetContact();
+
+					AddConstraint(CollisionConstraint);
+				}
 			}
 		}
 
diff --git a/Engine/Source/Runtime/Experimental/Chaos/Private/Chaos/Collision/PBDCollisionContainerSolverSimd.cpp b/Engine/Source/Runtime/Experimental/Chaos/Private/Chaos/Collision/PBDCollisionContainerSolverSimd.cpp
--- a/Engine/Source/Runtime/Experimental/Chaos/Private/Chaos/Collision/PBDCollisionContainerSolverSimd.cpp
+++ b/Engine/Source/Runtime/Experimental/Chaos/Private/Chaos/Collision/PBDCollisionContainerSolverSimd.cpp
@@ -367,22 +367,31 @@ namespace Chaos
 			}
 
 			// Add all the constraints in the island to the selected lane
+			int32 NumActiveConstraints = 0;
 			for (Private::FPBDIslandConstraint* IslandConstraint : IslandConstraints)
-			{
-				// NOTE: We will only ever be given constraints from our container (asserts in non-shipping)
-				FPBDCollisionConstraint& Constraint = IslandConstraint->GetConstraint()->AsUnsafe<FPBDCollisionConstraintHandle>()->GetContact();
+			{				
+				FConstraintHandle* Constraint = IslandConstraint->GetConstraint();
+
+				// Filter out sleeping constraints in any partially sleeping island
+				if (!Constraint->IsSleeping())
+				{
+					// NOTE: We will only ever be given constraints from our container (asserts in non-shipping)
+					FPBDCollisionConstraint& CollisionConstraint = Constraint->AsUnsafe<FPBDCollisionConstraintHandle>()->GetContact();
 
-				const int32 SolverIndex = SimdData.SimdNumConstraints.GetValue(IslandLaneIndex);
-				SimdData.SimdNumConstraints.SetValue(IslandLaneIndex, SolverIndex + 1);
+					const int32 SolverIndex = SimdData.SimdNumConstraints.GetValue(IslandLaneIndex);
+					SimdData.SimdNumConstraints.SetValue(IslandLaneIndex, SolverIndex + 1);
 
-				// Add the constraint to its island's lane
-				SimdData.SimdConstraints[SolverIndex].SetValue(IslandLaneIndex, &Constraint);
+					// Add the constraint to its island's lane
+					SimdData.SimdConstraints[SolverIndex].SetValue(IslandLaneIndex, &CollisionConstraint);
 
-				// Store the mapping from constraint to solver/lane
-				ConstraintSolverIds.Add({ SolverIndex, IslandLaneIndex });
+					// Store the mapping from constraint to solver/lane
+					ConstraintSolverIds.Add({ SolverIndex, IslandLaneIndex });
+
+					++NumActiveConstraints;
+				}
 			}
 
-			NumConstraints += IslandConstraints.Num();
+			NumConstraints += NumActiveConstraints;
 		}
 
 		void FPBDCollisionContainerSolverSimd::CreateSolvers()
diff --git a/Engine/Source/Runtime/Experimental/Chaos/Private/Chaos/Evolution/SolverBodyContainer.cpp b/Engine/Source/Runtime/Experimental/Chaos/Private/Chaos/Evolution/SolverBodyContainer.cpp
--- a/Engine/Source/Runtime/Experimental/Chaos/Private/Chaos/Evolution/SolverBodyContainer.cpp
+++ b/Engine/Source/Runtime/Experimental/Chaos/Private/Chaos/Evolution/SolverBodyContainer.cpp
@@ -13,7 +13,11 @@ namespace Chaos
 
 		const FPBDRigidParticleHandle* RigidParticle = Particle->CastToRigidParticle();
 		const FKinematicGeometryParticleHandle* KinematicParticle = Particle->CastToKinematicParticle();
-		if ((RigidParticle != nullptr) && RigidParticle->IsDynamic())
+
+		// NOTE: IsDynamic() is true whether the particle is awake or asleep.
+		// In the solver, we want sleeping particles to behave like kinematics
+		// since they do not experience any positional impulse and position correction.
+		if ((RigidParticle != nullptr) && RigidParticle->IsDynamic() && !RigidParticle->IsSleeping())
 		{
 			// Dynamic particle
 			SolverBody.SetX(RigidParticle->XCom());
@@ -96,7 +100,10 @@ namespace Chaos
 	
 		if (ItemIndex == INDEX_NONE)
 		{
-			if (InParticle->IsDynamic())
+			// NOTE: IsDynamic() is true whether the particle is awake or asleep.
+			// The SolverBodyIndex() for sleeping particles does not seem to be reset at each tick.
+			// As a workaround, we use the same procedure as kinematic particles to make sure it is persistent.
+			if (InParticle->IsDynamic() && !InParticle->IsSleeping())
 			{
 				// First time we have seen this particle, so add it
 				check(!bLocked);
@@ -121,6 +128,7 @@ namespace Chaos
 		}
 	
 		check(ItemIndex != INDEX_NONE);
+		check(ItemIndex < SolverBodies.Num());
 		
 		FSolverBody* SolverBody = &SolverBodies[ItemIndex];
 
diff --git a/Engine/Source/Runtime/Experimental/Chaos/Private/Chaos/Island/IslandGroupManager.cpp b/Engine/Source/Runtime/Experimental/Chaos/Private/Chaos/Island/IslandGroupManager.cpp
--- a/Engine/Source/Runtime/Experimental/Chaos/Private/Chaos/Island/IslandGroupManager.cpp
+++ b/Engine/Source/Runtime/Experimental/Chaos/Private/Chaos/Island/IslandGroupManager.cpp
@@ -291,7 +291,7 @@ namespace Chaos
 
 					IslandGroup->AddConstraintsAndBodies();
 					IslandGroup->GatherBodies(Dt, 0, IslandGroup->GetNumSolverBodies());
-					IslandGroup->GatherConstraints(Dt, 0, IslandGroup->GetNumConstraints());
+					IslandGroup->GatherConstraints(Dt, 0, IslandGroup->GetNumSolverConstraints());
 				}
 
 				SolveGroupConstraints(GroupIndex, Dt);
diff --git a/Engine/Source/Runtime/Experimental/Chaos/Private/Chaos/Island/IslandManager.cpp b/Engine/Source/Runtime/Experimental/Chaos/Private/Chaos/Island/IslandManager.cpp
--- a/Engine/Source/Runtime/Experimental/Chaos/Private/Chaos/Island/IslandManager.cpp
+++ b/Engine/Source/Runtime/Experimental/Chaos/Private/Chaos/Island/IslandManager.cpp
@@ -49,7 +49,19 @@ namespace Chaos::CVars
 
 	/** Cvar to enable/disable the island sleeping */
 	bool bChaosSolverSleepEnabled = true;
-	FAutoConsoleVariableRef CVarChaosSolverSleepEnabled(TEXT("p.Chaos.Solver.Sleep.Enabled"), bChaosSolverSleepEnabled, TEXT(""));
+	FAutoConsoleVariableRef CVarChaosSolverSleepEnabled(TEXT("p.Chaos.Solver.Sleep.Enabled"), bChaosSolverSleepEnabled, TEXT("Islands of non - moving particles will be deactivated."));
+
+	/** Cvar to enable/disable partial island sleeping for all islands */
+	bool bChaosSolverPartialIslandSleep = false;
+	FAutoConsoleVariableRef CVarChaosSolverPartialIslandSleep(TEXT("p.Chaos.Solver.Sleep.PartialIslandSleep"), bChaosSolverPartialIslandSleep, TEXT("Partial island sleeping: Single non-moving particles in islands will be deactivated even if touching moving particles. NOTE: Wake-up only works correctly if shock propagation is enabled."));
+
+	/** Cvar to enable/disable waking up the entire island if the state of any particle has changed (partial island sleeping only) */
+	bool bChaosSolverWakeOnParticleStateChanged = true;
+	FAutoConsoleVariableRef CVarChaosSolverWakeOnParticleStateChanged(TEXT("p.Chaos.Solver.Sleep.WakeOnParticleStateChanged"), bChaosSolverWakeOnParticleStateChanged, TEXT("State changes of a single particle will trigger the island to wake up entirely (partial island sleeping only)."));
+
+	/** Cvar to enable/disable triggering wake up for newly detected constraints with moving particles only (partial island sleeping only) */
+	bool bChaosSolverWakeOnMovingOnly = true;
+	FAutoConsoleVariableRef CVarChaosSolverWakeOnMovingOnly(TEXT("p.Chaos.Solver.Sleep.WakeOnMovingOnly"), bChaosSolverWakeOnMovingOnly, TEXT("Only newly detected constraints with moving particle will trigger wake-up event(partial island sleeping only)."));
 
 	/** Cvar to override the sleep counter threshold if necessary */
 	int32 ChaosSolverCollisionDefaultSleepCounterThreshold = 20;
@@ -81,7 +93,7 @@ namespace Chaos::CVars
 	/* Cvar to adjust the sleep angular threshold for floating particles */
 	FRealSingle IsolatedParticleSleepAngularThresholdMultiplier = 1.0f;
 	FAutoConsoleVariableRef CVarChaosSolverIsolatedParticleSleepAngularThresholdMultiplier(TEXT("p.Chaos.Solver.Sleep.IsolatedParticle.AngularMultiplier"), IsolatedParticleSleepAngularThresholdMultiplier, TEXT("A multiplier applied to SleepAngularThreshold for floating particles"));
-	
+
 	/** Cvar to enable/disable computing max iterations if island is dirty */
 	bool bChaosSolverComputeIterationSettings = true;
 	FAutoConsoleVariableRef CVarChaosSolverComputeIterationSettingsEnabled(TEXT("p.Chaos.Solver.ComputeIterationSettings.Enabled"), bChaosSolverComputeIterationSettings, TEXT("Recompute iteration settings every time an island is changed"));
@@ -879,6 +891,16 @@ namespace Chaos::Private
 			// Create the edge representing the constraint. Also creates the nodes if necessary
 			if (FPBDIslandConstraint* Edge = CreateGraphEdge(ContainerId, Constraint, ConstrainedParticles))
 			{
+				// NOTE: Used only by partial island sleeping. 
+				// We restrict setting this flag to prevent overhead of resetting it.
+				if (CVars::bChaosSolverPartialIslandSleep)
+				{
+					Edge->Flags.bIsNew = true;
+
+					// When a new constraint is added, it should not be sleeping.
+					check(!(!!Edge->Flags.bIsNew && !!Edge->Flags.bIsSleeping));
+				}
+
 				// Put the edge/nodes into an island and flag islands for merging if necessary
 				// NOTE: We do not wake the island here. It will be awake already if one of the
 				// particles is awake. It both particles are asleep we do not want to wake them
@@ -1273,8 +1295,23 @@ namespace Chaos::Private
 		// This is to cover the case where any particle is woken, or all particles are manually slept
 		if (bIsDynamic && (!bWasDynamic || (bIsSleeping != bWasSleeping) || bIsMoving))
 		{
-			const bool bIsSleepAllowed = true;
-			EnqueueIslandCheckSleep(Node->Island, bIsSleepAllowed);
+			if (CVars::bChaosSolverPartialIslandSleep && CVars::bChaosSolverWakeOnParticleStateChanged)
+			{
+				// When do we want to force the entire island to wake up for at least 1 step? (partial island sleeping only)
+				// Case 1) Switch from non-moving kinematic to non-sleeping dynamic OR
+				// Case 2) Switch dynamic particle from sleeping to not-sleeping
+				// @todo(chaos): This will wake up the entire island for partial island sleeping.
+				// Not a great solution for performance. Update this when implementing the next version.
+				const bool bSwitchStaticKinematicToAwakeDynamic = (!bWasDynamic && !bWasMoving) && (bIsDynamic && !bIsSleeping); // Case 1
+				const bool bSwitchSleepingDynamicToAwake = (bIsDynamic && bWasSleeping && !bIsSleeping); // Case 2
+				const bool bIsSleepAllowed = (bSwitchStaticKinematicToAwakeDynamic || bSwitchSleepingDynamicToAwake) ? false : true;
+				EnqueueIslandCheckSleep(Node->Island, bIsSleepAllowed);
+			}
+			else
+			{
+				const bool bIsSleepAllowed = true;
+				EnqueueIslandCheckSleep(Node->Island, bIsSleepAllowed);
+			}
 		}
 
 		// If we have a moving kinematic its island(s) cannot sleep
@@ -1976,7 +2013,7 @@ namespace Chaos::Private
 			}
 		}
 
-		// Breadth-first visit all nodes, assign a level to the connected node, and enqeue it
+		// Breadth-first visit all nodes, assign a level to the connected node, and enqueue it
 		for (int32 NodeIndex = 0; NodeIndex < NodeQueue.Num(); ++NodeIndex)
 		{
 			FPBDIslandParticle* Node = NodeQueue[NodeIndex];
@@ -2072,10 +2109,27 @@ namespace Chaos::Private
 
 			// If the island is awake or just changed state, make sure all the particles and constraints agree
 			// NOTE: also handles the case where we add to a sleeping island (by checking bCheckSleep)
-			// @todo(chaos): we should only need to do this when the state changes...?
+			// @todo(chaos) : (island sleeping only) we should only need to do this when the state changes... ?
 			if (!Island->Flags.bIsSleeping || (Island->Flags.bIsSleeping != Island->Flags.bWasSleeping) || !!Island->Flags.bCheckSleep)
 			{
-				PropagateIslandSleep(Island);
+				if (CVars::bChaosSolverPartialIslandSleep && !!Island->Flags.bIsSleepAllowed)
+				{
+					if (!Island->Flags.bIsSleeping)
+					{
+						ProcessParticleWakes(Island);
+					}
+					SyncConstraintSleepStateToParticles(Island);
+				}
+				else // For island sleeping to sync the sleep state OR if (!bIsSleepAllowed) to wake up the entire island
+				{
+					PropagateIslandSleep(Island);
+
+					if (CVars::bChaosSolverPartialIslandSleep)
+					{
+						// NOTE: This reset is needed if (CVars::bChaosSolverPartialIslandSleepEnabled && Island->Flags.bIsSleepAllowed)
+						ResetEdgeFlags(Island);
+					}
+				}
 				Island->UpdateSyncState();
 			}
 
@@ -2141,19 +2195,19 @@ namespace Chaos::Private
 		// be detected. If that particle is in an island with other awake particles it will be woken immediately but will not have 
 		// moved or have any collisions.
 
-		const auto ShouldIslandSleep = [](FPBDIsland * Island) -> bool
-		{
-			for (FPBDIslandParticle* Node : Island->GetParticles())
+		const auto ShouldIslandSleep = [](FPBDIsland* Island) -> bool
 			{
-				const bool bIsDynamic = IsParticleDynamic(Node->GetParticle());
-				const bool bIsSleeping = IsParticleSleeping(Node->GetParticle());
-				if (bIsDynamic && !bIsSleeping)
+				for (FPBDIslandParticle* Node : Island->GetParticles())
 				{
-					return false;
+					const bool bIsDynamic = IsParticleDynamic(Node->GetParticle());
+					const bool bIsSleeping = IsParticleSleeping(Node->GetParticle());
+					if (bIsDynamic && !bIsSleeping)
+					{
+						return false;
+					}
 				}
-			}
-			return true;
-		};
+				return true;
+			};
 
 		// If we have explicitly made some particles go to sleep (as opposed to them naturally sleeping based on low movement)
 		// we must check to see if the whole island can go to sleep. This is primarily to address the issue that we do not
@@ -2183,10 +2237,14 @@ namespace Chaos::Private
 	{
 		if (!CVars::bChaosSolverSleepEnabled)
 		{
+#if CHAOS_DEBUG_DRAW
+			// Debugging only: Update velocity smoothing so that debug draw works even if sleeping is disabled
+			UpdateAllParticleSleepMetrics(Dt);
+#endif
 			return;
 		}
 
-		// Isloated particles are not kept in any island and need to be handled separately.
+		// Isolated particles are not kept in any island and need to be handled separately.
 		ProcessParticlesSleep(Dt);
 
 		// @todo(chaos): can go wide except for PropagateSleepState
@@ -2194,12 +2252,22 @@ namespace Chaos::Private
 		{
 			if (!Island->Flags.bIsSleeping && !!Island->Flags.bIsSleepAllowed && !Island->Flags.bIsUsingCache)
 			{
-				// Update the sleep state based on particle movement etc
-				ProcessIslandSleep(Island, Dt);
+				if (!CVars::bChaosSolverPartialIslandSleep)
+				{
+					// Update the sleep state based on particle movement etc
+					ProcessIslandSleep(Island, Dt);
 
-				if (Island->Flags.bIsSleeping)
+					if (!!Island->Flags.bIsSleeping)
+					{
+						PropagateIslandSleep(Island);
+					}
+				}
+				else
 				{
-					PropagateIslandSleep(Island);
+					// Update the sleep state based on particle movement
+					// Deactivate non-moving particles
+					// Deactivate constraints between sleeping particles
+					ProcessPartialIslandSleep(Island, Dt);
 				}
 			}
 
@@ -2237,7 +2305,7 @@ namespace Chaos::Private
 						SleepCounterThreshold = FMath::Min(SleepCounterThreshold, TNumericLimits<int8>::Max());
 
 						// Did we exceed the velocity threshold?
-						if ((Rigid.VSmooth().SizeSquared() > FMath::Square(SleepLinearThreshold)) 
+						if ((Rigid.VSmooth().SizeSquared() > FMath::Square(SleepLinearThreshold))
 							|| (Rigid.WSmooth().SizeSquared() > FMath::Square(SleepAngularThreshold)))
 						{
 							continue;
@@ -2283,7 +2351,7 @@ namespace Chaos::Private
 			{
 				continue;
 			}
-			
+
 			UpdateParticleSleepMetrics(*Rigid, Dt);
 
 			// Did we exceed the velocity threshold?
@@ -2383,6 +2451,371 @@ namespace Chaos::Private
 		}
 	}
 
+	void FPBDIslandManager::ProcessPartialIslandSleep(FPBDIsland* Island, const FRealSingle Dt)
+	{
+		// Make all non-moving particles in the island go to sleep
+		SleepNonMovingParticles(Island, Dt);
+
+		// Deactivate all adjacent constraints between (a) 2 sleeping particles or (b) a sleeping and a non-dynamic particle
+		PropagateParticleSleepToConstraints(Island);
+	}
+
+	void FPBDIslandManager::SleepNonMovingParticles(FPBDIsland* Island, const FRealSingle Dt)
+	{
+		bool bAllNodesSleeping = true;
+
+		TArray<FGeometryParticleHandle*> SleptParticles;
+		for (FPBDIslandParticle* Node : Island->GetParticles())
+		{
+			// All zeroes means never sleep/disable
+			if ((Node->SleepLinearThresholdSq <= 0) && (Node->SleepAngularThresholdSq <= 0))
+			{
+				bAllNodesSleeping = false;
+				break;
+			}
+
+			// No need to check sleeping nodes again
+			if (!!Node->Flags.bIsSleeping)
+			{
+				continue;
+			}
+
+			FPBDRigidParticleHandle* Rigid = Node->GetParticle()->CastToRigidParticle();
+			if (Rigid == nullptr)
+			{
+				continue;
+			}
+
+			// Check the particle state against the thresholds
+			int32 SleepCounterThreshold = Node->SleepCounterThreshold;
+			int32 SleepCounter = 0;
+			if (SleepCounterThreshold < TNumericLimits<int32>::Max())
+			{
+				// Particles have a max sleep counter of 127 (to reduce counter space in the particle)
+				SleepCounterThreshold = FMath::Min(SleepCounterThreshold, TNumericLimits<int8>::Max());
+
+				UpdateParticleSleepMetrics(*Rigid, Dt);
+
+				// Did we exceed the velocity threshold?
+				if ((Rigid->VSmooth().SizeSquared() > Node->SleepLinearThresholdSq)
+					|| (Rigid->WSmooth().SizeSquared() > Node->SleepAngularThresholdSq))
+				{
+					bAllNodesSleeping = false;
+					continue;
+				}
+
+				// If we get here we might want to sleep
+				// Update the sleep counter and sleep if we exceed it
+				static_assert(sizeof(decltype(Rigid->SleepCounter())) == 1, "Expected int8 for SleepCounter(). Update clamp below");
+				SleepCounter = FMath::Min(int32(Rigid->SleepCounter()) + 1, int32(TNumericLimits<int8>::Max()));
+				if (SleepCounter > SleepCounterThreshold)
+				{
+					SleptParticles.Add(Rigid->Handle());
+					Node->Flags.bIsSleeping = true;
+					SleepCounter = 0;
+				}
+				else
+				{
+					bAllNodesSleeping = false;
+				}
+				Rigid->SetSleepCounter(int8(SleepCounter));
+			}
+		}
+		// Flag the island if all particles are sleeping
+		if (bAllNodesSleeping)
+		{
+			Island->Flags.bIsSleeping = true;
+		}
+
+		if (!SleptParticles.IsEmpty())
+		{
+			Particles.DeactivateParticles(SleptParticles);
+		}
+	}
+
+	void FPBDIslandManager::PropagateParticleSleepToConstraints(FPBDIsland* Island)
+	{
+		for (TArray<FPBDIslandConstraint*>& IslandEdges : Island->ContainerEdges)
+		{
+			for (FPBDIslandConstraint* IslandEdge : IslandEdges)
+			{
+				// Skip already sleeping constraints
+				if (!IslandEdge->Flags.bIsSleeping)
+				{
+					const FPBDIslandParticle* Node0 = IslandEdge->Nodes[0];
+					const FPBDIslandParticle* Node1 = IslandEdge->Nodes[1];
+					FConstraintHandle* Constraint = IslandEdge->GetConstraint();
+
+					if (Node0 != nullptr && Node1 != nullptr)
+					{
+						// Deactivate constraint if:
+						// Case (a) Both particles are sleeping
+						// Case (b) One particle is sleeping and the other is not dynamic (i.e. kinematic, static, etc.)
+						const bool bAreBothSleeping = !!Node0->Flags.bIsSleeping && !!Node1->Flags.bIsSleeping;
+						const bool bIsSleeping0Kinematic1 = !!Node0->Flags.bIsSleeping && !Node1->Flags.bIsDynamic;
+						const bool bIsSleeping1Kinematic0 = !!Node1->Flags.bIsSleeping && !Node0->Flags.bIsDynamic;
+						if (bAreBothSleeping || bIsSleeping0Kinematic1 || bIsSleeping1Kinematic0)
+						{
+							IslandEdge->Flags.bIsSleeping = true;
+							Constraint->SetIsSleeping(true);
+						}
+					}
+					else
+					{
+						// If one node is NULL, the edge sleep state should match the other node.
+						if (Node0 != nullptr) // && Node1 == nullptr
+						{
+							IslandEdge->Flags.bIsSleeping = Node0->Flags.bIsSleeping;
+							Constraint->SetIsSleeping(Node0->Flags.bIsSleeping);
+						}
+						else if (Node1 != nullptr) // && Node0 == nullptr
+						{
+							IslandEdge->Flags.bIsSleeping = Node1->Flags.bIsSleeping;
+							Constraint->SetIsSleeping(Node1->Flags.bIsSleeping);
+						}
+						else
+						{
+							IslandEdge->Flags.bIsSleeping = false;
+							Constraint->SetIsSleeping(false);
+						}
+					}
+				}
+			}
+		}
+	}
+
+	void FPBDIslandManager::ProcessParticleWakes(FPBDIsland* Island)
+	{
+		TArray<FGeometryParticleHandle*> ParticlesToWakeUp;
+		for (TArray<FPBDIslandConstraint*>& IslandEdges : Island->ContainerEdges)
+		{
+			for (FPBDIslandConstraint* IslandEdge : IslandEdges)
+			{
+				// For now, only new constraints can induce partial island wakes
+				if (!!IslandEdge->Flags.bIsNew)
+				{
+					IslandEdge->Flags.bWasProcessed = true;
+
+					// Only propagated the wake-up event if at least one particle connected be the new constraint is woken up.
+					// This means there will not be any wake-up propagation if
+					// 1) both particles are already awake.
+					// 2) both particles are sleeping.
+					// 3) neither particle is moving if (bChaosSolverWakeUpOnConstraintWithNonMovingParticle == false).
+					if (ProcessConnectedParticles(IslandEdge, ParticlesToWakeUp, !CVars::bChaosSolverWakeOnMovingOnly))
+					{
+						// Wake-up propagation: Breadth-first traversal of the constraint graph
+						// If level information is available, only propagate the wake-up away from kinematic particles (i.e. Node->Level >= Edge->Level).
+						// This will cause the top of a stack to wake up if a particles closer to the kinematic ground wakes up.
+						// Meanwhile particles closer to the ground remains sleeping.
+						TArray<FPBDIslandParticle*> NodeQueue;
+						NodeQueue.Reserve(Island->Nodes.Num());
+						NodeQueue.Add(IslandEdge->Nodes[0]);
+						NodeQueue.Add(IslandEdge->Nodes[1]);
+						for (int32 NodeIndex = 0; NodeIndex < NodeQueue.Num(); ++NodeIndex)
+						{
+							FPBDIslandParticle* Node = NodeQueue[NodeIndex];
+							for (FPBDIslandConstraint* Edge : Node->Edges)
+							{
+								if (!Edge->Flags.bWasProcessed)
+								{
+									FPBDIslandParticle* OtherNode = (Node == Edge->Nodes[0]) ? Edge->Nodes[1] : Edge->Nodes[0];
+									// Wake up the OtherNode if the node level is greater than or equal to the edge level (if level information available).
+									if ((OtherNode != nullptr) && OtherNode->Flags.bIsDynamic
+										&& bAssignLevels && (OtherNode->Level >= Edge->Level))
+									{
+										ProcessConnectedParticles(Edge, ParticlesToWakeUp, true);
+										NodeQueue.Add(OtherNode);
+									}
+								}
+							}
+						}
+					}
+				}
+			}
+		}
+		WakeUpParticles(ParticlesToWakeUp);
+	}
+
+	// Checks both particles connected to the edge and adds the sleeping particle to ParticlesToWake if
+	//	   (bAlwaysWakeUp == true): if both particles are dynamic and one particle is awake while the other one is sleeping
+	//     (bAlwaysWakeUp == false): if satisfying the conditions above and the awake particle is also moving
+	bool FPBDIslandManager::ProcessConnectedParticles(const FPBDIslandConstraint* Edge, TArray<FGeometryParticleHandle*>& ParticlesToWake, const bool bAlwaysWakeUp)
+	{
+		FPBDIslandParticle* Node0 = Edge->Nodes[0];
+		FPBDIslandParticle* Node1 = Edge->Nodes[1];
+
+		const auto IsMoving = [](FPBDIslandParticle* Node) -> bool
+			{
+				if (Node == nullptr)
+				{
+					return false;
+				}
+				FPBDRigidParticleHandle* Rigid = Node->GetParticle()->CastToRigidParticle();
+				if (Rigid == nullptr)
+				{
+					return false;
+				}
+
+				// Use the smoothed velocity to check for motion, unless it is zero (i.e. was never set).
+				const FReal VSmoothed = Rigid->VSmooth().SizeSquared();
+				const FReal WSmoothed = Rigid->WSmooth().SizeSquared();
+				const bool bIsVelSmoothUnset = (VSmoothed == 0) && (WSmoothed == 0);
+				if (bIsVelSmoothUnset)
+				{
+					// Use the regular velocity from the last time step
+					if ((Rigid->GetPreV().SizeSquared() > Node->SleepLinearThresholdSq)
+						|| (Rigid->GetPreW().SizeSquared() > Node->SleepAngularThresholdSq))
+					{
+						return true;
+					}
+				}
+				else
+				{
+					// Use the smooth velocities from the last step
+					if ((VSmoothed > Node->SleepLinearThresholdSq)
+						|| (WSmoothed > Node->SleepAngularThresholdSq))
+					{
+						return true;
+					}
+				}
+
+				return false;
+			};
+
+		const auto IsAwake = [](const FPBDIslandParticle* Node) -> bool
+			{
+				return Node && !Node->Flags.bIsSleeping && !!Node->Flags.bIsDynamic;
+			};
+
+		const auto IsDynamic = [](const FPBDIslandParticle* Node) -> bool
+			{
+				return Node && !!Node->Flags.bIsDynamic;
+			};
+
+		// If one particle is awake and the other one is sleeping, wake up the other dynamic particle.
+		const bool bIsAwake0 = IsAwake(Node0);
+		const bool bIsAwake1 = IsAwake(Node1);
+		const bool bIsDynamic0 = IsDynamic(Node0);
+		const bool bIsDynamic1 = IsDynamic(Node1);
+		const bool bTriggerWakeUp0 = bAlwaysWakeUp || IsMoving(Node0);
+		const bool bTriggerWakeUp1 = bAlwaysWakeUp || IsMoving(Node1);
+		if (bIsDynamic1 && !bIsAwake1 && bIsAwake0 && bTriggerWakeUp0)
+		{
+			Node1->Flags.bIsSleeping = false;
+			ParticlesToWake.Add(Node1->GetParticle());
+			return true;
+		}
+		else if (bIsDynamic0 && !bIsAwake0 && bIsAwake1 && bTriggerWakeUp1)
+		{
+			Node0->Flags.bIsSleeping = false;
+			ParticlesToWake.Add(Node0->GetParticle());
+			return true;
+		}
+		return false;
+	}
+
+	void FPBDIslandManager::SyncConstraintSleepStateToParticles(FPBDIsland* Island)
+	{
+		// Iterate over the edges
+		// Sync the constraint sleep state to match the connected particles
+		for (TArray<FPBDIslandConstraint*>& IslandEdges : Island->ContainerEdges)
+		{
+			for (FPBDIslandConstraint* IslandEdge : IslandEdges)
+			{
+				const FPBDIslandParticle* Node0 = IslandEdge->Nodes[0];
+				const FPBDIslandParticle* Node1 = IslandEdge->Nodes[1];
+				FConstraintHandle* Constraint = IslandEdge->GetConstraint();
+
+				if (Node0 != nullptr && Node1 != nullptr)
+				{
+					// Deactivate constraint if:
+					// Case (a) Both particles are sleeping
+					// Case (b) One particle is sleeping and the other is not dynamic (i.e. kinematic, static, etc.)
+					const bool bAreBothSleeping = !!Node0->Flags.bIsSleeping && !!Node1->Flags.bIsSleeping;
+					const bool bIsSleeping0Kinematic1 = !!Node0->Flags.bIsSleeping && !Node1->Flags.bIsDynamic;
+					const bool bIsSleeping1Kinematic0 = !!Node1->Flags.bIsSleeping && !Node0->Flags.bIsDynamic;
+					if (bAreBothSleeping || bIsSleeping0Kinematic1 || bIsSleeping1Kinematic0)
+					{
+						IslandEdge->Flags.bIsSleeping = true;
+						Constraint->SetIsSleeping(true);
+					}
+
+					// Activate the constraint if:
+					// Case (a) at least one particle is awake
+					const bool bIsAwakeDynamic0 = !Node0->Flags.bIsSleeping && !!Node0->Flags.bIsDynamic;
+					const bool bIsAwakeDynamic1 = !Node1->Flags.bIsSleeping && !!Node1->Flags.bIsDynamic;
+					if (bIsAwakeDynamic0 || bIsAwakeDynamic1)
+					{
+						Constraint->SetIsSleeping(false);
+						IslandEdge->Flags.bIsSleeping = false;
+					}
+				}
+				else
+				{
+					// If one node is NULL, the edge sleep state should match the other node.
+					if (Node0 != nullptr) // && Node1 == nullptr
+					{
+						IslandEdge->Flags.bIsSleeping = Node0->Flags.bIsSleeping;
+						Constraint->SetIsSleeping(Node0->Flags.bIsSleeping);
+					}
+					else if (Node1 != nullptr) // && Node0 == nullptr
+					{
+						IslandEdge->Flags.bIsSleeping = Node1->Flags.bIsSleeping;
+						Constraint->SetIsSleeping(Node1->Flags.bIsSleeping);
+					}
+					else // Node0 == nullptr && Node1 == nullptr
+					{
+						IslandEdge->Flags.bIsSleeping = false;
+						Constraint->SetIsSleeping(false);
+					}
+				}
+
+				// Reset
+				IslandEdge->Flags.bIsNew = false;
+				IslandEdge->Flags.bWasProcessed = false;
+			}
+		}
+	}
+
+	void FPBDIslandManager::WakeUpParticles(TArray<FGeometryParticleHandle*>& ParticlesToWake)
+	{
+		if (!ParticlesToWake.IsEmpty())
+		{
+			const auto UpdateShapeTransforms = [](const TArray<FGeometryParticleHandle*>& Particles)
+				{
+					for (FGeometryParticleHandle* Particle : Particles)
+					{
+						for (const TUniquePtr<FPerShapeData>& Shape : Particle->ShapesArray())
+						{
+							Shape->UpdateLeafWorldTransform(Particle);
+						}
+					}
+				};
+
+			Particles.ActivateParticles(ParticlesToWake);
+			for (FGeometryParticleHandle* Particle : ParticlesToWake)
+			{
+				if (FPBDRigidParticleHandle* Rigid = Particle->CastToRigidParticle())
+				{
+					Rigid->SetSleepCounter(0);
+				}
+			}
+			UpdateShapeTransforms(ParticlesToWake);
+		}
+	}
+
+	void FPBDIslandManager::ResetEdgeFlags(FPBDIsland* Island)
+	{
+		for (TArray<FPBDIslandConstraint*>& IslandEdges : Island->ContainerEdges)
+		{
+			for (FPBDIslandConstraint* IslandEdge : IslandEdges)
+			{
+				IslandEdge->Flags.bIsNew = false;
+				IslandEdge->Flags.bWasProcessed = false;
+			}
+		}
+	}
+
 	void FPBDIslandManager::ProcessDisable(TFunctionRef<void(FPBDRigidParticleHandle*)> ParticleDisableFunctor)
 	{
 		// @todo(chaos): parallelize
@@ -2531,6 +2964,7 @@ namespace Chaos::Private
 				{
 					CHAOS_CONSTRAINTGRAPH_VALIDATE_TEST(Node->Particle->Sleeping())
 				}
+				CHAOS_CONSTRAINTGRAPH_VALIDATE_TEST(Node->Particle->Sleeping() == Node->Flags.bIsSleeping);
 			}
 
 			// Make sure the edge points to the node and that we're in the same island
@@ -2568,10 +3002,26 @@ namespace Chaos::Private
 				CHAOS_CONSTRAINTGRAPH_VALIDATE_TEST(Edge->Island->ContainerEdges[Edge->ContainerIndex][Edge->IslandArrayIndex] == Edge);
 
 				// Make sure the island and edge sleep states match
-				if (Edge->Constraint->SupportsSleeping())
+				if (!CVars::bChaosSolverPartialIslandSleep && Edge->Constraint->SupportsSleeping())
 				{
 					CHAOS_CONSTRAINTGRAPH_VALIDATE_TEST(Edge->Island->Flags.bIsSleeping == Edge->Flags.bIsSleeping);
 				}
+
+				// Make sure the sleep state matches between the edge and connected nodes
+				if (Edge->Flags.bIsSleeping)
+				{
+					// The edge should be sleeping if it is connected to 2 non-moving nodes, i.e. dynamic, sleeping OR kinematic, not moving
+					const bool bIsNode0NotMoving = (Edge->Nodes[0] != nullptr) ? Edge->Nodes[0]->Flags.bIsSleeping || !Edge->Nodes[0]->Flags.bIsMoving : true;
+					const bool bIsNode1NotMoving = (Edge->Nodes[1] != nullptr) ? Edge->Nodes[1]->Flags.bIsSleeping || !Edge->Nodes[1]->Flags.bIsMoving : true;
+					CHAOS_CONSTRAINTGRAPH_VALIDATE_TEST(bIsNode0NotMoving && bIsNode1NotMoving);
+				}
+				else
+				{
+					// The edge should be awake if it is connected to at least 1 non-sleeping dynamic or 1 moving kinematic.
+					const bool bIsConnectedToNonSleeping = (Edge->Nodes[0] != nullptr ? !Edge->Nodes[0]->Flags.bIsSleeping : false) || (Edge->Nodes[1] != nullptr ? !Edge->Nodes[1]->Flags.bIsSleeping : false);
+					const bool bIsConnectedToMovingKinematic = (Edge->Nodes[0] != nullptr ? Edge->Nodes[0]->Flags.bIsMoving : false) || (Edge->Nodes[1] != nullptr ? Edge->Nodes[1]->Flags.bIsMoving : false);
+					CHAOS_CONSTRAINTGRAPH_VALIDATE_TEST(bIsConnectedToNonSleeping || bIsConnectedToMovingKinematic);
+				}
 			}
 
 			// Constraint sleep state should match edge
@@ -2579,6 +3029,10 @@ namespace Chaos::Private
 			{
 				CHAOS_CONSTRAINTGRAPH_VALIDATE_TEST(Edge->Flags.bIsSleeping == Edge->Constraint->IsSleeping());
 			}
+
+			// Make sure that flags were reset to false
+			CHAOS_CONSTRAINTGRAPH_VALIDATE_TEST(!Edge->Flags.bIsNew);
+			CHAOS_CONSTRAINTGRAPH_VALIDATE_TEST(!Edge->Flags.bWasProcessed);
 		}
 
 		// Check various Island properties
@@ -2631,6 +3085,14 @@ namespace Chaos::Private
 		return true;
 	}
 
+	void FPBDIslandManager::UpdateAllParticleSleepMetrics(const FRealSingle Dt)
+	{
+		for (FTransientPBDRigidParticleHandle& Rigid : Particles.GetActiveDynamicMovingKinematicParticlesView())
+		{
+			UpdateParticleSleepMetrics(Rigid, Dt);
+		}
+	}
+
 #if CHAOS_DEBUG_DRAW
 	void FPBDIslandManager::DebugDrawSleepState(const DebugDraw::FChaosDebugDrawSettings* DebugDrawSettings) const
 	{
@@ -2669,17 +3131,32 @@ namespace Chaos::Private
 		// Loop over particles in islands
 		for (FPBDIsland* Island : Islands)
 		{
+			const FColor SleepingColor = FColor::Black;
+
+			if (Island->Flags.bIsSleeping)
+			{
+				for (FPBDIslandParticle* Node : Island->GetParticles())
+				{
+					DebugDraw::DrawParticleShapes(FRigidTransform3::Identity, Node->GetParticle(), SleepingColor, DebugDrawSettings);
+				}
+				continue;
+			}
+
 			for (FPBDIslandParticle* Node : Island->GetParticles())
 			{
 				FColor Color = FColor(128, 128, 128);
 
-				if (!Island->Flags.bIsSleeping)
+				// All zeroes means never sleep/disable
+				if ((Node->SleepLinearThresholdSq > 0) || (Node->SleepAngularThresholdSq > 0))
 				{
-					// All zeroes means never sleep/disable
-					if ((Node->SleepLinearThresholdSq > 0) || (Node->SleepAngularThresholdSq > 0))
+					// Check the particle state against the thresholds
+					if (FPBDRigidParticleHandle* Rigid = Node->GetParticle()->CastToRigidParticle())
 					{
-						// Check the particle state against the thresholds
-						if (FPBDRigidParticleHandle* Rigid = Node->GetParticle()->CastToRigidParticle())
+						if (Rigid->Sleeping())
+						{
+							Color = SleepingColor;
+						}
+						else
 						{
 							// Did we exceed the velocity threshold?
 							const bool bIsParticlePreventingSleep = ((Rigid->VSmooth().SizeSquared() > Node->SleepLinearThresholdSq) || (Rigid->WSmooth().SizeSquared() > Node->SleepAngularThresholdSq));
diff --git a/Engine/Source/Runtime/Experimental/Chaos/Private/Chaos/Joint/PBDJointContainerSolver.cpp b/Engine/Source/Runtime/Experimental/Chaos/Private/Chaos/Joint/PBDJointContainerSolver.cpp
--- a/Engine/Source/Runtime/Experimental/Chaos/Private/Chaos/Joint/PBDJointContainerSolver.cpp
+++ b/Engine/Source/Runtime/Experimental/Chaos/Private/Chaos/Joint/PBDJointContainerSolver.cpp
@@ -149,10 +149,16 @@ namespace Chaos
 			int32 CurrentIndex = 0;
 			for (Private::FPBDIslandConstraint* IslandConstraint : IslandConstraints)
 			{
-				// We will only ever be given constraints from our container (asserts in non-shipping)
-				const int32 ContainerConstraintIndex = IslandConstraint->GetConstraint()->AsUnsafe<FPBDJointConstraintHandle>()->GetConstraintIndex();
+				FConstraintHandle* Constraint = IslandConstraint->GetConstraint();
 
-				AddConstraint(ContainerConstraintIndex);
+				// Filter out sleeping constraints in any partially sleeping island
+				if (!Constraint->IsSleeping())
+				{
+					// We will only ever be given constraints from our container (asserts in non-shipping)
+					const int32 ContainerConstraintIndex = Constraint->AsUnsafe<FPBDJointConstraintHandle>()->GetConstraintIndex();
+
+					AddConstraint(ContainerConstraintIndex);
+				}
 			}
 		}
 
diff --git a/Engine/Source/Runtime/Experimental/Chaos/Public/Chaos/Evolution/IndexedConstraintContainer.h b/Engine/Source/Runtime/Experimental/Chaos/Public/Chaos/Evolution/IndexedConstraintContainer.h
--- a/Engine/Source/Runtime/Experimental/Chaos/Public/Chaos/Evolution/IndexedConstraintContainer.h
+++ b/Engine/Source/Runtime/Experimental/Chaos/Public/Chaos/Evolution/IndexedConstraintContainer.h
@@ -157,12 +157,18 @@ namespace Chaos
 		{
 			for (const Private::FPBDIslandConstraint* IslandConstraint : IslandConstraints)
 			{
-				// We will only ever be given constraints from our container (asserts in non-shipping)
-				const FIndexedConstraintHandle* IndexedConstraintHandle = IslandConstraint->GetConstraint()->AsUnsafe<FIndexedConstraintHandle>();
+				FConstraintHandle* Constraint = IslandConstraint->GetConstraint();
 
-				int32 ConstraintIndex = IndexedConstraintHandle->GetConstraintIndex();
+				// Filter out sleeping constraints in any partially sleeping island
+				if (!Constraint->IsSleeping())
+				{
+					// We will only ever be given constraints from our container (asserts in non-shipping)
+					const FIndexedConstraintHandle* IndexedConstraintHandle = Constraint->AsUnsafe<FIndexedConstraintHandle>();
 
-				ConstraintIndices.Add(ConstraintIndex);
+					int32 ConstraintIndex = IndexedConstraintHandle->GetConstraintIndex();
+
+					ConstraintIndices.Add(ConstraintIndex);
+				}
 			}
 		}
 
diff --git a/Engine/Source/Runtime/Experimental/Chaos/Public/Chaos/Island/IslandManager.h b/Engine/Source/Runtime/Experimental/Chaos/Public/Chaos/Island/IslandManager.h
--- a/Engine/Source/Runtime/Experimental/Chaos/Public/Chaos/Island/IslandManager.h
+++ b/Engine/Source/Runtime/Experimental/Chaos/Public/Chaos/Island/IslandManager.h
@@ -170,6 +170,8 @@ namespace Chaos::Private
 			struct
 			{
 				uint32 bIsSleeping : 1;
+				uint32 bIsNew : 1;
+				uint32 bWasProcessed : 1;
 			};
 		private:
 			uint32 Bits;
@@ -582,7 +584,7 @@ namespace Chaos::Private
 		CHAOS_API void SortIslandEdges(FPBDIsland* Island);
 		CHAOS_API void RandomizeIslandEdges(FPBDIsland* Island);
 
-		// Sleeping
+		// Island Sleeping
 		CHAOS_API void ProcessSleep(const FRealSingle Dt);
 		CHAOS_API void ProcessParticlesSleep(const FRealSingle Dt);
 		CHAOS_API void ProcessIslandSleep(FPBDIsland* Island, const FRealSingle Dt);
@@ -590,6 +592,18 @@ namespace Chaos::Private
 		CHAOS_API void PropagateIslandSleepToParticles(FPBDIsland* Island);
 		CHAOS_API void PropagateIslandSleepToConstraints(FPBDIsland* Island);
 
+		// Partial Island Sleeping
+		CHAOS_API void ProcessPartialIslandSleep(FPBDIsland* Island, const FRealSingle Dt);
+		CHAOS_API void SleepNonMovingParticles(FPBDIsland* Island, const FRealSingle Dt);
+		CHAOS_API void PropagateParticleSleepToConstraints(FPBDIsland* Island);
+
+		// Partial Island Waking
+		CHAOS_API void ProcessParticleWakes(FPBDIsland* Island);
+		CHAOS_API bool ProcessConnectedParticles(const FPBDIslandConstraint* Edge, TArray<FGeometryParticleHandle*>& ParticlesToWake, const bool bAlwaysWakeUp);
+		CHAOS_API void SyncConstraintSleepStateToParticles(FPBDIsland* Island);
+		CHAOS_API void WakeUpParticles(TArray<FGeometryParticleHandle*>& ParticlesToWake);
+		CHAOS_API void ResetEdgeFlags(FPBDIsland* Island);
+
 		// Disabling
 		CHAOS_API void ProcessDisable(TFunctionRef<void(FPBDRigidParticleHandle*)> ParticleDisableFunctor);
 		CHAOS_API void ProcessParticlesDisable(TArray<FPBDRigidParticleHandle*>& OutDisableParticles);
@@ -598,6 +612,7 @@ namespace Chaos::Private
 		CHAOS_API void FinalizeIslands();
 		CHAOS_API void ApplyDeterminism();
 		CHAOS_API bool Validate() const;
+		CHAOS_API void UpdateAllParticleSleepMetrics(const FRealSingle Dt);
 
 		CHAOS_API int32 GetNextVisitEpoch();
 
@@ -666,7 +681,10 @@ namespace Chaos::Private
 			{
 				for (FPBDIslandConstraint* Edge : Island->ContainerEdges[ContainerId])
 				{
-					Visitor(Edge);
+					if (!Edge->Flags.bIsSleeping)
+					{
+						Visitor(Edge);
+					}
 				}
 			}
 		}
