Partial Island Sleeping - Game thread sleeping unit tests

diff --git a/Engine/Source/Programs/HeadlessChaos/Private/HeadlessChaosTestSleeping.cpp b/Engine/Source/Programs/HeadlessChaos/Private/HeadlessChaosTestSleeping.cpp
new file mode 100644
--- /dev/null
+++ b/Engine/Source/Programs/HeadlessChaos/Private/HeadlessChaosTestSleeping.cpp
@@ -0,0 +1,661 @@
+// Copyright Epic Games, Inc. All Rights Reserved.
+
+#include "ChaosSolversModule.h"
+#include "HeadlessChaosTestUtility.h"
+#include "PBDRigidsSolver.h"
+#include "PhysicsProxy/SingleParticlePhysicsProxy.h"
+#include "PhysicsProxy/JointConstraintProxy.h"
+#include "Chaos/PBDJointConstraintData.h"
+#include "HAL/ConsoleManager.h"
+
+namespace ChaosTest
+{
+	using namespace Chaos;
+
+	class SleepingTests : public ::testing::TestWithParam<bool>
+	{
+	};
+
+	class FSleepingTest
+	{
+	public:
+		FSleepingTest(const bool bPartialSleeping)
+			: MaterialData()
+			, FloorProxy()
+			, ParticleProxies()
+			, ParticleHandles()
+			, JointHandles()
+			, TickCount(0)
+			, CVarPartialSleeping()
+		{
+			// Create a new material before creating the solver
+			FMaterialHandle MaterialHandle = FPhysicalMaterialManager::Get().Create();
+			FPhysicsMaterial* PhysicsMaterial = MaterialHandle.Get();
+			PhysicsMaterial->Friction = 0.7;
+			PhysicsMaterial->StaticFriction = 0.0;
+			PhysicsMaterial->Restitution = 0.0;
+			PhysicsMaterial->Density = 1.0;
+			PhysicsMaterial->SleepingLinearThreshold = 1.0;
+			PhysicsMaterial->SleepingAngularThreshold = 0.05;
+			PhysicsMaterial->SleepCounterThreshold = 4;
+			MaterialData.Materials.Add(MaterialHandle);
+
+			Module = FChaosSolversModule::GetModule();
+			Solver = Module->CreateSolver(nullptr, /*AsyncDt=*/-1);
+			InitSolverSettings(Solver);
+			Solver->SetThreadingMode_External(EThreadingModeTemp::SingleThread);
+
+			CVarPartialSleeping = IConsoleManager::Get().FindConsoleVariable(TEXT("p.Chaos.Solver.Sleep.PartialIslandSleep"), false);
+			check(CVarPartialSleeping && CVarPartialSleeping->IsVariableBool());
+			CVarPartialSleeping->Set(bPartialSleeping);
+		}
+
+		~FSleepingTest()
+		{
+			for (FSingleParticlePhysicsProxy* ParticleProxy : ParticleProxies)
+			{
+				Solver->UnregisterObject(ParticleProxy);
+			}
+			for (FJointConstraint* Joint : JointHandles)
+			{
+				if (Joint)
+					Solver->UnregisterObject(Joint);
+			}
+			if (FloorProxy)
+				Solver->UnregisterObject(FloorProxy);
+			Module->DestroySolver(Solver);
+		}
+
+		FRigidBodyHandle_External* MakeFloor(const FVec3& Position = FVec3(0))
+		{
+			FloorProxy = FSingleParticlePhysicsProxy::Create(Chaos::FGeometryParticle::CreateParticle());
+			FRigidBodyHandle_External* FloorParticle = &FloorProxy->GetGameThreadAPI();
+			FImplicitObjectPtr FloorGeom = Chaos::FImplicitObjectPtr(new TBox<FReal, 3>(FVec3(-5000, -5000, -100), FVec3(5000, 5000, 0)));
+			FloorParticle->SetGeometry(FloorGeom);
+			FloorParticle->SetObjectState(EObjectStateType::Static);
+			Solver->RegisterObject(FloorProxy);
+			FloorParticle->SetX(Position);
+			ChaosTest::SetParticleSimDataToCollide({ FloorProxy->GetParticle_LowLevel() });
+			for (const TUniquePtr<FPerShapeData>& Shape : FloorParticle->ShapesArray())
+			{
+				Shape->SetMaterialData(MaterialData);
+			}
+			return FloorParticle;
+		}
+
+		FRigidBodyHandle_External* MakeBox(const FVec3& Position, const FVec3& Dimensions = FVec3(200, 400, 100), const FReal Mass = 1)
+		{
+			FSingleParticlePhysicsProxy* BoxProxy = FSingleParticlePhysicsProxy::Create(Chaos::FPBDRigidParticle::CreateParticle());
+			FRigidBodyHandle_External* BoxParticle = &BoxProxy->GetGameThreadAPI();
+			const FVec3 HalfSize = Dimensions / 2.0f;
+			FImplicitObjectPtr CollidingCubeGeom = Chaos::FImplicitObjectPtr(new TBox<FReal, 3>(-HalfSize, HalfSize));
+			BoxParticle->SetGeometry(CollidingCubeGeom);
+			Solver->RegisterObject(BoxProxy);
+			BoxParticle->SetGravityEnabled(true);
+			BoxParticle->SetX(Position);
+			SetBoxInertiaTensor(*BoxParticle, /*Dimensions=*/Dimensions, /*Mass=*/Mass);
+			ChaosTest::SetParticleSimDataToCollide({ BoxProxy->GetParticle_LowLevel() });
+			for (const TUniquePtr<FPerShapeData>& Shape : BoxParticle->ShapesArray())
+			{
+				Shape->SetMaterialData(MaterialData);
+			}
+
+			ParticleProxies.Add(BoxProxy);
+			ParticleHandles.Add(BoxParticle);
+
+			return BoxParticle;
+		}
+
+		FRigidBodyHandle_External* MakeCube(const FVec3& Position, const FReal Size = 200, const FReal Mass = 1)
+		{
+			FSingleParticlePhysicsProxy* CubeProxy = FSingleParticlePhysicsProxy::Create(Chaos::FPBDRigidParticle::CreateParticle());
+			FRigidBodyHandle_External* CubeParticle = &CubeProxy->GetGameThreadAPI();
+			const FReal HalfSize = Size / 2.0f;
+			FImplicitObjectPtr CollidingCubeGeom = Chaos::FImplicitObjectPtr(new TBox<FReal, 3>(FVec3(-HalfSize), FVec3(HalfSize)));
+			CubeParticle->SetGeometry(CollidingCubeGeom);
+			Solver->RegisterObject(CubeProxy);
+			CubeParticle->SetGravityEnabled(true);
+			CubeParticle->SetX(Position);
+			SetCubeInertiaTensor(*CubeParticle, /*Dimension=*/Size, /*Mass=*/Mass);
+			ChaosTest::SetParticleSimDataToCollide({ CubeProxy->GetParticle_LowLevel() });
+			for (const TUniquePtr<FPerShapeData>& Shape : CubeParticle->ShapesArray())
+			{
+				Shape->SetMaterialData(MaterialData);
+			}
+
+			ParticleProxies.Add(CubeProxy);
+			ParticleHandles.Add(CubeParticle);
+
+			return CubeParticle;
+		}
+
+		FRigidBodyHandle_External* MakeSphere(const FVec3& Position, const FReal Radius = 100, const FReal Mass = 1)
+		{
+			FSingleParticlePhysicsProxy* SphereProxy = FSingleParticlePhysicsProxy::Create(Chaos::FPBDRigidParticle::CreateParticle());
+			FRigidBodyHandle_External* SphereParticle = &SphereProxy->GetGameThreadAPI();
+			FImplicitObjectPtr CollidingSphereGeom = Chaos::FImplicitObjectPtr(new TSphere<FReal, 3>(FVec3(0), Radius));
+			SphereParticle->SetGeometry(CollidingSphereGeom);
+			Solver->RegisterObject(SphereProxy);
+			SphereParticle->SetGravityEnabled(true);
+			SphereParticle->SetX(Position);
+			SetSphereInertiaTensor(*SphereParticle, /*Radius=*/Radius, /*Mass=*/Mass);
+			ChaosTest::SetParticleSimDataToCollide({ SphereProxy->GetParticle_LowLevel() });
+			for (const TUniquePtr<FPerShapeData>& Shape : SphereParticle->ShapesArray())
+			{
+				Shape->SetMaterialData(MaterialData);
+			}
+
+			ParticleProxies.Add(SphereProxy);
+			ParticleHandles.Add(SphereParticle);
+
+			return SphereParticle;
+		}
+
+		FJointConstraint* AddJoint(const FProxyBasePair& InConstrainedParticles, const FTransformPair& InTransform)
+		{
+			FJointConstraint* Joint = new FJointConstraint();
+			Joint->SetParticleProxies(InConstrainedParticles);
+
+			FPBDJointSettings SettingsTemp = Joint->GetJointSettings();
+			SettingsTemp.ConnectorTransforms = InTransform;
+			Joint->SetJointSettings(SettingsTemp);
+
+			Solver->RegisterObject(Joint);
+
+			JointHandles.Add(Joint);
+			return Joint;
+		}
+
+		void MakeStackOfCubes(const int32 Num, const FReal Size = 200)
+		{
+			MakeFloor();
+			const FReal InitialHeight = Size / 2.0f;
+			for (int32 Id = 0; Id < Num; ++Id)
+			{
+				const FReal Z = InitialHeight + Size * Id;
+				MakeCube(FVec3(0, 0, Z), Size);
+			}
+		}
+
+		void MakeStackOfSpheres(const int32 Num, const FReal Radius = 100)
+		{
+			MakeFloor();
+			const FReal InitialHeight = Radius;
+			for (int32 Id = 0; Id < Num; ++Id)
+			{
+				const FReal Z = InitialHeight + 2 * Radius * Id;
+				MakeSphere(FVec3(0, 0, Z), Radius);
+			}
+		}
+
+		void MakeChainOfSpheres(const int32 Num, const FReal Radius = 100)
+		{
+			const int32 DistanceMultiplier = 3;
+			for (int32 Id = 0; Id < Num; ++Id)
+			{
+				const FReal Z = -DistanceMultiplier * Radius * Id;
+				FRigidBodyHandle_External* SphereParticle = MakeSphere(FVec3(0, 0, Z), Radius);
+			}
+			for (int32 Id = 0; Id < Num - 1; ++Id)
+			{
+				FRigidTransform3 Transform1 = FRigidTransform3(FVec3(0, 0, -DistanceMultiplier * Radius * 0.5), FRotation3::FromIdentity());
+				FRigidTransform3 Transform2 = FRigidTransform3(FVec3(0, 0, DistanceMultiplier * Radius * 0.5), FRotation3::FromIdentity());
+				AddJoint({ ParticleProxies[Id], ParticleProxies[Id + 1] }, { Transform1 , Transform2 });
+			}
+
+		}
+
+		// Build a brick wall using the following pattern (NumY = 5, NumZ = 3)
+		// 
+		// if (bShiftUnevenRows)
+		// P10 -P11 -P12 -P13 - P14
+		//   \  / \  / \  / \  /  \
+		//    P5 - P6 - P7 - P8 - P9
+		//	  / \  / \  / \  / \ /
+		//	P0 - P1 - P2 - P3 - P4
+		//  |    |    |    |    |
+		//          Floor
+		// 
+		// if (!bShiftUnevenRows)
+		// P10 -P11 -P12 -P13 -P14
+		//  |    |    |    |    |
+		//  P5 - P6 - P7 - P8 - P9
+		//  |    |    |    |    |
+		//	P0 - P1 - P2 - P3 - P4
+		//  |    |    |    |    |
+		//          Floor
+		// NOTE: The COM of P0 is located at (0, 0, 0) and the top of the floor at (0, 0, -HalfBrickSize).
+		void MakeBrickWall(const int32 NumY, const int32 NumZ, const FVec3& BrickDimensions = FVec3(100, 200, 50), const FReal Mass = 2, const bool bShiftUnevenRows = true)
+		{
+			const FReal HalfWidth = BrickDimensions.Y * 0.5;
+			const FReal HalfHeight = BrickDimensions.Z * 0.5;
+			MakeFloor(FVec3(0, 0, -HalfHeight));
+			for (int32 Z = 0; Z < NumZ; ++Z)
+			{
+				for (int32 Y = 0; Y < NumY; ++Y)
+				{
+					const bool bIsUneven = Z % 2 == 1;
+					if (bShiftUnevenRows && bIsUneven) // uneven
+					{
+						//if (Y + 1 == NumY) continue; // skip the last brick
+						// Shift every uneven row by HalfWidth
+						MakeBox(FVec3(0, Y * BrickDimensions.Y + HalfWidth, Z * BrickDimensions.Z), BrickDimensions, Mass);
+					}
+					else // even
+					{
+						MakeBox(FVec3(0, Y * BrickDimensions.Y, Z * BrickDimensions.Z), BrickDimensions, Mass);
+					}
+				}
+			}
+		}
+
+		void Advance()
+		{
+			Solver->AdvanceAndDispatch_External(1 / 60.0);
+			Solver->UpdateGameThreadStructures();
+			++TickCount;
+		}
+
+		void AdvanceUntilSleeping(const int32 MaxIterations = 200)
+		{
+			const int32 MaxTickCount = TickCount + MaxIterations;
+			bool bIsSleeping = false;
+			while (!bIsSleeping && (TickCount < MaxTickCount))
+			{
+				Advance();
+
+				bIsSleeping = true;
+				for (FRigidBodyHandle_External* ParticleHandle : ParticleHandles)
+				{
+					// Check that none of the particles is dynamic
+					if (ParticleHandle->ObjectState() == EObjectStateType::Dynamic)
+					{
+						bIsSleeping = false;
+					}
+				}
+			}
+
+			EXPECT_TRUE(bIsSleeping);
+			EXPECT_LT(TickCount, MaxTickCount);
+		}
+
+		FMaterialData MaterialData;
+
+		FSingleParticlePhysicsProxy* FloorProxy;
+		TArray<FSingleParticlePhysicsProxy*> ParticleProxies;
+		TArray<FRigidBodyHandle_External*> ParticleHandles;
+		TArray< FJointConstraint*> JointHandles;
+
+		FPBDRigidsSolver* Solver;
+		FChaosSolversModule* Module;
+
+		int32 TickCount;
+
+		IConsoleVariable* CVarPartialSleeping;
+	};
+
+	// Dummy test to make sure a stack of cubes will fall asleep.
+	TEST_P(SleepingTests, SleepChainOfSpheres)
+	{
+		FSleepingTest Test(GetParam());
+		Test.MakeChainOfSpheres(5);
+		Test.ParticleHandles[0]->SetObjectState(EObjectStateType::Static);
+
+		Test.AdvanceUntilSleeping();
+
+		for (int32 Id = 1; Id < Test.ParticleHandles.Num(); ++Id)
+		{
+			const FRigidBodyHandle_External* ParticleHandle = Test.ParticleHandles[Id];
+			EXPECT_TRUE(ParticleHandle->ObjectState() == EObjectStateType::Sleeping);
+		}
+	}
+
+	// Dummy test to make sure a stack of cubes will fall asleep.
+	TEST_P(SleepingTests, SleepStackOfCubes)
+	{
+		FSleepingTest Test(GetParam());
+		Test.MakeStackOfCubes(5);
+
+		Test.AdvanceUntilSleeping();
+
+		for (FRigidBodyHandle_External* ParticleHandle : Test.ParticleHandles)
+		{
+			EXPECT_TRUE(ParticleHandle->ObjectState() == EObjectStateType::Sleeping);
+		}
+	}
+
+	// Dummy test to make sure a stack of perfectly aligned spheres will fall asleep.
+	TEST_P(SleepingTests, SleepStackOfSpheres)
+	{
+		FSleepingTest Test(GetParam());
+		Test.MakeStackOfSpheres(5);
+
+		Test.AdvanceUntilSleeping();
+
+		for (FRigidBodyHandle_External* ParticleHandle : Test.ParticleHandles)
+		{
+			EXPECT_TRUE(ParticleHandle->ObjectState() == EObjectStateType::Sleeping);
+		}
+	}
+
+	TEST_P(SleepingTests, DropCubeOnStack)
+	{
+		const int32 NumOfStackedObjects = 5;
+		const FReal Dimension = 200;
+		const bool bPartialSleepEnabled = GetParam();
+
+		// Create a stack of cubes
+		FSleepingTest Test(bPartialSleepEnabled);
+		Test.MakeStackOfCubes(NumOfStackedObjects, /*Size=*/Dimension);
+
+		// Simulate until all particles are sleeping
+		Test.AdvanceUntilSleeping();
+
+		// Check that everything is sleeping
+		for (FRigidBodyHandle_External* ParticleHandle : Test.ParticleHandles)
+		{
+			EXPECT_TRUE(ParticleHandle->ObjectState() == EObjectStateType::Sleeping);
+		}
+
+		// Add another cube on top
+		// @todo (Chaos): This presently only triggers a wake-up event if the top cube is moving. 
+		// This is not great and should be updated
+		const FReal PosZ = 1100; // Right above the stack
+		FRigidBodyHandle_External* TopCube = Test.MakeCube(FVec3(0, 0, PosZ), Dimension);
+		TopCube->SetV(FVec3(0, 0, -10));
+
+		Test.Advance();
+
+		if (!bPartialSleepEnabled)
+		{
+			// The entire stack will wake up
+			for (FRigidBodyHandle_External* ParticleHandle : Test.ParticleHandles)
+			{
+				EXPECT_TRUE(ParticleHandle->ObjectState() == EObjectStateType::Dynamic);
+			}
+		}
+		else
+		{
+			// Only the top of the stack will wake up
+			EXPECT_TRUE(Test.ParticleHandles[0]->ObjectState() == EObjectStateType::Sleeping);
+			EXPECT_TRUE(Test.ParticleHandles[1]->ObjectState() == EObjectStateType::Sleeping);
+			EXPECT_TRUE(Test.ParticleHandles[2]->ObjectState() == EObjectStateType::Sleeping);
+			EXPECT_TRUE(Test.ParticleHandles[3]->ObjectState() == EObjectStateType::Sleeping);
+			EXPECT_TRUE(Test.ParticleHandles[4]->ObjectState() == EObjectStateType::Dynamic);
+			EXPECT_TRUE(Test.ParticleHandles[5]->ObjectState() == EObjectStateType::Dynamic);
+		}
+
+	}
+
+	TEST_P(SleepingTests, CollideWithStack)
+	{
+		const int32 NumOfStackedObjects = 5;
+		const FReal Dimension = 200;
+		const bool bPartialSleepEnabled = GetParam();
+
+		// Create a stack of cubes
+		FSleepingTest Test(bPartialSleepEnabled);
+		Test.MakeStackOfCubes(NumOfStackedObjects, /*Size=*/Dimension);
+
+		// Simulate until all particles are sleeping
+		Test.AdvanceUntilSleeping();
+
+		// Check that everything is sleeping
+		for (FRigidBodyHandle_External* ParticleHandle : Test.ParticleHandles)
+		{
+			EXPECT_TRUE(ParticleHandle->ObjectState() == EObjectStateType::Sleeping);
+		}
+
+		// Make the stack collide with a particle 
+		const FReal PosZ = 550; // slightly above the 3rd cube from below
+		FRigidBodyHandle_External* CollidingCube = Test.MakeCube(FVec3(200, 0, PosZ), Dimension);
+		CollidingCube->SetV(FVec3(-100, 0, 0), true);
+
+		Test.Advance();
+
+		if (!bPartialSleepEnabled)
+		{
+			// The entire stack will wake up
+			for (FRigidBodyHandle_External* ParticleHandle : Test.ParticleHandles)
+			{
+				EXPECT_TRUE(ParticleHandle->ObjectState() == EObjectStateType::Dynamic);
+			}
+		}
+		else
+		{
+			// Only the top of the stack will wake up
+			EXPECT_TRUE(Test.ParticleHandles[0]->ObjectState() == EObjectStateType::Sleeping);
+			EXPECT_TRUE(Test.ParticleHandles[1]->ObjectState() == EObjectStateType::Sleeping);
+			EXPECT_TRUE(Test.ParticleHandles[2]->ObjectState() == EObjectStateType::Dynamic);
+			EXPECT_TRUE(Test.ParticleHandles[3]->ObjectState() == EObjectStateType::Dynamic);
+			EXPECT_TRUE(Test.ParticleHandles[4]->ObjectState() == EObjectStateType::Dynamic);
+			EXPECT_TRUE(Test.ParticleHandles[5]->ObjectState() == EObjectStateType::Dynamic);
+		}
+	}
+
+	TEST_P(SleepingTests, PullCubeOutOfStack)
+	{
+		const int32 NumOfStackedObjects = 5;
+		const FReal Dimension = 200;
+		const bool bPartialSleepEnabled = GetParam();
+
+		// Create a stack of cubes
+		FSleepingTest Test(bPartialSleepEnabled);
+		Test.MakeStackOfCubes(NumOfStackedObjects, /*Size=*/Dimension);
+
+		// Simulate until all particles are sleeping
+		Test.AdvanceUntilSleeping();
+
+		// Check that everything is sleeping
+		for (FRigidBodyHandle_External* ParticleHandle : Test.ParticleHandles)
+		{
+			EXPECT_TRUE(ParticleHandle->ObjectState() == EObjectStateType::Sleeping);
+		}
+
+		// Pull one of the particles out of the stack by applying an impulse
+		Test.ParticleHandles[3]->SetLinearImpulse(FVec3(0, 100, 0), true);
+
+		Test.Advance();
+
+		if (!bPartialSleepEnabled)
+		{
+			// The entire stack will wake up
+			for (FRigidBodyHandle_External* ParticleHandle : Test.ParticleHandles)
+			{
+				EXPECT_TRUE(ParticleHandle->ObjectState() == EObjectStateType::Dynamic);
+			}
+		}
+		else
+		{			
+			// The entire stack will wake up
+			// @todo (Chaos): Not great, this should trigger a partial wake-up in the future.
+			for (FRigidBodyHandle_External* ParticleHandle : Test.ParticleHandles)
+			{
+				EXPECT_TRUE(ParticleHandle->ObjectState() == EObjectStateType::Dynamic);
+			}
+		}
+
+	}
+
+	TEST_P(SleepingTests, TeleportCubeOutOfStack)
+	{
+		const int32 NumOfStackedObjects = 5;
+		const FReal Dimension = 200;
+		const bool bPartialSleepEnabled = GetParam();
+
+		// Create a stack of cubes
+		FSleepingTest Test(bPartialSleepEnabled);
+		Test.MakeStackOfCubes(NumOfStackedObjects, /*Size=*/Dimension);
+
+		// Simulate until all particles are sleeping
+		Test.AdvanceUntilSleeping();
+
+		// Check that everything is sleeping
+		for (FRigidBodyHandle_External* ParticleHandle : Test.ParticleHandles)
+		{
+			EXPECT_TRUE(ParticleHandle->ObjectState() == EObjectStateType::Sleeping);
+		}
+
+		// Teleport one of the cubes out of the stack
+		Test.ParticleHandles[2]->SetX(FVec3(0, 500, 0), true);
+
+		Test.Advance();
+
+		if (!bPartialSleepEnabled)
+		{
+			// The entire stack will wake up
+			for (FRigidBodyHandle_External* ParticleHandle : Test.ParticleHandles)
+			{
+				EXPECT_TRUE(ParticleHandle->ObjectState() == EObjectStateType::Dynamic);
+			}
+		}
+		else
+		{
+			// @todo (Chaos): The entire stack wakes up. We want to change that.
+			EXPECT_TRUE(Test.ParticleHandles[0]->ObjectState() == EObjectStateType::Dynamic);
+			EXPECT_TRUE(Test.ParticleHandles[1]->ObjectState() == EObjectStateType::Dynamic);
+			EXPECT_TRUE(Test.ParticleHandles[2]->ObjectState() == EObjectStateType::Dynamic);
+			EXPECT_TRUE(Test.ParticleHandles[3]->ObjectState() == EObjectStateType::Dynamic);
+			EXPECT_TRUE(Test.ParticleHandles[4]->ObjectState() == EObjectStateType::Dynamic);
+		}
+
+	}
+
+	// Uses a traditional brick wall where every 2nd row is shifted (see SleepingTests::MakeBrickWall)
+	TEST_P(SleepingTests, CollideWithBrickWall)
+	{
+		const int32 NumOfStackedObjects = 5;
+		const FVec3 Dimensions = FVec3(100, 200, 50);
+		const bool bPartialSleepEnabled = GetParam();
+
+		// Create a stack of cubes
+		FSleepingTest Test(bPartialSleepEnabled);
+		Test.MakeBrickWall(/*NumY=*/ 5, /*NumZ=*/3, /*Dimensions=*/Dimensions);
+
+		// Simulate until all particles are sleeping
+		Test.AdvanceUntilSleeping(/*MaxIterations=*/500);
+
+		// Check that everything is sleeping
+		for (FRigidBodyHandle_External* ParticleHandle : Test.ParticleHandles)
+		{
+			EXPECT_TRUE(ParticleHandle->ObjectState() == EObjectStateType::Sleeping);
+		}
+
+		// Make the stack collide with a particle located: 
+		// - between the 2nd & 3rd row in height
+		// - between the 2nd & 3rd column in width
+		FRigidBodyHandle_External* CollidingCube = Test.MakeCube(FVec3(50, 300, 75), 50);
+		CollidingCube->SetV(FVec3(-200, 0, 0), true);
+
+		Test.Advance();
+
+		//// Debug print of all particle positions
+		//for (int32 Id = 0; Id < Test.ParticleHandles.Num(); ++Id)
+		//{
+		//	FRigidBodyHandle_External* ParticleHandle = Test.ParticleHandles[Id];
+		//	printf("Particle Id: %d, State: %d, Position: %f, %f, %f\n", Id, 
+		//		ParticleHandle->ObjectState(), ParticleHandle->GetX().X, ParticleHandle->GetX().Y, ParticleHandle->GetX().Z);
+		//}
+
+		if (!bPartialSleepEnabled)
+		{
+			// The entire brick wall will wake up
+			for (FRigidBodyHandle_External* ParticleHandle : Test.ParticleHandles)
+			{
+				EXPECT_TRUE(ParticleHandle->ObjectState() == EObjectStateType::Dynamic);
+			}
+		}
+		else
+		{
+			// Only the brick 6 colliding with the spawned cube and the bricks directly above wake up
+			// 1st row
+			EXPECT_TRUE(Test.ParticleHandles[0]->ObjectState() == EObjectStateType::Sleeping);
+			EXPECT_TRUE(Test.ParticleHandles[1]->ObjectState() == EObjectStateType::Sleeping);
+			EXPECT_TRUE(Test.ParticleHandles[2]->ObjectState() == EObjectStateType::Sleeping);
+			EXPECT_TRUE(Test.ParticleHandles[3]->ObjectState() == EObjectStateType::Sleeping);
+			EXPECT_TRUE(Test.ParticleHandles[4]->ObjectState() == EObjectStateType::Sleeping);
+			// 2nd row
+			EXPECT_TRUE(Test.ParticleHandles[5]->ObjectState() == EObjectStateType::Sleeping);
+			EXPECT_TRUE(Test.ParticleHandles[6]->ObjectState() == EObjectStateType::Dynamic);
+			EXPECT_TRUE(Test.ParticleHandles[7]->ObjectState() == EObjectStateType::Sleeping);
+			EXPECT_TRUE(Test.ParticleHandles[8]->ObjectState() == EObjectStateType::Sleeping);
+			EXPECT_TRUE(Test.ParticleHandles[9]->ObjectState() == EObjectStateType::Sleeping);
+			// 3rd row
+			EXPECT_TRUE(Test.ParticleHandles[10]->ObjectState() == EObjectStateType::Sleeping);
+			EXPECT_TRUE(Test.ParticleHandles[11]->ObjectState() == EObjectStateType::Dynamic);
+			EXPECT_TRUE(Test.ParticleHandles[12]->ObjectState() == EObjectStateType::Dynamic);
+			EXPECT_TRUE(Test.ParticleHandles[13]->ObjectState() == EObjectStateType::Sleeping);
+			EXPECT_TRUE(Test.ParticleHandles[14]->ObjectState() == EObjectStateType::Sleeping);
+		}
+	}
+
+	// Uses a block of bricks without shifting (see SleepingTests::MakeBrickWall)
+	TEST_P(SleepingTests, CollideWithBlockOfBricks)
+	{
+		const int32 NumOfStackedObjects = 5;
+		const FVec3 Dimensions = FVec3(100, 200, 50);
+		const bool bPartialSleepEnabled = GetParam();
+
+		// Create a stack of cubes
+		FSleepingTest Test(bPartialSleepEnabled);
+		Test.MakeBrickWall(/*NumY=*/ 5, /*NumZ=*/3, /*Dimensions=*/Dimensions, /*Mass=*/2, /*bShiftUnevenRows*/false);
+
+		// Simulate until all particles are sleeping
+		Test.AdvanceUntilSleeping(/*MaxIterations=*/500);
+
+		// Check that everything is sleeping
+		for (FRigidBodyHandle_External* ParticleHandle : Test.ParticleHandles)
+		{
+			EXPECT_TRUE(ParticleHandle->ObjectState() == EObjectStateType::Sleeping);
+		}
+
+		// Make the stack collide with a particle located: 
+		// - between the 2nd & 3rd row in height
+		// - between the 2nd & 3rd column in width
+		FRigidBodyHandle_External* CollidingCube = Test.MakeCube(FVec3(50, 300, 75), 50);
+		CollidingCube->SetV(FVec3(-200, 0, 0), true);
+
+		Test.Advance();
+
+		//// Debug print of all particle positions
+		//for (int32 Id = 0; Id < Test.ParticleHandles.Num(); ++Id)
+		//{
+		//	FRigidBodyHandle_External* ParticleHandle = Test.ParticleHandles[Id];
+		//	printf("Particle Id: %d, State: %d, Position: %f, %f, %f\n", Id, 
+		//		ParticleHandle->ObjectState(), ParticleHandle->GetX().X, ParticleHandle->GetX().Y, ParticleHandle->GetX().Z);
+		//}
+
+		if (!bPartialSleepEnabled)
+		{
+			// The entire brick wall will wake up
+			for (FRigidBodyHandle_External* ParticleHandle : Test.ParticleHandles)
+			{
+				EXPECT_TRUE(ParticleHandle->ObjectState() == EObjectStateType::Dynamic);
+			}
+		}
+		else
+		{
+			// Only the bricks 6 & 7 colliding with the spawned cube and the bricks directly above wake up
+			// 1st row
+			EXPECT_TRUE(Test.ParticleHandles[0]->ObjectState() == EObjectStateType::Sleeping);
+			EXPECT_TRUE(Test.ParticleHandles[1]->ObjectState() == EObjectStateType::Sleeping);
+			EXPECT_TRUE(Test.ParticleHandles[2]->ObjectState() == EObjectStateType::Sleeping);
+			EXPECT_TRUE(Test.ParticleHandles[3]->ObjectState() == EObjectStateType::Sleeping);
+			EXPECT_TRUE(Test.ParticleHandles[4]->ObjectState() == EObjectStateType::Sleeping);
+			// 2nd row
+			EXPECT_TRUE(Test.ParticleHandles[5]->ObjectState() == EObjectStateType::Sleeping);
+			EXPECT_TRUE(Test.ParticleHandles[6]->ObjectState() == EObjectStateType::Dynamic);
+			EXPECT_TRUE(Test.ParticleHandles[7]->ObjectState() == EObjectStateType::Dynamic);
+			EXPECT_TRUE(Test.ParticleHandles[8]->ObjectState() == EObjectStateType::Sleeping);
+			EXPECT_TRUE(Test.ParticleHandles[9]->ObjectState() == EObjectStateType::Sleeping);
+			// 3rd row
+			EXPECT_TRUE(Test.ParticleHandles[10]->ObjectState() == EObjectStateType::Dynamic);
+			EXPECT_TRUE(Test.ParticleHandles[11]->ObjectState() == EObjectStateType::Dynamic);
+			EXPECT_TRUE(Test.ParticleHandles[12]->ObjectState() == EObjectStateType::Dynamic);
+			EXPECT_TRUE(Test.ParticleHandles[13]->ObjectState() == EObjectStateType::Dynamic);
+			EXPECT_TRUE(Test.ParticleHandles[14]->ObjectState() == EObjectStateType::Sleeping);
+		}
+	}
+
+	INSTANTIATE_TEST_SUITE_P(, SleepingTests, testing::Bool());
+}
\ No newline at end of file
diff --git a/Engine/Source/Programs/HeadlessChaos/Public/HeadlessChaosTestUtility.h b/Engine/Source/Programs/HeadlessChaos/Public/HeadlessChaosTestUtility.h
--- a/Engine/Source/Programs/HeadlessChaos/Public/HeadlessChaosTestUtility.h
+++ b/Engine/Source/Programs/HeadlessChaos/Public/HeadlessChaosTestUtility.h
@@ -172,6 +172,27 @@ namespace ChaosTest {
 		Particle.SetInvI(TVec3<FRealSingle>(InvElement));
 	}
 
+	template <typename TParticle>
+	void SetBoxInertiaTensor(TParticle& Particle, const FVec3& Dimensions, float Mass)
+	{
+		float ElementX = Mass * (Dimensions.Y * Dimensions.Y + Dimensions.Z * Dimensions.Z) / 12.0f;
+		float ElementY = Mass * (Dimensions.X * Dimensions.X + Dimensions.Z * Dimensions.Z) / 12.0f;
+		float ElementZ = Mass * (Dimensions.Y * Dimensions.Y + Dimensions.X * Dimensions.X) / 12.0f;
+		float InvElementX = 1.f / ElementX;
+		float InvElementY = 1.f / ElementY;
+		float InvElementZ = 1.f / ElementX;
+		Particle.SetI(TVec3<FRealSingle>(ElementX, ElementY, ElementZ));
+		Particle.SetInvI(TVec3<FRealSingle>(InvElementX, InvElementY, InvElementZ));
+	}
+
+	template <typename TParticle>
+	void SetSphereInertiaTensor(TParticle& Particle, float Radius, float Mass)
+	{
+		float Element = 2.0f * Mass * Radius * Radius / 5.0f;
+		float InvElement = 1.f / Element;
+		Particle.SetI(TVec3<FRealSingle>(Element));
+		Particle.SetInvI(TVec3<FRealSingle>(InvElement));
+	}
 
 	extern FImplicitConvex3 CreateConvexBox(const FConvex::FVec3Type& BoxSize, const FReal Margin);
 	extern FImplicitConvex3 CreateConvexBox(const FConvex::FVec3Type& BoxMin, const FConvex::FVec3Type& BoxMax, const FReal Margin);
