Add a unit test for GJK
A failure case that was fixed on consoles

diff --git a/Engine/Source/Programs/HeadlessChaos/Private/HeadlessChaosTestGJK.cpp b/Engine/Source/Programs/HeadlessChaos/Private/HeadlessChaosTestGJK.cpp
--- a/Engine/Source/Programs/HeadlessChaos/Private/HeadlessChaosTestGJK.cpp
+++ b/Engine/Source/Programs/HeadlessChaos/Private/HeadlessChaosTestGJK.cpp
@@ -2467,4 +2467,36 @@ namespace ChaosTest
 		}
 	}
 
+	GTEST_TEST(GJKTests, GJK_SweepSphereConvexMesh)
+	{
+		// Test case bug that was not working on XSX, and maybe on Switch
+		const TArray<Chaos::FConvex::FVec3Type> ConvexVertices
+		{
+			{ -1320.50159, 461.919220, 18.5806618 },
+			{ -1321.64062, 464.953125, -81.3668213 },
+			{ -723.913879, 585.305603, -65.5020828 },
+			{ -1026.47400, -493.330170, -13.7665596 },
+			{ -1027.61304, -490.296295, -113.714035 },
+			{ -723.348694, 578.755981, 34.2816010 },
+			{ -668.640442, -417.582031, -31.4263611 },
+			{ -669.205688, -411.032410, -131.210037 }
+		};
+
+		const Chaos::FConvexPtr ConvexPtr(new Chaos::FConvex(ConvexVertices, 0.0f));
+		const Chaos::TImplicitObjectInstanced<Chaos::FConvex> Convex(ConvexPtr);
+		const Chaos::FSphere Sphere(Chaos::FVec3::ZeroVector, 40.0);
+
+		const UE::Math::TQuat BToARotation{ -0.0135237072, -0.149320737, 0.0850137547, 0.985034525 };
+		const Chaos::FVec3 StartPoint{ -2305.07544, -225.373749, 362.317993 };
+		const Chaos::FVec3 RayDir{ 0.958641946, 0.168052554, -0.229703903 };
+		const Chaos::FReal RayLength{ 1214.51257 };
+		const bool bComputeMTD{ true };
+		const Chaos::FVec3 InitialDir{ 270.538879, -2088.81055, -1029.13525 };
+
+		Chaos::FReal TimeOffImpact{ };
+		Chaos::FVec3 Position{ };
+		Chaos::FVec3 Normal{ };
+		bool bHit = Chaos::GJKRaycast2<Chaos::FReal>(Convex, Sphere, Chaos::FRigidTransform3(UE::Math::TTransform<double>(BToARotation, StartPoint)), RayDir, RayLength, TimeOffImpact, Position, Normal, 0, true, InitialDir, 0.0);
+		EXPECT_FALSE(bHit);
+	}
 }
\ No newline at end of file
