Add a GJK unit test for Convex sweep on Convex
This is a failling test on PS5

diff --git a/Engine/Source/Programs/HeadlessChaos/Private/HeadlessChaosTestGJK.cpp b/Engine/Source/Programs/HeadlessChaos/Private/HeadlessChaosTestGJK.cpp
--- a/Engine/Source/Programs/HeadlessChaos/Private/HeadlessChaosTestGJK.cpp
+++ b/Engine/Source/Programs/HeadlessChaos/Private/HeadlessChaosTestGJK.cpp
@@ -2499,4 +2499,51 @@ namespace ChaosTest
 		bool bHit = Chaos::GJKRaycast2<Chaos::FReal>(Convex, Sphere, Chaos::FRigidTransform3(UE::Math::TTransform<double>(BToARotation, StartPoint)), RayDir, RayLength, TimeOffImpact, Position, Normal, 0, true, InitialDir, 0.0);
 		EXPECT_FALSE(bHit);
 	}
+
+	GTEST_TEST(GJKTests, GJK_SweepConvexMeshConvexMesh)
+	{
+		// Test case bug on PS5
+		const TArray<Chaos::FConvex::FVec3Type> ConvexVerticesA
+		{
+			{ 372.876678, 1253.23181, 1937.96875 },
+			{ -516.649231, 1253.23181, 1937.96875 },
+			{ -516.649231, 1253.23181, 1873.81213 },
+			{ -516.649231, 1253.23181, 1873.81213 },
+			{ 372.876678, 1253.23181, 1873.81213 },
+			{ 372.876678, -2157.82104, 2122.00000 },
+			{ 372.876678, -2157.82104, 2057.84326 },
+			{ -516.649231, -2157.82104, 2057.84326 },
+			{ -516.649231, -2157.82104, 2122.00000 }
+		};
+		const Chaos::FConvexPtr ConvexPtrA(new Chaos::FConvex(ConvexVerticesA, 0.0f));
+		const Chaos::TImplicitObjectInstanced<Chaos::FConvex> ConvexA(ConvexPtrA);
+
+		const TArray<Chaos::FConvex::FVec3Type> ConvexVerticesB
+		{
+			{ 73.4703674, -61.1765442, -0.19269731 },
+			{ -58.2119789, -43.0196724, 150.211761 },
+			{ 56.7236938, -43.0196762, 150.211761 },
+			{ 73.4703674, 87.8292542, -0.19269731 },
+			{ -73.4703674, 87.8292542, -0.19269731 },
+			{ -73.4703751, -61.1765366, -0.19269731 },
+			{ -45.8896980, 52.8858337, 150.211761 },
+			{ 44.4014816, 52.8858109, 150.211761 }
+		};
+		const Chaos::FConvexPtr ConvexPtrB(new Chaos::FConvex(ConvexVerticesB, 0.0f));
+		const Chaos::TImplicitObjectInstanced<Chaos::FConvex> ConvexB(ConvexPtrB);
+
+		const UE::Math::TQuat BToARotation{ -5.54229738e-017, -3.12762759e-018, 0.02172332, 0.99976402 };
+		const Chaos::FVec3 StartPoint{ -57.6924553, -1033.34961, 2064.99536 };
+		const Chaos::FVec3 RayDir{ -0.00042043, 0.99854773, -0.05387306 };
+		const Chaos::FReal RayLength{ 5.3171567916870117 };
+		const Chaos::FVec3 InitialDir{ 173.583862, 1020.29828, -2064.99536 };
+		const bool bComputeMTD{ true };
+
+		Chaos::FReal Time{ };
+		Chaos::FVec3 Position{ };
+		Chaos::FVec3 Normal{ };
+		bool bHit{ Chaos::GJKRaycast2<Chaos::FReal>(ConvexA, ConvexB, Chaos::FRigidTransform3(UE::Math::TTransform<double>(BToARotation, StartPoint)), RayDir, RayLength, Time, Position, Normal, 0, true, InitialDir, 0.0) };
+		EXPECT_TRUE(bHit);
+		EXPECT_NEAR(-0.000022, Time, 1e-5);
+	}
 }
\ No newline at end of file
