Sleep counter bug fix:
- Update the sleep counter for active rigid particle during partial sleep processing (i.e. for moving and non-moving particles)
- Reset sleep counter to 0 for all particles processed in a partial or full wake event (prevents sleeping in the same frame)

diff --git a/Engine/Source/Programs/HeadlessChaos/Private/HeadlessChaosTestGraphEvolution.cpp b/Engine/Source/Programs/HeadlessChaos/Private/HeadlessChaosTestGraphEvolution.cpp
--- a/Engine/Source/Programs/HeadlessChaos/Private/HeadlessChaosTestGraphEvolution.cpp
+++ b/Engine/Source/Programs/HeadlessChaos/Private/HeadlessChaosTestGraphEvolution.cpp
@@ -2200,5 +2200,111 @@ namespace ChaosTest
 		}
 	}
 
+	TEST_P(GraphEvolutionTests, TestConstraintGraph_SleepCounterReset)
+	{
+		// Partial island sleeping only
+		if (GetParam())
+		{	
+			FGraphEvolutionTest Test(2, GetParam());
+			Test.MakeChain();
+
+			// Number of steps is a few less than the sleep counter threshold
+			IConsoleVariable* CVarPartialSleeping = IConsoleManager::Get().FindConsoleVariable(TEXT("p.Chaos.Solver.Sleep.Defaults.SleepCounterThreshold"), false);
+			check(CVarPartialSleeping && CVarPartialSleeping->IsVariableInt());			
+			const int32 SleepCounterThreshold = CVarPartialSleeping->GetInt();
+			const int32 Steps = SleepCounterThreshold - 2;
+
+			for (int32 I = 0; I < Steps; ++I)
+			{
+				Test.Advance();
+			}
+
+			for (FPBDRigidParticleHandle* Particle : Test.ParticleHandles)
+			{
+				EXPECT_FALSE(Particle->IsSleeping());
+				const int32 SleepCounter = Particle->SleepCounter();
+				EXPECT_EQ(SleepCounter, Steps);
+			}
+
+			// Set a non-zero velocity to make the particle move.
+			FParticleDynamics Dynamics;
+			Dynamics.SetAcceleration(FVec3(0, 0, 0));
+			Dynamics.SetAngularAcceleration(FVec3(0, 0, 0));
+			Dynamics.SetLinearImpulseVelocity(FVec3(50, 50, 0));
+			Dynamics.SetAngularImpulseVelocity(FVec3(0, 0, 0));
+			for (FPBDRigidParticleHandle* Particle : Test.ParticleHandles)
+			{
+				Particle->SetDynamics(Dynamics);
+				Test.Evolution.ResetVSmoothFromForces(*Particle);
+			}
+
+			// Step one more time, the sleep counter should be reset to 0 because the particles are moving
+		    Test.Advance();
+
+			for (FPBDRigidParticleHandle* Particle : Test.ParticleHandles)
+			{
+				EXPECT_FALSE(Particle->IsSleeping());
+				const int32 SleepCounter = Particle->SleepCounter();
+				EXPECT_EQ(SleepCounter, 0);
+			}
+		}
+	}
+
+	TEST_P(GraphEvolutionTests, TestConstraintGraph_PreventSleepDuringWakeEvent)
+	{
+		// Partial island sleeping only
+		if (GetParam())
+		{
+			FGraphEvolutionTest Test(2, GetParam());
+			Test.MakeChain();
+
+			// Number of steps is the sleep counter threshold
+			IConsoleVariable* CVarPartialSleeping = IConsoleManager::Get().FindConsoleVariable(TEXT("p.Chaos.Solver.Sleep.Defaults.SleepCounterThreshold"), false);
+			check(CVarPartialSleeping && CVarPartialSleeping->IsVariableInt());
+			const int32 SleepCounterThreshold = CVarPartialSleeping->GetInt();
+			const int32 Steps = SleepCounterThreshold;
+
+			for (int32 I = 0; I < Steps; ++I)
+			{
+				Test.Advance();
+			}
+
+			for (FPBDRigidParticleHandle* Particle : Test.ParticleHandles)
+			{
+				EXPECT_FALSE(Particle->IsSleeping());
+				const int32 SleepCounter = Particle->SleepCounter();
+				EXPECT_EQ(SleepCounter, Steps);
+			}
+
+			// Add another particle and connect it to the top of the chain
+			Test.ParticleHandles.Add(Test.Evolution.CreateDynamicParticles(1)[0]);
+			Test.Evolution.EnableParticle(Test.ParticleHandles.Last());
+			Test.ConstraintHandles.Add(Test.Constraints.AddConstraint({ Test.ParticleHandles.Last(1), Test.ParticleHandles.Last() }));
+
+			// NOTE: Set a non-zero velocity to make the particle move and trigger a wake-up event.
+			FParticleDynamics Dynamics;
+			Dynamics.SetAcceleration(FVec3(0, 0, 0));
+			Dynamics.SetAngularAcceleration(FVec3(0, 0, 0));
+			Dynamics.SetLinearImpulseVelocity(FVec3(50, 50, 0));
+			Dynamics.SetAngularImpulseVelocity(FVec3(0, 0, 0));
+			Test.ParticleHandles.Last()->SetDynamics(Dynamics);
+			Test.Evolution.ResetVSmoothFromForces(*Test.ParticleHandles.Last());
+
+			// Step one more time, the particles should not sleep since a wake-up event just happened.
+			Test.Advance();
+
+			for (FPBDRigidParticleHandle* Particle : Test.ParticleHandles)
+			{
+				EXPECT_FALSE(Particle->IsSleeping());
+			}
+
+			// Resting particles
+			EXPECT_EQ(Test.ParticleHandles[0]->SleepCounter(), 1);
+			EXPECT_EQ(Test.ParticleHandles[1]->SleepCounter(), 1);
+			// Moving particle
+			EXPECT_EQ(Test.ParticleHandles[2]->SleepCounter(), 0);
+		}
+	}
+
 	INSTANTIATE_TEST_SUITE_P(, GraphEvolutionTests,	testing::Bool());
 }
diff --git a/Engine/Source/Runtime/Experimental/Chaos/Private/Chaos/Island/IslandManager.cpp b/Engine/Source/Runtime/Experimental/Chaos/Private/Chaos/Island/IslandManager.cpp
--- a/Engine/Source/Runtime/Experimental/Chaos/Private/Chaos/Island/IslandManager.cpp
+++ b/Engine/Source/Runtime/Experimental/Chaos/Private/Chaos/Island/IslandManager.cpp
@@ -2541,25 +2541,26 @@ namespace Chaos::Private
 					|| (Rigid->WSmooth().SizeSquared() > Node->SleepAngularThresholdSq))
 				{
 					bAllNodesSleeping = false;
-					continue;
 				}
-
-				// If we get here we might want to sleep
-				// Update the sleep counter and sleep if we exceed it
-				static_assert(sizeof(decltype(Rigid->SleepCounter())) == 1, "Expected int8 for SleepCounter(). Update clamp below");
-				SleepCounter = FMath::Min(int32(Rigid->SleepCounter()) + 1, int32(TNumericLimits<int8>::Max()));
-				if (SleepCounter > SleepCounterThreshold)
-				{
-					SleptParticles.Add(Rigid->Handle());
-					Node->Flags.bIsSleeping = true;
-					SleepCounter = 0;
-				}
-				else
+				else // motionless
 				{
-					bAllNodesSleeping = false;
+					// If we get here we might want to sleep
+					// Update the sleep counter and sleep if we exceed it
+					static_assert(sizeof(decltype(Rigid->SleepCounter())) == 1, "Expected int8 for SleepCounter(). Update clamp below");
+					SleepCounter = FMath::Min(int32(Rigid->SleepCounter()) + 1, int32(TNumericLimits<int8>::Max()));
+					if (SleepCounter > SleepCounterThreshold)
+					{
+						SleptParticles.Add(Rigid->Handle());
+						Node->Flags.bIsSleeping = true;
+						SleepCounter = 0;
+					}
+					else
+					{
+						bAllNodesSleeping = false;
+					}
 				}
-				Rigid->SetSleepCounter(int8(SleepCounter));
 			}
+			Rigid->SetSleepCounter(int8(SleepCounter));
 		}
 		// Flag the island if all particles are sleeping
 		if (bAllNodesSleeping)
@@ -2662,6 +2663,9 @@ namespace Chaos::Private
 		if (!CVars::bChaosSolverPartialIslandWake && NumNodesConstraintWake > 0)
 		{
 			PropagateIslandSleep(Island);
+
+			// Particle sleep counter reset is needed to prevent immediate re-sleeping in partial island sleeping.
+			ResetParticleSleepCounter(Island);
 		}
 		else // Otherwise, propagate the partial wake event through the constraint graph.
 		{
@@ -2685,6 +2689,9 @@ namespace Chaos::Private
 		if (!CVars::bChaosSolverPartialIslandWake && NumNodesNodeWake > 0)
 		{
 			PropagateIslandSleep(Island);
+
+			// Particle sleep counter reset is needed to prevent immediate re-sleeping in partial island sleeping.
+			ResetParticleSleepCounter(Island);
 		}
 		else // Otherwise, propagate the partial wake event through the constraint graph.
 		{
@@ -2846,6 +2853,7 @@ namespace Chaos::Private
 
 		// Otherwise, the particle is added to the node queue and woken up if needed.
 		NodeQueue.Add(Node);
+		Rigid->SetSleepCounter(0); // Force reset in partial wake events to prevent immediate re-sleep
 		if (!bIsAwake)
 		{
 			Node->Flags.bIsSleeping = false;
@@ -2943,6 +2951,18 @@ namespace Chaos::Private
 		}
 	}
 
+	void FPBDIslandManager::ResetParticleSleepCounter(FPBDIsland* Island)
+	{
+		for (FPBDIslandParticle* Node : Island->Nodes)
+		{
+			FPBDRigidParticleHandle* Rigid = Node->GetParticle()->CastToRigidParticle();
+			if (Rigid != nullptr)
+			{
+				Rigid->SetSleepCounter(0);
+			}
+		}
+	}
+
 	void FPBDIslandManager::ResetEdgeFlags(FPBDIsland* Island)
 	{
 		for (TArray<FPBDIslandConstraint*>& IslandEdges : Island->ContainerEdges)
diff --git a/Engine/Source/Runtime/Experimental/Chaos/Public/Chaos/Island/IslandManager.h b/Engine/Source/Runtime/Experimental/Chaos/Public/Chaos/Island/IslandManager.h
--- a/Engine/Source/Runtime/Experimental/Chaos/Public/Chaos/Island/IslandManager.h
+++ b/Engine/Source/Runtime/Experimental/Chaos/Public/Chaos/Island/IslandManager.h
@@ -610,6 +610,7 @@ namespace Chaos::Private
 		CHAOS_API bool ProcessParticleWake(FPBDIslandParticle* Node, const FPBDIslandParticle* OtherNode, TArray<FPBDIslandParticle*>& NodeQueue, TArray<FGeometryParticleHandle*>& WakeUpQueue, const bool bAlwaysPropagate);
 		CHAOS_API void SyncConstraintSleepStateToParticles(FPBDIsland* Island);
 		CHAOS_API void WakeUpParticles(TArray<FGeometryParticleHandle*>& ParticlesToWake, const bool bResetSleepCounter = true);
+		CHAOS_API void ResetParticleSleepCounter(FPBDIsland* Island);
 		CHAOS_API void ResetEdgeFlags(FPBDIsland* Island);
 		CHAOS_API void ResetNodeFlags(FPBDIsland* Island);
 
