Chaos : fix compilation issues resulting from CL 42691607
- Use the non-deprecated Geometry collection physics proxy constructor

diff --git a/Engine/Source/Programs/HeadlessChaos/Private/GeometryCollection/GeometryCollectionTestFramework.cpp b/Engine/Source/Programs/HeadlessChaos/Private/GeometryCollection/GeometryCollectionTestFramework.cpp
--- a/Engine/Source/Programs/HeadlessChaos/Private/GeometryCollection/GeometryCollectionTestFramework.cpp
+++ b/Engine/Source/Programs/HeadlessChaos/Private/GeometryCollection/GeometryCollectionTestFramework.cpp
@@ -203,7 +203,7 @@ namespace GeometryCollectionTest
 		FGeometryCollectionPhysicsProxy* PhysObject =
 			new FGeometryCollectionPhysicsProxy(
 				nullptr,			// UObject owner
-				*DynamicCollection, // Game thread collection
+				DynamicCollection.ToSharedRef(), // Game thread collection
 				SimulationParams,
 				SimFilterData,
 				QueryFilterData);
