Chaos
- fix race condition in Unregister when collision events are pending
- add unit test

diff --git a/Engine/Source/Programs/HeadlessChaos/Private/GeometryCollection/GeometryCollectionTestCollisionEvents.cpp b/Engine/Source/Programs/HeadlessChaos/Private/GeometryCollection/GeometryCollectionTestCollisionEvents.cpp
new file mode 100644
--- /dev/null
+++ b/Engine/Source/Programs/HeadlessChaos/Private/GeometryCollection/GeometryCollectionTestCollisionEvents.cpp
@@ -0,0 +1,197 @@
+// Copyright Epic Games, Inc. All Rights Reserved.
+
+#include "EventManager.h"
+#include "EventsData.h"
+#include "GeometryCollection/GeometryCollectionTestCollisionResolution.h"
+#include "GeometryCollection/GeometryCollectionTestFramework.h"
+#include "GeometryCollection/GeometryCollectionAlgo.h"
+#include "GeometryCollectionProxyData.h"
+#include "HeadlessChaosTestUtility.h"
+
+namespace GeometryCollectionTest
+{
+	using namespace ChaosTest;
+
+	namespace CollisionEventsTests
+	{
+		void AddSpheres(FFramework& UnitTest)
+		{
+			CreationParameters Params;
+			Params.CollisionType = ECollisionTypeEnum::Chaos_Volumetric;
+
+			Params.EnableClustering = false;
+
+			FVector Scale(50.f);
+			Params.GeomTransform.SetScale3D(Scale); // Sphere radius
+			Params.EnableClustering = false;
+
+			// Make a dynamic simplicial sphere
+			Params.SimplicialType = ESimplicialType::Chaos_Simplicial_Sphere;
+			Params.ImplicitType = EImplicitTypeEnum::Chaos_Implicit_Sphere;
+
+			Params.DynamicState = EObjectStateTypeEnum::Chaos_Object_Dynamic;
+			Params.RootTransform = FTransform(FQuat::MakeFromEuler(FVector(0)), FVector(0, 0, 75));
+			FGeometryCollectionWrapper* DynamicSphereCollection =
+				TNewSimulationObject<GeometryType::GeometryCollectionWithSingleRigid>::Init(Params)->template As<FGeometryCollectionWrapper>();
+			UnitTest.AddSimulationObject(DynamicSphereCollection);
+
+			// Make a kinematic implicit sphere
+			Params.SimplicialType = ESimplicialType::Chaos_Simplicial_Sphere;
+			Params.ImplicitType = EImplicitTypeEnum::Chaos_Implicit_Sphere;
+
+			Params.DynamicState = EObjectStateTypeEnum::Chaos_Object_Kinematic;
+			Params.RootTransform = FTransform(FQuat::MakeFromEuler(FVector(0)), FVector(0));
+			FGeometryCollectionWrapper* KinematicSphereCollection =
+				TNewSimulationObject<GeometryType::GeometryCollectionWithSingleRigid>::Init(Params)->template As<FGeometryCollectionWrapper>();
+			UnitTest.AddSimulationObject(KinematicSphereCollection);
+
+			// Hard code masstolocal on rest collection to identity
+			{
+				TManagedArray<FTransform>& MassToLocal =
+					DynamicSphereCollection->RestCollection->template ModifyAttribute<FTransform>(
+						TEXT("MassToLocal"), FTransformCollection::TransformGroup);
+				check(MassToLocal.Num() == 1);
+				MassToLocal[0] = FTransform::Identity;
+			}
+			{
+				TManagedArray<FTransform>& MassToLocal =
+					KinematicSphereCollection->RestCollection->template ModifyAttribute<FTransform>(
+						TEXT("MassToLocal"), FTransformCollection::TransformGroup);
+				check(MassToLocal.Num() == 1);
+				MassToLocal[0] = FTransform::Identity;
+			}
+		}
+	}
+
+	// Check that we get collision events from two overlapping GC spheres
+	GTEST_TEST(AllTraits, GeometryCollection_CollisionEventsTest)
+	{
+		// simplicial sphere to implicit sphere
+		FFramework UnitTest;
+
+		CollisionEventsTests::AddSpheres(UnitTest);
+
+		UnitTest.Initialize();
+
+		// Listen to collision events
+		struct FTestCollisionHandler
+		{
+			void OnCollision(const Chaos::FCollisionEventData& CollisionEventData)
+			{
+				EXPECT_EQ(CollisionEventData.CollisionData.AllCollisionsArray.Num(), 1);
+			};
+		};
+		FTestCollisionHandler CollisionHandler;
+		Chaos::FEventManager* EventManager = UnitTest.Solver->GetEventManager();
+		EventManager->RegisterHandler<Chaos::FCollisionEventData>(Chaos::EEventType::Collision, &CollisionHandler, &FTestCollisionHandler::OnCollision);
+
+		FSolverCollisionFilterSettings CollisionFilterSettings;
+		CollisionFilterSettings.FilterEnabled = true;
+		UnitTest.Solver->SetCollisionFilterSettings(CollisionFilterSettings);
+
+		// Should get one collision event between the two spheres
+		UnitTest.Advance();
+	}
+
+
+	// Test collision events in async physics mode
+	GTEST_TEST(AllTraits, GeometryCollection_CollisionEventsAsyncTest)
+	{
+		// simplicial sphere to implicit sphere
+		FrameworkParameters UnitTestParams;
+		UnitTestParams.ThreadingMode = Chaos::EThreadingMode::DedicatedThread;
+		FFramework UnitTest;
+
+		CollisionEventsTests::AddSpheres(UnitTest);
+
+		UnitTest.Initialize();
+
+		// Listen to collision events
+		struct FTestCollisionHandler
+		{
+			void OnCollision(const Chaos::FCollisionEventData& CollisionEventData)
+			{
+				EXPECT_EQ(CollisionEventData.CollisionData.AllCollisionsArray.Num(), 1);
+			};
+		};
+		FTestCollisionHandler CollisionHandler;
+		Chaos::FEventManager* EventManager = UnitTest.Solver->GetEventManager();
+		EventManager->RegisterHandler<Chaos::FCollisionEventData>(Chaos::EEventType::Collision, &CollisionHandler, &FTestCollisionHandler::OnCollision);
+
+		FSolverCollisionFilterSettings CollisionFilterSettings;
+		CollisionFilterSettings.FilterEnabled = true;
+		UnitTest.Solver->SetCollisionFilterSettings(CollisionFilterSettings);
+
+		// Should get one collision event between the two spheres
+		UnitTest.Advance();
+	}
+
+	// Destroy a particle while async physics is running. We will still get a collision event but the proxy will
+	// have been marked for destruction. The user must check GetMarkedDeleted() before using the proxy, especially
+	// if the owing game-side object might have been destroyed (which will happen if GC occurs at the right time).
+	GTEST_TEST(AllTraits, GeometryCollection_CollisionEventsAsyncDestroyTest)
+	{
+		// simplicial sphere to implicit sphere
+		FrameworkParameters UnitTestParams;
+		UnitTestParams.ThreadingMode = Chaos::EThreadingMode::DedicatedThread;
+		UnitTestParams.AsyncDt = 1.0 / 60.0;
+		FFramework UnitTest;
+
+		CollisionEventsTests::AddSpheres(UnitTest);
+
+		UnitTest.Initialize();
+
+		WrapperBase* DestroyedWrapper = nullptr;
+
+		// Listen to collision events
+		struct FTestCollisionHandler
+		{
+			// This callback will be called twice.
+			// - the first time between the two active particles
+			// - the second time between the same two particles, but after the 
+			//   object has been "destroyed"
+			void OnCollision(const Chaos::FCollisionEventData& CollisionEventData)
+			{
+				// We should get one collision between the two spheres
+				EXPECT_EQ(CollisionEventData.CollisionData.AllCollisionsArray.Num(), 1);
+				if (CollisionEventData.CollisionData.AllCollisionsArray.Num() == 1)
+				{
+					// After the dynamic particle is destroyed, we get a collision with 
+					// the dynamic proxy's MarkeDeleted flag set. The user can check this
+					// to know that the game-side object may have been destroyed
+					const FCollidingData& Collision = CollisionEventData.CollisionData.AllCollisionsArray[0];
+					EXPECT_EQ(Collision.Proxy1->GetMarkedDeleted(), bIsDeleted);
+				}
+			};
+
+			bool bIsDeleted = false;
+		};
+		FTestCollisionHandler CollisionHandler;
+		Chaos::FEventManager* EventManager = UnitTest.Solver->GetEventManager();
+		EventManager->RegisterHandler<Chaos::FCollisionEventData>(Chaos::EEventType::Collision, &CollisionHandler, &FTestCollisionHandler::OnCollision);
+
+		FSolverCollisionFilterSettings CollisionFilterSettings;
+		CollisionFilterSettings.FilterEnabled = true;
+		UnitTest.Solver->SetCollisionFilterSettings(CollisionFilterSettings);
+
+		// Run one tick - should get one collision event between the two spheres
+		CollisionHandler.bIsDeleted = false;
+		UnitTest.Advance();
+
+		// Runa second tick, but remove the dynamic sphere mid-tick,
+		// after collision resolution, but before event dispatch
+		UnitTest.Solver->AdvanceAndDispatch_External(UnitTestParams.Dt);
+		UnitTest.Solver->UpdateGameThreadStructures();
+
+		// Unregister the GC and destroy the particle owner (WrapperBase).
+		// The proxy will be destroyed on the physics thread in a subsequent tick,
+		// but the user data will still reference the deleted owner.
+		DestroyedWrapper = UnitTest.PhysicsObjects[0];
+		UnitTest.RemoveSimulationObject(DestroyedWrapper);
+		
+		// Dispatch events - should get one collision event between the two
+		// sphere, but the dynamic sphere is marked as deleted
+		CollisionHandler.bIsDeleted = true;
+		UnitTest.Solver->SyncEvents_GameThread();
+	}
+}
\ No newline at end of file
diff --git a/Engine/Source/Programs/HeadlessChaos/Private/GeometryCollection/GeometryCollectionTestFramework.cpp b/Engine/Source/Programs/HeadlessChaos/Private/GeometryCollection/GeometryCollectionTestFramework.cpp
--- a/Engine/Source/Programs/HeadlessChaos/Private/GeometryCollection/GeometryCollectionTestFramework.cpp
+++ b/Engine/Source/Programs/HeadlessChaos/Private/GeometryCollection/GeometryCollectionTestFramework.cpp
@@ -134,6 +134,12 @@ namespace GeometryCollectionTest
 			FTransform::Identity, GeomTransform, NumberOfMaterials);
 	}
 
+	FGeometryCollectionWrapper::FGeometryCollectionWrapper() 
+		: WrapperBase(WrapperType::GeometryCollection)
+	{
+		FChaosUserData::Set<FCustomChaosPayload>(&UserData, reinterpret_cast<FCustomChaosPayload*>(this));
+	}
+
 	WrapperBase* CommonInit(const TSharedPtr<FGeometryCollection>& RestCollection, const CreationParameters& Params)
 	{
 		FTransformCollection SingleTransform = FTransformCollection::SingleTransform();
@@ -150,6 +156,8 @@ namespace GeometryCollectionTest
 
 		TSharedPtr<FGeometryDynamicCollection> DynamicCollection = GeometryCollectionToGeometryDynamicCollection(RestCollection, Params.DynamicState);
 
+		FGeometryCollectionWrapper* Wrapper = new FGeometryCollectionWrapper();
+
 		FSimulationParameters SimulationParams;
 		{
 			SimulationParams.RestCollectionShared = RestCollection;
@@ -181,6 +189,7 @@ namespace GeometryCollectionTest
 			SimulationParams.PositionSolverIterations = 8;
 			SimulationParams.VelocitySolverIterations = 1;
 			SimulationParams.ProjectionSolverIterations = 1;
+			SimulationParams.UserData = reinterpret_cast<void*>(&Wrapper->UserData);
 
 			Chaos::FErrorReporter ErrorReporter;
 			BuildSimulationData(ErrorReporter, *RestCollection.Get(), SimulationParams.Shared);
@@ -199,7 +208,7 @@ namespace GeometryCollectionTest
 
 		// Enable all collisions
 		SimFilterData.Word1 = 0xFFFF; // this body channel
-		SimFilterData.Word3 = 0xFFFF; // collision candidate channels
+		SimFilterData.Word3 = 0xFFFF; // collision candidate channels (also enables EPDF_ContactNotify)
 
 		FGeometryCollectionPhysicsProxy* PhysObject =
 			new FGeometryCollectionPhysicsProxy(
@@ -208,7 +217,11 @@ namespace GeometryCollectionTest
 				SimulationParams,
 				SimFilterData,
 				QueryFilterData);
-		return new FGeometryCollectionWrapper(RestCollection, DynamicCollection, PhysObject);
+
+		Wrapper->RestCollection = RestCollection;
+		Wrapper->DynamicCollection = DynamicCollection;
+		Wrapper->PhysObject = PhysObject;
+		return Wrapper;
 	}
 
 	template <>
@@ -280,7 +293,7 @@ namespace GeometryCollectionTest
 	, Module(FChaosSolversModule::GetModule())
 	, Solver(nullptr)
 	{
-		Solver = Module->CreateSolver(nullptr, /*AsyncDt=*/-1, Parameters.ThreadingMode);	//until refactor is done, solver must be created after thread change
+		Solver = Module->CreateSolver(nullptr, Parameters.AsyncDt, Parameters.ThreadingMode);	//until refactor is done, solver must be created after thread change
 		ChaosTest::InitSolverSettings(Solver);
 	}
 
@@ -313,6 +326,23 @@ namespace GeometryCollectionTest
 		PhysicsObjects.Add(Object);
 	}
 
+	void FFramework::RemoveSimulationObject(WrapperBase* Object)
+	{
+		if (PhysicsObjects.Remove(Object) > 0)
+		{
+			if (FGeometryCollectionWrapper* GCW = Object->As<FGeometryCollectionWrapper>())
+			{
+				Solver->UnregisterObject(GCW->PhysObject);
+			}
+			else if (RigidBodyWrapper* BCW = Object->As<RigidBodyWrapper>())
+			{
+				Solver->UnregisterObject(BCW->Particle);
+			}
+
+			delete Object;
+		}
+	}
+
 	void FFramework::Initialize()
 	{
 		for (WrapperBase* Object : PhysicsObjects)
@@ -331,9 +361,9 @@ namespace GeometryCollectionTest
 
 	void FFramework::Advance()
 	{
-		Solver->SyncEvents_GameThread();
 		Solver->AdvanceAndDispatch_External(Dt);
 		Solver->UpdateGameThreadStructures();
+		Solver->SyncEvents_GameThread();
 	}
 
 } // end namespace GeometryCollectionTest
diff --git a/Engine/Source/Programs/HeadlessChaos/Public/GeometryCollection/GeometryCollectionTestFramework.h b/Engine/Source/Programs/HeadlessChaos/Public/GeometryCollection/GeometryCollectionTestFramework.h
--- a/Engine/Source/Programs/HeadlessChaos/Public/GeometryCollection/GeometryCollectionTestFramework.h
+++ b/Engine/Source/Programs/HeadlessChaos/Public/GeometryCollection/GeometryCollectionTestFramework.h
@@ -34,25 +34,27 @@ namespace GeometryCollectionTest
 	{
 		WrapperType Type;
 		WrapperBase(WrapperType TypeIn) : Type(TypeIn) {}
+		virtual ~WrapperBase() {}
 		template<class AS_T> AS_T* As() { return AS_T::StaticType() == Type ? static_cast<AS_T*>(this) : nullptr; }
 		template<class AS_T> const AS_T* As() const { return AS_T::StaticType() == Type ? static_cast<const AS_T*>(this) : nullptr; }
 	};
 
 	struct FGeometryCollectionWrapper : public WrapperBase
 	{
-		FGeometryCollectionWrapper() : WrapperBase(WrapperType::GeometryCollection) {}
-		FGeometryCollectionWrapper(
-			TSharedPtr<FGeometryCollection> RestCollectionIn,
-			TSharedPtr<FGeometryDynamicCollection> DynamicCollectionIn,
-			FGeometryCollectionPhysicsProxy* PhysObjectIn)
-			: WrapperBase(WrapperType::GeometryCollection)
-			, RestCollection(RestCollectionIn)
-			, DynamicCollection(DynamicCollectionIn)
-			, PhysObject(PhysObjectIn) {}
+		FGeometryCollectionWrapper();
+		//FGeometryCollectionWrapper(
+		//	TSharedPtr<FGeometryCollection> RestCollectionIn,
+		//	TSharedPtr<FGeometryDynamicCollection> DynamicCollectionIn,
+		//	FGeometryCollectionPhysicsProxy* PhysObjectIn)
+		//	: WrapperBase(WrapperType::GeometryCollection)
+		//	, RestCollection(RestCollectionIn)
+		//	, DynamicCollection(DynamicCollectionIn)
+		//	, PhysObject(PhysObjectIn) {}
 		static WrapperType StaticType() { return WrapperType::GeometryCollection; }
 		TSharedPtr<FGeometryCollection> RestCollection;
 		TSharedPtr<FGeometryDynamicCollection> DynamicCollection;
 		FGeometryCollectionPhysicsProxy* PhysObject;
+		FChaosUserData UserData;
 	};
 
 	struct RigidBodyWrapper : public WrapperBase
@@ -120,9 +122,8 @@ namespace GeometryCollectionTest
 
 	struct FrameworkParameters
 	{
-		FrameworkParameters() : Dt(1/60.) {}
-		FrameworkParameters(FReal dt) : Dt(Dt) {}
-		FReal Dt;
+		FReal Dt = 1.0 / 60.0;
+		FReal AsyncDt = -1.0;
 		Chaos::EThreadingMode ThreadingMode = Chaos::EThreadingMode::SingleThread;
 	};
 
@@ -134,6 +135,7 @@ namespace GeometryCollectionTest
 		virtual ~FFramework();
 
 		void AddSimulationObject(WrapperBase* Object);
+		void RemoveSimulationObject(WrapperBase* Object);
 		void Initialize();
 		void Advance();
 		FReal Dt;
diff --git a/Engine/Source/Runtime/Experimental/Chaos/Private/PBDRigidsSolver.cpp b/Engine/Source/Runtime/Experimental/Chaos/Private/PBDRigidsSolver.cpp
--- a/Engine/Source/Runtime/Experimental/Chaos/Private/PBDRigidsSolver.cpp
+++ b/Engine/Source/Runtime/Experimental/Chaos/Private/PBDRigidsSolver.cpp
@@ -833,6 +833,32 @@ namespace Chaos
 	bool bResimDelayDestroyObject = true;
 	FAutoConsoleVariableRef CVarResimDelayDestroyObject(TEXT("p.Resim.DelayDestroyObject"), bResimDelayDestroyObject, TEXT("Keep physics objects and physics proxy available until after they go out of rewind history. Disabled by default due to not fully implemented."));
 
+	void FPBDRigidsSolver::ClearCollisionEventsForProxy(const IPhysicsProxyBase* Proxy)
+	{
+		// @todo(chaos): Race Condition: In async mode UnregisterObject may be called while the physics thread is
+		// running, and perhaps just about to dispatch events. The command below will not execute until the 
+		// next frame, so we may still get collision events referencing our owner (via the UserData). 
+		// See FEventDefaults::RegisterCollisionEvent. The logic for collecting events checks Proxy->GetMarkDeleted(), 
+		// but the events may have already been collected, and that flag is also not intended for PT access and is
+		// not thread safe.
+		// For now we require the user to check GetMarkDeleted() in their collision handler to avoid the race condition.
+
+		GetEventManager()->template ClearEvents<FCollisionEventData>(EEventType::Collision, [Proxy]
+			(FCollisionEventData& EventDataInOut)
+			{
+				FCollisionDataArray const& CollisionData = EventDataInOut.CollisionData.AllCollisionsArray;
+				if (CollisionData.Num() > 0)
+				{
+					check(Proxy);
+					TArray<int32> const* const CollisionIndices = EventDataInOut.PhysicsProxyToCollisionIndices.PhysicsProxyToIndicesMap.Find(Proxy);
+					if (CollisionIndices)
+					{
+						EventDataInOut.PhysicsProxyToCollisionIndices.PhysicsProxyToIndicesMap.Remove(Proxy);
+					}
+				}
+			});
+	}
+
 	void FPBDRigidsSolver::UnregisterObject(FSingleParticlePhysicsProxy* Proxy)
 	{
 		UE_CHAOS_ASYNC_INITBODY_WRITESCOPELOCK(GetExternalDataLock_External());
@@ -872,22 +898,9 @@ namespace Chaos
 		{
 			UE_LOG(LogPBDRigidsSolver, Verbose, TEXT("FPBDRigidsSolver::UnregisterObject() ~ Dequeue"));
 
-				// Generally need to remove stale events for particles that no longer exist
-				GetEventManager()->template ClearEvents<FCollisionEventData>(EEventType::Collision, [Proxy]
-				(FCollisionEventData& EventDataInOut)
-				{
-					FCollisionDataArray const& CollisionData = EventDataInOut.CollisionData.AllCollisionsArray;
-					if (CollisionData.Num() > 0)
-					{
-						check(Proxy);
-						TArray<int32> const* const CollisionIndices = EventDataInOut.PhysicsProxyToCollisionIndices.PhysicsProxyToIndicesMap.Find(Proxy);
-						if (CollisionIndices)
-						{
-							EventDataInOut.PhysicsProxyToCollisionIndices.PhysicsProxyToIndicesMap.Remove(Proxy);
-						}
-					}
-
-				});
+			// Remove stale events for particles that no longer exist
+			// See comments in ClearCollisionEventsForProxy for potential race condition
+			ClearCollisionEventsForProxy(Proxy);
 
 			FGeometryParticleHandle* Handle = Proxy->GetHandle_LowLevel();
 			if (MRewindData)
@@ -953,6 +966,10 @@ namespace Chaos
 
 		EnqueueCommandImmediate([InProxy, this]()
 		{
+			// Remove stale events for particles that no longer exist
+			// See comments in ClearCollisionEventsForProxy for potential race condition
+			ClearCollisionEventsForProxy(InProxy);
+
 			GeometryCollectionPhysicsProxies_Internal.RemoveSingle(InProxy);
 			InProxy->SyncBeforeDestroy();
 			PendingDestroyGeometryCollectionPhysicsProxy.Add(InProxy);
@@ -1004,6 +1021,10 @@ namespace Chaos
 		EnqueueCommandImmediate(
 			[Proxy, this]()
 			{
+				// Remove stale events for particles that no longer exist
+				// See comments in ClearCollisionEventsForProxy for potential race condition
+				ClearCollisionEventsForProxy(Proxy);
+
 				if (Proxy->GetParticle_Internal())
 				{
 					MEvolution->DisableParticle(Proxy->GetParticle_Internal());
diff --git a/Engine/Source/Runtime/Experimental/Chaos/Public/PBDRigidsSolver.h b/Engine/Source/Runtime/Experimental/Chaos/Public/PBDRigidsSolver.h
--- a/Engine/Source/Runtime/Experimental/Chaos/Public/PBDRigidsSolver.h
+++ b/Engine/Source/Runtime/Experimental/Chaos/Public/PBDRigidsSolver.h
@@ -373,6 +373,8 @@ namespace Chaos
 		CHAOS_API virtual void PushPhysicsState(const FReal ExternalDt, const int32 NumSteps, const int32 NumExternalSteps) override;
 		CHAOS_API virtual void SetExternalTimestampConsumed_Internal(const int32 Timestamp) override;
 
+		CHAOS_API void ClearCollisionEventsForProxy(const IPhysicsProxyBase* Proxy);
+
 		CHAOS_API void ApplyCVars();
 
 		CHAOS_API void UpdateIsDeterministic();
