Chaos - Enable drive linear limits for non-linear joint constraints

diff --git a/Engine/Source/Programs/HeadlessChaos/Private/HeadlessChaosTestJointForces.cpp b/Engine/Source/Programs/HeadlessChaos/Private/HeadlessChaosTestJointForces.cpp
--- a/Engine/Source/Programs/HeadlessChaos/Private/HeadlessChaosTestJointForces.cpp
+++ b/Engine/Source/Programs/HeadlessChaos/Private/HeadlessChaosTestJointForces.cpp
@@ -437,51 +437,60 @@ namespace ChaosTest {
 	//
 	GTEST_TEST(JointForceTests, TestLinearDriveForceMode_MaxForce)
 	{
-		for (int32 IndexSimd = 0; IndexSimd < 2; IndexSimd++)
+		for (int32 IndexUseLinear = 0; IndexUseLinear < 2; ++IndexUseLinear)
 		{
-			bool bUseSimd = IndexSimd == 0;
-			const int32 NumSolverIterations = 20;
-			const FReal Gravity = 1000;
-			const FReal Dt = 0.01;
-			const int32 NumSteps = 100;
-			const int32 NumBodies = 2;
-			const FReal Extension = 10;
-
-			FJointChainTest<FPBDRigidsEvolutionGBF> Test(NumSolverIterations, Gravity);
-			Test.InitChain(NumBodies, FVec3(0, 0, -1));
-			Test.Evolution.GetJointCombinedConstraints().LinearConstraints.SetUseSimd(bUseSimd);
+			const bool bUseLinear = IndexUseLinear == 0;
 
-			if (!Test.Evolution.GetJointCombinedConstraints().LinearConstraints.GetSettings().bUsePositionBasedDrives)
+			for (int32 IndexSimd = 0; IndexSimd < 2; IndexSimd++)
 			{
-				Test.Evolution.SetNumVelocityIterations(NumSolverIterations);
-			}
+				bool bUseSimd = IndexSimd == 0;
+				const int32 NumSolverIterations = 20;
+				const FReal Gravity = 1000.0;
+				const FReal Dt = 0.01;
+				const int32 NumSteps = 100;
+				const int32 NumBodies = 2;
+				const FReal Extension = 10;
 
-			// Disable all limits
-			Test.JointSettings[0].LinearMotionTypes = { EJointMotionType::Free, EJointMotionType::Free, EJointMotionType::Free };
-			Test.JointSettings[0].AngularMotionTypes = { EJointMotionType::Free, EJointMotionType::Free, EJointMotionType::Free };
+				FJointChainTest<FPBDRigidsEvolutionGBF> Test(NumSolverIterations, Gravity);
+				Test.InitChain(NumBodies, FVec3(0, 0, -1));
 
-			// Set up the drive with a force limit
-			const FReal Stiffness = Test.ParticleMasses[1] * Gravity / Extension;
-			const FReal Damping = 2 * FMath::Sqrt(Stiffness * Test.ParticleMasses[1]);
-			Test.JointSettings[0].bLinearPositionDriveEnabled = { false, false, true };
-			Test.JointSettings[0].bLinearVelocityDriveEnabled = { false, false, true };
-			Test.JointSettings[0].LinearDriveForceMode = EJointForceMode::Force;
-			Test.JointSettings[0].LinearDriveStiffness = FVec3(0, 0, Stiffness);
-			Test.JointSettings[0].LinearDriveDamping = FVec3(0, 0, Damping);
-			Test.JointSettings[0].LinearDriveMaxForce = FVec3(0, 0, 0.5 * Stiffness * Extension);
+				FPBDJointConstraints& Constraints = bUseLinear ? Test.Evolution.GetJointCombinedConstraints().LinearConstraints : Test.Evolution.GetJointCombinedConstraints().NonLinearConstraints;
 
-			Test.Create();
+				Constraints.SetUseSimd(bUseSimd);
 
-			FConstGenericParticleHandle P1 = Test.GetParticle(1);
+				if (!Constraints.GetSettings().bUsePositionBasedDrives)
+				{
+					Test.Evolution.SetNumVelocityIterations(NumSolverIterations);
+				}
 
-			// Run the sim - the dangling box should reach a steady state
-			for (int32 i = 0; i < NumSteps; ++i)
-			{
-				Test.Evolution.AdvanceOneTimeStep(Dt);
-				Test.Evolution.EndFrame(Dt);
+				// Disable all limits
+				Test.JointSettings[0].LinearMotionTypes = { EJointMotionType::Free, EJointMotionType::Free, EJointMotionType::Free };
+				Test.JointSettings[0].AngularMotionTypes = { EJointMotionType::Free, EJointMotionType::Free, EJointMotionType::Free };
 
-				const FReal ForceZ = -Test.Evolution.GetJointCombinedConstraints().LinearConstraints.GetConstraintLinearImpulse(0).Z / Dt;
-				EXPECT_LT(ForceZ, Test.JointSettings[0].LinearDriveMaxForce.Z + 0.1);
+				// Set up the drive with a force limit
+				const FReal Stiffness = Test.ParticleMasses[1] * Gravity / Extension;
+				const FReal Damping = 2 * FMath::Sqrt(Stiffness * Test.ParticleMasses[1]);
+				Test.JointSettings[0].bLinearPositionDriveEnabled = { false, false, true };
+				Test.JointSettings[0].bLinearVelocityDriveEnabled = { false, false, true };
+				Test.JointSettings[0].LinearDriveForceMode = EJointForceMode::Force;
+				Test.JointSettings[0].LinearDriveStiffness = FVec3(0, 0, Stiffness);
+				Test.JointSettings[0].LinearDriveDamping = FVec3(0, 0, Damping);
+				Test.JointSettings[0].LinearDriveMaxForce = FVec3(0, 0, 0.5 * Stiffness * Extension);
+				Test.JointSettings[0].bUseLinearSolver = bUseLinear;
+
+				Test.Create();
+
+				FConstGenericParticleHandle P1 = Test.GetParticle(1);
+
+				// Run the sim - the dangling box should reach a steady state
+				for (int32 i = 0; i < NumSteps; ++i)
+				{
+					Test.Evolution.AdvanceOneTimeStep(Dt);
+					Test.Evolution.EndFrame(Dt);
+
+					const FReal ForceZ = -Constraints.GetConstraintLinearImpulse(0).Z / Dt;
+					EXPECT_LT(ForceZ, Test.JointSettings[0].LinearDriveMaxForce.Z + 0.1);
+				}
 			}
 		}
 	}
diff --git a/Engine/Source/Runtime/Experimental/Chaos/Private/Chaos/Joint/PBDJointSolverGaussSeidel.cpp b/Engine/Source/Runtime/Experimental/Chaos/Private/Chaos/Joint/PBDJointSolverGaussSeidel.cpp
--- a/Engine/Source/Runtime/Experimental/Chaos/Private/Chaos/Joint/PBDJointSolverGaussSeidel.cpp
+++ b/Engine/Source/Runtime/Experimental/Chaos/Private/Chaos/Joint/PBDJointSolverGaussSeidel.cpp
@@ -1269,6 +1269,7 @@ namespace Chaos
 		const FVec3& Axis,
 		const FReal Delta,
 		const FReal TargetVel,
+		const FReal LambdaLimit,
 		FReal& Lambda)
 	{
 		if (bRealTypeCompatibleWithISPC && bChaos_Joint_ISPC_Enabled)
@@ -1302,7 +1303,11 @@ namespace Chaos
 			const FReal S = SpringMassScale * JointStiffness * Dt * Dt;
 			const FReal D = SpringMassScale * JointDamping * Dt;
 			const FReal Multiplier = (FReal)1 / ((S + D) * II + (FReal)1);
-			const FReal DLambda = SolverStiffness * Multiplier * (S * Delta - D * VelDt - Lambda);
+			FReal DLambda = SolverStiffness * Multiplier * (S * Delta - D * VelDt - Lambda);
+			if (LambdaLimit > SMALL_NUMBER)
+			{
+				DLambda = FMath::Clamp(DLambda, -LambdaLimit - Lambda, LambdaLimit - Lambda);
+			}
 	
 			const FVec3 DP0 = (InvM(0) * DLambda) * Axis;
 			const FVec3 DP1 = (-InvM(1) * DLambda) * Axis;
@@ -2041,7 +2046,7 @@ namespace Chaos
 				const FReal JointStiffness = FPBDJointUtilities::GetSoftLinearStiffness(SolverSettings, JointSettings);
 				const FReal JointDamping = FPBDJointUtilities::GetSoftLinearDamping(SolverSettings, JointSettings);
 				const bool bAccelerationMode = FPBDJointUtilities::GetLinearSoftAccelerationMode(SolverSettings, JointSettings);
-				ApplyPositionConstraintSoft(Dt, JointStiffness, JointDamping, bAccelerationMode, Axis, Error, 0.0f, LinearSoftLambda);
+				ApplyPositionConstraintSoft(Dt, JointStiffness, JointDamping, bAccelerationMode, Axis, Error, 0.0f, 0.0f, LinearSoftLambda);
 			}
 
 			// Update the linear violation
@@ -2083,7 +2088,7 @@ namespace Chaos
 				const FReal JointStiffness = FPBDJointUtilities::GetSoftLinearStiffness(SolverSettings, JointSettings);
 				const FReal JointDamping = FPBDJointUtilities::GetSoftLinearDamping(SolverSettings, JointSettings);
 				const bool bAccelerationMode = FPBDJointUtilities::GetLinearSoftAccelerationMode(SolverSettings, JointSettings);
-				ApplyPositionConstraintSoft(Dt, JointStiffness, JointDamping, bAccelerationMode, Axis, AxialError, 0.0f, LinearSoftLambda);
+				ApplyPositionConstraintSoft(Dt, JointStiffness, JointDamping, bAccelerationMode, Axis, AxialError, 0.0f, 0.0f, LinearSoftLambda);
 			}
 			else if (AxialMotion != EJointMotionType::Free)
 			{
@@ -2111,7 +2116,7 @@ namespace Chaos
 				const FReal JointStiffness = FPBDJointUtilities::GetSoftLinearStiffness(SolverSettings, JointSettings);
 				const FReal JointDamping = FPBDJointUtilities::GetSoftLinearDamping(SolverSettings, JointSettings);
 				const bool bAccelerationMode = FPBDJointUtilities::GetLinearSoftAccelerationMode(SolverSettings, JointSettings);
-				ApplyPositionConstraintSoft(Dt, JointStiffness, JointDamping, bAccelerationMode, RadialAxis, RadialError, 0.0f, LinearSoftLambda);
+				ApplyPositionConstraintSoft(Dt, JointStiffness, JointDamping, bAccelerationMode, RadialAxis, RadialError, 0.0f, 0.0f, LinearSoftLambda);
 			}
 			else if (RadialMotion != EJointMotionType::Free)
 			{
@@ -2155,7 +2160,7 @@ namespace Chaos
 				const FReal JointStiffness = FPBDJointUtilities::GetSoftLinearStiffness(SolverSettings, JointSettings);
 				const FReal JointDamping = FPBDJointUtilities::GetSoftLinearDamping(SolverSettings, JointSettings);
 				const bool bAccelerationMode = FPBDJointUtilities::GetLinearSoftAccelerationMode(SolverSettings, JointSettings);
-				ApplyPositionConstraintSoft(Dt, JointStiffness, JointDamping, bAccelerationMode, Axis, Error, 0.0f, LinearSoftLambda);
+				ApplyPositionConstraintSoft(Dt, JointStiffness, JointDamping, bAccelerationMode, Axis, Error, 0.0f, 0.0f, LinearSoftLambda);
 			}
 			else
 			{
@@ -2182,11 +2187,12 @@ namespace Chaos
 		const FReal JointStiffness = FPBDJointUtilities::GetLinearDriveStiffness(SolverSettings, JointSettings, AxisIndex);
 		const FReal JointDamping = FPBDJointUtilities::GetLinearDriveDamping(SolverSettings, JointSettings, AxisIndex);
 		const bool bAccelerationMode = FPBDJointUtilities::GetLinearDriveAccelerationMode(SolverSettings, JointSettings);
+		const FReal LambdaLimit = JointSettings.LinearDriveMaxForce[AxisIndex] * Dt * Dt;
 
 		if ((FMath::Abs(DeltaPos) > PositionTolerance) || (JointDamping > 0.0f))
 		{
 			FReal ReturnedLambda = LinearDriveLambdas[AxisIndex];
-			ApplyPositionConstraintSoft(Dt, JointStiffness, JointDamping, bAccelerationMode, Axis, DeltaPos, DeltaVel, ReturnedLambda);
+			ApplyPositionConstraintSoft(Dt, JointStiffness, JointDamping, bAccelerationMode, Axis, DeltaPos, DeltaVel, LambdaLimit, ReturnedLambda);
 			LinearDriveLambdas[AxisIndex] = ReturnedLambda;
 		}
 	}
diff --git a/Engine/Source/Runtime/Experimental/Chaos/Public/Chaos/Joint/PBDJointSolverGaussSeidel.h b/Engine/Source/Runtime/Experimental/Chaos/Public/Chaos/Joint/PBDJointSolverGaussSeidel.h
--- a/Engine/Source/Runtime/Experimental/Chaos/Public/Chaos/Joint/PBDJointSolverGaussSeidel.h
+++ b/Engine/Source/Runtime/Experimental/Chaos/Public/Chaos/Joint/PBDJointSolverGaussSeidel.h
@@ -312,6 +312,7 @@ namespace Chaos
 			const FVec3& Axis,
 			const FReal Delta,
 			const FReal TargetVel,
+			const FReal LambdaLimit,
 			FReal& Lambda);
 
 		void ApplyRotationConstraint(
