HeadlessChaos - fix GeometryCollection_RigidBodies_ClusterTest_UnionClusterCollisions

diff --git a/Engine/Source/Programs/HeadlessChaos/Private/GeometryCollection/GeometryCollectionTestClustering.cpp b/Engine/Source/Programs/HeadlessChaos/Private/GeometryCollection/GeometryCollectionTestClustering.cpp
--- a/Engine/Source/Programs/HeadlessChaos/Private/GeometryCollection/GeometryCollectionTestClustering.cpp
+++ b/Engine/Source/Programs/HeadlessChaos/Private/GeometryCollection/GeometryCollectionTestClustering.cpp
@@ -1844,10 +1844,14 @@ namespace GeometryCollectionTest
 
 				EXPECT_TRUE(!InitialRootPosition.Equals(Root->GetX())); // root moves
 
-				EXPECT_TRUE(RelativeChildOffsets[0].Equals(GlobalTransform1[0].GetRelativeTransform(RootTransform).GetTranslation()));
-				EXPECT_TRUE(RelativeChildOffsets[1].Equals(GlobalTransform1[1].GetRelativeTransform(RootTransform).GetTranslation()));
-				EXPECT_TRUE(RelativeChildOffsets[2].Equals(GlobalTransform2[0].GetRelativeTransform(RootTransform).GetTranslation()));
-				EXPECT_TRUE(RelativeChildOffsets[3].Equals(GlobalTransform2[1].GetRelativeTransform(RootTransform).GetTranslation()));
+				FVector ActualRelativeChildOffsets0 = GlobalTransform1[0].GetRelativeTransform(RootTransform).GetTranslation();
+				FVector ActualRelativeChildOffsets1 = GlobalTransform1[1].GetRelativeTransform(RootTransform).GetTranslation();
+				FVector ActualRelativeChildOffsets2 = GlobalTransform2[0].GetRelativeTransform(RootTransform).GetTranslation();
+				FVector ActualRelativeChildOffsets3 = GlobalTransform2[1].GetRelativeTransform(RootTransform).GetTranslation();
+				EXPECT_TRUE(RelativeChildOffsets[0].Equals(ActualRelativeChildOffsets0));
+				EXPECT_TRUE(RelativeChildOffsets[1].Equals(ActualRelativeChildOffsets1));
+				EXPECT_TRUE(RelativeChildOffsets[2].Equals(ActualRelativeChildOffsets2));
+				EXPECT_TRUE(RelativeChildOffsets[3].Equals(ActualRelativeChildOffsets3));
 			}
 
 			//
diff --git a/Engine/Source/Programs/HeadlessChaos/Private/GeometryCollection/GeometryCollectionTestFramework.cpp b/Engine/Source/Programs/HeadlessChaos/Private/GeometryCollection/GeometryCollectionTestFramework.cpp
--- a/Engine/Source/Programs/HeadlessChaos/Private/GeometryCollection/GeometryCollectionTestFramework.cpp
+++ b/Engine/Source/Programs/HeadlessChaos/Private/GeometryCollection/GeometryCollectionTestFramework.cpp
@@ -187,7 +187,7 @@ namespace GeometryCollectionTest
 			SimulationParams.UseMACD = false;
 			SimulationParams.AllowPartialIslandSleep = true;
 			SimulationParams.PositionSolverIterations = 8;
-			SimulationParams.VelocitySolverIterations = 1;
+			SimulationParams.VelocitySolverIterations = 2;
 			SimulationParams.ProjectionSolverIterations = 1;
 			SimulationParams.UserData = reinterpret_cast<void*>(&Wrapper->UserData);
 
