Add a  unit test for Dirty particles

diff --git a/Engine/Source/Programs/HeadlessChaos/Private/HeadlessChaosTestDirtyParticle.cpp b/Engine/Source/Programs/HeadlessChaos/Private/HeadlessChaosTestDirtyParticle.cpp
--- a/Engine/Source/Programs/HeadlessChaos/Private/HeadlessChaosTestDirtyParticle.cpp
+++ b/Engine/Source/Programs/HeadlessChaos/Private/HeadlessChaosTestDirtyParticle.cpp
@@ -3,6 +3,9 @@
 #include "HeadlessChaos.h"
 #include "HeadlessChaosTestUtility.h"
 
+#include "ChaosSolversModule.h"
+#include "Chaos/ParticleHandle.h"
+#include "PhysicsProxy/PhysicsProxies.h"
 #include "Chaos/PBDRigidsEvolutionGBF.h"
 #include "Chaos/Box.h"
 #include "Chaos/Sphere.h"
@@ -67,5 +70,64 @@ namespace ChaosTest
 
 #endif
 	}
+
+	// Simple dynamics particle falling under gravity, 
+	// Then switch it to a kinematic state
+	GTEST_TEST(DirtyParticleTests, ChangeStateTest)
+	{
+		FChaosSolversModule* Module = FChaosSolversModule::GetModule();
+
+		const float AsyncDt = -1.0f;
+		Chaos::FPBDRigidsSolver* Solver = Module->CreateSolver(nullptr, AsyncDt, Chaos::EThreadingMode::SingleThread, TEXT("Test Dirty Particles"));
+		ChaosTest::InitSolverSettings(Solver);
+
+		FVector3d HalfExtents = FVec3(1.0, 1.0, 1.0) * FTransform::Identity.GetScale3D().GetAbs();
+		FImplicitObjectPtr BoxGeom = FImplicitObjectPtr(new Chaos::FImplicitBox3(-HalfExtents, HalfExtents));
+		FSingleParticlePhysicsProxy* Proxy = FSingleParticlePhysicsProxy::Create(Chaos::TPBDRigidParticle<FReal, 3>::CreateParticle());
+		FRigidBodyHandle_External& ParticleGT = Proxy->GetGameThreadAPI();
+		ParticleGT.SetGeometry(BoxGeom);
+		ParticleGT.SetX(FTransform::Identity.GetLocation());
+		ParticleGT.SetR(FTransform::Identity.GetRotation());
+		ParticleGT.SetM(1.0);
+		ParticleGT.SetInvM(1.0f);
+
+		Proxy->GetGameThreadAPI().SetObjectState(EObjectStateType::Dynamic);
+		Solver->RegisterObject(Proxy);
+
+		FReal Altitude = ParticleGT.GetX().Z;
+		EXPECT_EQ(Altitude, 0.0);
+
+		const float Dt = 1 / 60.0f;
+		Solver->AdvanceAndDispatch_External(Dt);
+		Solver->UpdateGameThreadStructures();
+		Solver->SyncEvents_GameThread();
+
+		FReal OldAltitude = Altitude;
+		Altitude = ParticleGT.GetX().Z;
+		EXPECT_LT(Altitude, OldAltitude);
+
+		Solver->AdvanceAndDispatch_External(Dt);
+		Solver->UpdateGameThreadStructures();
+		Solver->SyncEvents_GameThread();
+
+		OldAltitude = Altitude;
+		Altitude = ParticleGT.GetX().Z;
+		EXPECT_LT(Altitude, OldAltitude);
+
+		// Change the particle state to kinematic
+		Proxy->GetGameThreadAPI().SetObjectState(EObjectStateType::Kinematic);
+
+		Solver->AdvanceAndDispatch_External(Dt);
+		Solver->UpdateGameThreadStructures();
+		Solver->SyncEvents_GameThread();
+
+		// Expecting the particle is not falling down anymore
+		OldAltitude = Altitude;
+		Altitude = ParticleGT.GetX().Z;
+		EXPECT_EQ(Altitude, OldAltitude);
+
+		Solver->UnregisterObject(Proxy);
+		FChaosSolversModule::GetModule()->DestroySolver(Solver);
+	}
 }
 
