Set the flag FPBDIslandParticle::bWasUpdated for all updates involving dynamic rigid particles triggered in FSingleParticlePhysicsProxy::PushToPhysicsState.

diff --git a/Engine/Source/Programs/HeadlessChaos/Private/HeadlessChaosTestGraphEvolution.cpp b/Engine/Source/Programs/HeadlessChaos/Private/HeadlessChaosTestGraphEvolution.cpp
--- a/Engine/Source/Programs/HeadlessChaos/Private/HeadlessChaosTestGraphEvolution.cpp
+++ b/Engine/Source/Programs/HeadlessChaos/Private/HeadlessChaosTestGraphEvolution.cpp
@@ -1599,8 +1599,7 @@ namespace ChaosTest
 		Dynamics.SetAngularImpulseVelocity(FVec3(0, 0, 0));
 
 		// This emulates what happens in @PushToPhysicsStateImp (SingleParticlePhysicsProxy.cpp).
-		Test.ParticleHandles[2]->SetDynamics(Dynamics);
-		Test.Evolution.ResetVSmoothFromForces(*Test.ParticleHandles[2]);
+		Test.Evolution.SetParticleDynamics(Test.ParticleHandles[2], Dynamics);
 		Test.Evolution.SetParticleObjectState(Test.ParticleHandles[2], EObjectStateType::Dynamic);
 
 		Test.Advance();
@@ -1671,8 +1670,7 @@ namespace ChaosTest
 			Dynamics.SetAngularAcceleration(FVec3(0, 0, 0));
 			Dynamics.SetLinearImpulseVelocity(FVec3(50, 50, 0));
 			Dynamics.SetAngularImpulseVelocity(FVec3(0, 0, 0));
-			Test.ParticleHandles.Last()->SetDynamics(Dynamics);
-			Test.Evolution.ResetVSmoothFromForces(*Test.ParticleHandles.Last());
+			Test.Evolution.SetParticleDynamics(Test.ParticleHandles.Last(), Dynamics);
 
 			Test.Advance();
 
@@ -1749,8 +1747,7 @@ namespace ChaosTest
 			Dynamics.SetAngularAcceleration(FVec3(0, 0, 0));
 			Dynamics.SetLinearImpulseVelocity(FVec3(50, 50, 0));
 			Dynamics.SetAngularImpulseVelocity(FVec3(0, 0, 0));
-			Test.ParticleHandles.Last()->SetDynamics(Dynamics);
-			Test.Evolution.ResetVSmoothFromForces(*Test.ParticleHandles.Last());
+			Test.Evolution.SetParticleDynamics(Test.ParticleHandles.Last(), Dynamics);
 
 			Test.Advance();
 
@@ -1845,8 +1842,7 @@ namespace ChaosTest
 			Dynamics.SetAngularAcceleration(FVec3(0, 0, 0));
 			Dynamics.SetLinearImpulseVelocity(FVec3(50, 50, 0));
 			Dynamics.SetAngularImpulseVelocity(FVec3(0, 0, 0));
-			Test.ParticleHandles.Last()->SetDynamics(Dynamics);
-			Test.Evolution.ResetVSmoothFromForces(*Test.ParticleHandles.Last());
+			Test.Evolution.SetParticleDynamics(Test.ParticleHandles.Last(), Dynamics);
 
 			Test.Advance();
 
@@ -1956,8 +1952,7 @@ namespace ChaosTest
 			Dynamics.SetAngularAcceleration(FVec3(0, 0, 0));
 			Dynamics.SetLinearImpulseVelocity(FVec3(50, 50, 0));
 			Dynamics.SetAngularImpulseVelocity(FVec3(0, 0, 0));
-			Test.ParticleHandles.Last()->SetDynamics(Dynamics);
-			Test.Evolution.ResetVSmoothFromForces(*Test.ParticleHandles.Last());
+			Test.Evolution.SetParticleDynamics(Test.ParticleHandles.Last(), Dynamics);
 
 			Test.Advance();
 
@@ -2215,8 +2210,7 @@ namespace ChaosTest
 		Dynamics.SetAngularAcceleration(FVec3(0, 0, 0));
 		Dynamics.SetLinearImpulseVelocity(FVec3(50, 50, 0));
 		Dynamics.SetAngularImpulseVelocity(FVec3(0, 0, 0));
-		Test.ParticleHandles[7]->SetDynamics(Dynamics);
-		Test.Evolution.ResetVSmoothFromForces(*Test.ParticleHandles[7]);
+		Test.Evolution.SetParticleDynamics(Test.ParticleHandles[7], Dynamics);
 
 		Test.Advance();
 
@@ -2323,8 +2317,7 @@ namespace ChaosTest
 			Dynamics.SetAngularImpulseVelocity(FVec3(0, 0, 0));
 			for (FPBDRigidParticleHandle* Particle : Test.ParticleHandles)
 			{
-				Particle->SetDynamics(Dynamics);
-				Test.Evolution.ResetVSmoothFromForces(*Particle);
+				Test.Evolution.SetParticleDynamics(Particle, Dynamics);
 			}
 
 			// Step one more time, the sleep counter should be reset to 0 because the particles are moving
@@ -2376,8 +2369,7 @@ namespace ChaosTest
 			Dynamics.SetAngularAcceleration(FVec3(0, 0, 0));
 			Dynamics.SetLinearImpulseVelocity(FVec3(50, 50, 0));
 			Dynamics.SetAngularImpulseVelocity(FVec3(0, 0, 0));
-			Test.ParticleHandles.Last()->SetDynamics(Dynamics);
-			Test.Evolution.ResetVSmoothFromForces(*Test.ParticleHandles.Last());
+			Test.Evolution.SetParticleDynamics(Test.ParticleHandles.Last(), Dynamics);
 
 			// Step one more time, the particles should not sleep since a wake-up event just happened.
 			Test.Advance();
@@ -2550,8 +2542,7 @@ namespace ChaosTest
 		Dynamics.SetAngularAcceleration(FVec3(0, 0, 0));
 		Dynamics.SetLinearImpulseVelocity(FVec3(0, 0.1, 0));
 		Dynamics.SetAngularImpulseVelocity(FVec3(0, 0, 0));
-		Test.ParticleHandles[4]->SetDynamics(Dynamics);
-		Test.Evolution.ResetVSmoothFromForces(*Test.ParticleHandles[4]);
+		Test.Evolution.SetParticleDynamics(Test.ParticleHandles[4], Dynamics);
 
 		Test.Advance();
 
diff --git a/Engine/Source/Programs/HeadlessChaos/Private/HeadlessChaosTestSleeping.cpp b/Engine/Source/Programs/HeadlessChaos/Private/HeadlessChaosTestSleeping.cpp
--- a/Engine/Source/Programs/HeadlessChaos/Private/HeadlessChaosTestSleeping.cpp
+++ b/Engine/Source/Programs/HeadlessChaos/Private/HeadlessChaosTestSleeping.cpp
@@ -515,6 +515,124 @@ namespace ChaosTest
 		CVarMomentumPropagation->Set(MomentumPropagationDefault);
 	}
 
+	TEST_P(SleepingTests, UpdateMassOfCubeInStack)
+	{
+		const int32 NumOfStackedObjects = 5;
+		const FReal Dimension = 200;
+		const bool bPartialSleepEnabled = GetParam();
+
+		// Turn off momentum propagation to test level-based waking only
+		IConsoleVariable* CVarMomentumPropagation = IConsoleManager::Get().FindConsoleVariable(TEXT("p.Chaos.Solver.Sleep.MomentumPropagation"), false);
+		check(CVarMomentumPropagation && CVarMomentumPropagation->IsVariableFloat());
+		const FRealSingle MomentumPropagationDefault = CVarMomentumPropagation->GetFloat();
+		CVarMomentumPropagation->Set(0);
+
+		IConsoleVariable* CVarPartialWakePreIntegration = IConsoleManager::Get().FindConsoleVariable(TEXT("p.Chaos.Solver.Sleep.PartialWakePreIntegration"), false);
+		check(CVarPartialWakePreIntegration && CVarPartialWakePreIntegration->IsVariableBool());
+		const bool bPartialWakePreIntegrationDefault = CVarPartialWakePreIntegration->GetBool();
+		CVarPartialWakePreIntegration->Set(true);
+
+		// Create a stack of cubes
+		FSleepingTest Test(bPartialSleepEnabled);
+		Test.MakeStackOfCubes(NumOfStackedObjects, /*Size=*/Dimension);
+
+		// Simulate until all particles are sleeping
+		Test.AdvanceUntilSleeping();
+
+		// Check that everything is sleeping
+		for (FRigidBodyHandle_External* ParticleHandle : Test.ParticleHandles)
+		{
+			EXPECT_TRUE(ParticleHandle->ObjectState() == EObjectStateType::Sleeping);
+		}
+
+		// Update the particle mass without waking the particle
+		Test.ParticleHandles[3]->SetM(10);
+
+		Test.Advance();
+
+		// Updating the particle mass only will not wake the stack.
+		// Neither in full nor in partial island sleep mode.
+		for (FRigidBodyHandle_External* ParticleHandle : Test.ParticleHandles)
+		{
+			EXPECT_TRUE(ParticleHandle->ObjectState() == EObjectStateType::Sleeping);
+		}
+
+		// Update the particle mass and wake up the particle
+		Test.ParticleHandles[3]->SetM(20);
+		Test.ParticleHandles[3]->SetObjectState(EObjectStateType::Dynamic);
+
+		Test.Advance();
+
+		if (!bPartialSleepEnabled)
+		{
+			// The entire stack wakes up.
+			for (FRigidBodyHandle_External* ParticleHandle : Test.ParticleHandles)
+			{
+				EXPECT_TRUE(ParticleHandle->ObjectState() == EObjectStateType::Dynamic);
+			}
+		}
+		else
+		{
+			// The stack wakes up only above the level of where the mass was changed.
+			EXPECT_TRUE(Test.ParticleHandles[0]->ObjectState() == EObjectStateType::Sleeping);
+			EXPECT_TRUE(Test.ParticleHandles[1]->ObjectState() == EObjectStateType::Sleeping);
+			EXPECT_TRUE(Test.ParticleHandles[2]->ObjectState() == EObjectStateType::Dynamic);
+			EXPECT_TRUE(Test.ParticleHandles[3]->ObjectState() == EObjectStateType::Dynamic);
+			EXPECT_TRUE(Test.ParticleHandles[4]->ObjectState() == EObjectStateType::Dynamic);
+		}
+
+		CVarMomentumPropagation->Set(MomentumPropagationDefault);
+		CVarPartialWakePreIntegration->Set(bPartialWakePreIntegrationDefault);
+	}
+
+	TEST_P(SleepingTests, DeleteBottomCubeInStack)
+	{
+		const int32 NumOfStackedObjects = 5;
+		const FReal Dimension = 200;
+		const bool bPartialSleepEnabled = GetParam();
+
+		// Turn off momentum propagation to test level-based waking only
+		IConsoleVariable* CVarMomentumPropagation = IConsoleManager::Get().FindConsoleVariable(TEXT("p.Chaos.Solver.Sleep.MomentumPropagation"), false);
+		check(CVarMomentumPropagation && CVarMomentumPropagation->IsVariableFloat());
+		const FRealSingle MomentumPropagationDefault = CVarMomentumPropagation->GetFloat();
+		CVarMomentumPropagation->Set(0);
+
+		IConsoleVariable* CVarPartialWakePreIntegration = IConsoleManager::Get().FindConsoleVariable(TEXT("p.Chaos.Solver.Sleep.PartialWakePreIntegration"), false);
+		check(CVarPartialWakePreIntegration && CVarPartialWakePreIntegration->IsVariableBool());
+		const bool bPartialWakePreIntegrationDefault = CVarPartialWakePreIntegration->GetBool();
+		CVarPartialWakePreIntegration->Set(true);
+
+		// Create a stack of cubes
+		FSleepingTest Test(bPartialSleepEnabled);
+		Test.MakeStackOfCubes(NumOfStackedObjects, /*Size=*/Dimension);
+
+		// Simulate until all particles are sleeping
+		Test.AdvanceUntilSleeping();
+
+		// Check that everything is sleeping
+		for (FRigidBodyHandle_External* ParticleHandle : Test.ParticleHandles)
+		{
+			EXPECT_TRUE(ParticleHandle->ObjectState() == EObjectStateType::Sleeping);
+		}
+
+		// Delete the bottom cube
+		FSingleParticlePhysicsProxy* BottomBoxProxy = Test.ParticleProxies[0];
+		Test.ParticleHandles.RemoveAt(0);
+		Test.ParticleProxies.RemoveAt(0);
+		Test.Solver->UnregisterObject(BottomBoxProxy);
+
+		Test.Advance();
+
+		// The entire stack will wake up
+		for (FRigidBodyHandle_External* ParticleHandle : Test.ParticleHandles)
+		{
+			EXPECT_TRUE(ParticleHandle->ObjectState() == EObjectStateType::Dynamic);
+		}
+
+		CVarMomentumPropagation->Set(MomentumPropagationDefault);
+		CVarPartialWakePreIntegration->Set(bPartialWakePreIntegrationDefault);
+	}
+
 	TEST_P(SleepingTests, TeleportCubeOutOfStack)
 	{
 		const int32 NumOfStackedObjects = 5;
diff --git a/Engine/Source/Runtime/Experimental/Chaos/Private/Chaos/Island/IslandManager.cpp b/Engine/Source/Runtime/Experimental/Chaos/Private/Chaos/Island/IslandManager.cpp
--- a/Engine/Source/Runtime/Experimental/Chaos/Private/Chaos/Island/IslandManager.cpp
+++ b/Engine/Source/Runtime/Experimental/Chaos/Private/Chaos/Island/IslandManager.cpp
@@ -837,7 +837,7 @@ namespace Chaos::Private
 			if (Node->Island != nullptr)
 			{
 				// We permit sleeping for this island if partial island sleeping and pre-integration waking are enabled.
-				// This will allows us to perform partial waking if a particle is moved into the island by updating its transform explicitly.
+				// This will allows us to perform partial island waking if a dynamic particle has been woken by the user or had its transform set explicitly.
 				if (CVars::bChaosSolverPartialIslandSleep && CVars::bChaosSolverPartialWakePreIntegration)
 				{
 					bIsSleepAllowed = true;
diff --git a/Engine/Source/Runtime/Experimental/Chaos/Private/Chaos/PBDRigidsEvolution.cpp b/Engine/Source/Runtime/Experimental/Chaos/Private/Chaos/PBDRigidsEvolution.cpp
--- a/Engine/Source/Runtime/Experimental/Chaos/Private/Chaos/PBDRigidsEvolution.cpp
+++ b/Engine/Source/Runtime/Experimental/Chaos/Private/Chaos/PBDRigidsEvolution.cpp
@@ -1183,6 +1183,8 @@ namespace Chaos
 		// E.g., this will reset the sleep counter to prevent the particle from immediately sleeping again.
 		// NOTE: This also allows us to "wake" a kinematic which will wake all islands that the kinematic is in.
 		IslandManager.WakeParticleIslands(Particle);
+
+		IslandManager.FlagParticleUpdated(Particle);
 	}
 
 	void FPBDRigidsEvolutionBase::SetParticleSleepType(FPBDRigidParticleHandle* Particle, ESleepType InSleepType)
diff --git a/Engine/Source/Runtime/Experimental/Chaos/Private/Chaos/PBDRigidsEvolutionGBF.cpp b/Engine/Source/Runtime/Experimental/Chaos/Private/Chaos/PBDRigidsEvolutionGBF.cpp
--- a/Engine/Source/Runtime/Experimental/Chaos/Private/Chaos/PBDRigidsEvolutionGBF.cpp
+++ b/Engine/Source/Runtime/Experimental/Chaos/Private/Chaos/PBDRigidsEvolutionGBF.cpp
@@ -1701,6 +1701,19 @@ void FPBDRigidsEvolutionGBF::OnParticleMoved(FGeometryParticleHandle* InParticle
 	}
 }
 
+void FPBDRigidsEvolutionGBF::SetParticleDynamics(FPBDRigidParticleHandle* Rigid, const FParticleDynamics& Dynamics)
+{
+	Rigid->SetDynamics(Dynamics);
+	ResetVSmoothFromForces(*Rigid);
+	GetIslandManager().FlagParticleUpdated(Rigid->Handle());
+}
+
+void FPBDRigidsEvolutionGBF::SetParticleMassProps(FPBDRigidParticleHandle* Rigid, const FParticleMassProps& MassProperties)
+{
+	Rigid->SetMassProps(MassProperties);
+	GetIslandManager().FlagParticleUpdated(Rigid->Handle());
+}
+
 void FPBDRigidsEvolutionGBF::ApplyParticleTransformCorrectionDelta(FGeometryParticleHandle* InParticle, const FVec3& InPosDelta, const FVec3& InRotDelta, const bool bApplyToConnectedBodies, const bool bInRecalculateFrictionOnConnectedBodies, const TArray<FParticleID>& ExcludeConnections)
 {
 	ApplyParticleTransformCorrection(
@@ -1857,6 +1870,8 @@ void FPBDRigidsEvolutionGBF::SetParticleVelocities(FGeometryParticleHandle* InPa
 				WakeParticle(Rigid);
 			}
 
+			GetIslandManager().FlagParticleUpdated(Rigid->Handle());
+
 			// @todo(chaos): do we want to reset static friction an any existing contacts as well?
 		}
 	}
diff --git a/Engine/Source/Runtime/Experimental/Chaos/Private/PhysicsProxy/SingleParticlePhysicsProxy.cpp b/Engine/Source/Runtime/Experimental/Chaos/Private/PhysicsProxy/SingleParticlePhysicsProxy.cpp
--- a/Engine/Source/Runtime/Experimental/Chaos/Private/PhysicsProxy/SingleParticlePhysicsProxy.cpp
+++ b/Engine/Source/Runtime/Experimental/Chaos/Private/PhysicsProxy/SingleParticlePhysicsProxy.cpp
@@ -143,15 +143,14 @@ void PushToPhysicsStateImp(const Chaos::FDirtyPropertiesManager& Manager, Chaos:
 
 			if(bHasDynamicData)
 			{
-				if(auto NewData = ParticleData.FindMassProps(Manager,DataIdx))
+				if(const FParticleMassProps* NewData = ParticleData.FindMassProps(Manager,DataIdx))
 				{
-					RigidHandle->SetMassProps(*NewData);
+					Evolution.SetParticleMassProps(RigidHandle, *NewData);
 				}
 
-				if(auto NewData = ParticleData.FindDynamics(Manager, DataIdx))
+				if(const FParticleDynamics* NewData = ParticleData.FindDynamics(Manager, DataIdx))
 				{
-					RigidHandle->SetDynamics(*NewData);
-					Evolution.ResetVSmoothFromForces(*RigidHandle);
+					Evolution.SetParticleDynamics(RigidHandle, *NewData);
 				}
 
 				if(NewDynamicMiscData )
diff --git a/Engine/Source/Runtime/Experimental/Chaos/Public/Chaos/PBDRigidsEvolutionGBF.h b/Engine/Source/Runtime/Experimental/Chaos/Public/Chaos/PBDRigidsEvolutionGBF.h
--- a/Engine/Source/Runtime/Experimental/Chaos/Public/Chaos/PBDRigidsEvolutionGBF.h
+++ b/Engine/Source/Runtime/Experimental/Chaos/Public/Chaos/PBDRigidsEvolutionGBF.h
@@ -213,6 +213,15 @@ namespace Chaos
 		 */
 		CHAOS_API void SetParticleVelocities(FGeometryParticleHandle* InParticle, const FVec3& InV, const FVec3f& InW);
 
+		/**
+		 * User has changed particle dynamics
+		 */
+		CHAOS_API void SetParticleDynamics(FPBDRigidParticleHandle* Rigid, const FParticleDynamics& Dynamics);
+
+		/**
+		 * User has changed particle mass properties
+		 */
+		CHAOS_API void SetParticleMassProps(FPBDRigidParticleHandle* Rigid, const FParticleMassProps& MassProperties);
 
 		/**
 		 * Reload the particles cache for all particles where appropriate
@@ -273,9 +282,6 @@ namespace Chaos
 				const FVec3 PredictedAngularVelocity = Particle.GetW() + Particle.AngularAcceleration() * FakeDT + Particle.AngularImpulseVelocity();
 				Particle.WSmooth() = FMath::Lerp(Particle.WSmooth(), PredictedAngularVelocity, SmoothRate);
 			}
-
-			// Flag the particle as potentially updated when calling ResetVSmoothFromForces
-			GetIslandManager().FlagParticleUpdated(Particle.Handle());
 		}
 
 		template<typename TParticleView> 
