Bug fix in teleport unit test: Set particle transform using identity rotation to prevent assert

diff --git a/Engine/Source/Programs/HeadlessChaos/Private/HeadlessChaosTestGraphEvolution.cpp b/Engine/Source/Programs/HeadlessChaos/Private/HeadlessChaosTestGraphEvolution.cpp
--- a/Engine/Source/Programs/HeadlessChaos/Private/HeadlessChaosTestGraphEvolution.cpp
+++ b/Engine/Source/Programs/HeadlessChaos/Private/HeadlessChaosTestGraphEvolution.cpp
@@ -1532,7 +1532,7 @@ namespace ChaosTest
 		const bool bIsTeleport = true;
 		const bool bWakeUp = true;
 		Test.Evolution.SetParticleObjectState(Test.ParticleHandles[4], EObjectStateType::Dynamic);
-		Test.Evolution.SetParticleTransform(Test.ParticleHandles[4], FVec3(0, 3, 0), FRotation3(), bIsTeleport, bWakeUp);
+		Test.Evolution.SetParticleTransform(Test.ParticleHandles[4], FVec3(0, 3, 0), FRotation3::Identity, bIsTeleport, bWakeUp);
 
 		Test.Advance();
 
