Bug fix: Make per-particle IsPartialSleepAlllowed flag work for teleporting and pre-integration wake

Limitation: The IsPartialSleepAlllowed flag needs to be set at least a step before teleporting any particle in the island (dependency on the state of the previous step).

diff --git a/Engine/Source/Programs/HeadlessChaos/Private/HeadlessChaosTestSleeping.cpp b/Engine/Source/Programs/HeadlessChaos/Private/HeadlessChaosTestSleeping.cpp
--- a/Engine/Source/Programs/HeadlessChaos/Private/HeadlessChaosTestSleeping.cpp
+++ b/Engine/Source/Programs/HeadlessChaos/Private/HeadlessChaosTestSleeping.cpp
@@ -762,6 +762,82 @@ namespace ChaosTest
 		CVarMomentumPropagation->Set(MomentumPropagationDefault);
 	}
 
+	TEST_P(SleepingTests, TeleportCubeFromOneStackToAnotherWithPartialSleepDisallowed)
+	{
+		const int32 NumOfStackedObjects = 5;
+		const FReal Dimension = 200;
+		const bool bPartialSleepEnabled = true;
+		const bool bPartialIslandSleepAllowed = GetParam();
+
+		IConsoleVariable* CVarPartialWakePreIntegration = IConsoleManager::Get().FindConsoleVariable(TEXT("p.Chaos.Solver.Sleep.PartialWakePreIntegration"), false);
+		check(CVarPartialWakePreIntegration && CVarPartialWakePreIntegration->IsVariableBool());
+		const bool bPartialWakePreIntegrationDefault = CVarPartialWakePreIntegration->GetBool();
+		CVarPartialWakePreIntegration->Set(true);
+
+		// Create two identical stacks of cubes which don't touch
+		FSleepingTest Test(bPartialSleepEnabled);
+		Test.MakeStackOfCubes(NumOfStackedObjects, /*Size=*/Dimension);
+		Test.MakeStackOfCubes(NumOfStackedObjects, /*Size=*/Dimension, /*Pos=*/FVec3(0, 500, 0));
+
+		// Limitation: This flag needs to be set at least a step before teleporting the particle
+		// since the island flag is only updated in CreateConstraintGraph and already needs to be read in UpdateExplicitSleep.
+		Test.ParticleHandles[4]->SetPartialIslandSleepAllowed(bPartialIslandSleepAllowed);
+
+		// Simulate until all particles are sleeping
+		Test.AdvanceUntilSleeping();
+
+		// Check that everything is sleeping
+		for (FRigidBodyHandle_External* ParticleHandle : Test.ParticleHandles)
+		{
+			EXPECT_TRUE(ParticleHandle->ObjectState() == EObjectStateType::Sleeping);
+		}
+
+		// Teleport one of the cubes out of the stack into the other stack
+		const FVec3 NewX = Test.ParticleHandles[4]->GetX() + FVec3(0, 500, 0);
+		Test.ParticleHandles[4]->SetX(NewX, true);
+
+		Test.Advance();
+
+		if (bPartialIslandSleepAllowed)
+		{
+			// Stack we teleport out of
+			// Only one particle above the one which is teleported will wake up (level-based waking).
+			EXPECT_TRUE(Test.ParticleHandles[0]->ObjectState() == EObjectStateType::Sleeping);
+			EXPECT_TRUE(Test.ParticleHandles[1]->ObjectState() == EObjectStateType::Sleeping);
+			EXPECT_TRUE(Test.ParticleHandles[2]->ObjectState() == EObjectStateType::Sleeping);
+			EXPECT_TRUE(Test.ParticleHandles[3]->ObjectState() == EObjectStateType::Dynamic);
+			EXPECT_TRUE(Test.ParticleHandles[4]->ObjectState() == EObjectStateType::Dynamic);
+
+			// Stack we teleport into
+			// Only one particle above the one which is teleported into the stack will wake up (level-based waking).
+			EXPECT_TRUE(Test.ParticleHandles[5]->ObjectState() == EObjectStateType::Sleeping);
+			EXPECT_TRUE(Test.ParticleHandles[6]->ObjectState() == EObjectStateType::Sleeping);
+			EXPECT_TRUE(Test.ParticleHandles[7]->ObjectState() == EObjectStateType::Sleeping);
+			EXPECT_TRUE(Test.ParticleHandles[8]->ObjectState() == EObjectStateType::Dynamic);
+			EXPECT_TRUE(Test.ParticleHandles[9]->ObjectState() == EObjectStateType::Dynamic);
+		}
+		else
+		{
+			// Stack we teleport out of
+			// Everything wakes up since partial island sleeping is deactivated
+			EXPECT_TRUE(Test.ParticleHandles[0]->ObjectState() == EObjectStateType::Dynamic);
+			EXPECT_TRUE(Test.ParticleHandles[1]->ObjectState() == EObjectStateType::Dynamic);
+			EXPECT_TRUE(Test.ParticleHandles[2]->ObjectState() == EObjectStateType::Dynamic);
+			EXPECT_TRUE(Test.ParticleHandles[3]->ObjectState() == EObjectStateType::Dynamic);
+			EXPECT_TRUE(Test.ParticleHandles[4]->ObjectState() == EObjectStateType::Dynamic);
+
+			// Stack we teleport into
+			// Everything wakes up since partial island sleeping is deactivated
+			EXPECT_TRUE(Test.ParticleHandles[5]->ObjectState() == EObjectStateType::Dynamic);
+			EXPECT_TRUE(Test.ParticleHandles[6]->ObjectState() == EObjectStateType::Dynamic);
+			EXPECT_TRUE(Test.ParticleHandles[7]->ObjectState() == EObjectStateType::Dynamic);
+			EXPECT_TRUE(Test.ParticleHandles[8]->ObjectState() == EObjectStateType::Dynamic);
+			EXPECT_TRUE(Test.ParticleHandles[9]->ObjectState() == EObjectStateType::Dynamic);
+		}
+
+		CVarPartialWakePreIntegration->Set(bPartialWakePreIntegrationDefault);
+	}
+
 	// Uses a traditional brick wall where every 2nd row is shifted (see SleepingTests::MakeBrickWall)
 	TEST_P(SleepingTests, CollideWithBrickWall)
 	{
diff --git a/Engine/Source/Runtime/Experimental/Chaos/Private/Chaos/Island/IslandManager.cpp b/Engine/Source/Runtime/Experimental/Chaos/Private/Chaos/Island/IslandManager.cpp
--- a/Engine/Source/Runtime/Experimental/Chaos/Private/Chaos/Island/IslandManager.cpp
+++ b/Engine/Source/Runtime/Experimental/Chaos/Private/Chaos/Island/IslandManager.cpp
@@ -948,6 +948,9 @@ namespace Chaos::Private
 	{
 		SCOPE_CYCLE_COUNTER(STAT_IslandManager_UpdateParticles);
 
+		// Reset per-island flag before updating graph nodes and repopulating it.
+		ResetIslandPartialSleepAllowed();
+
 		// Process state changed from registered particles.
 		// To reduce cache-misses it would be nice to iterate over transient particle handles 
 		// here but the view we need (Particles.GetActiveDynamicMovingKinematicParticlesView()) 
@@ -1656,6 +1659,14 @@ namespace Chaos::Private
 		}
 	}
 
+	void FPBDIslandManager::ResetIslandPartialSleepAllowed()
+	{
+		for (FPBDIsland* Island : Islands)
+		{
+			Island->Flags.bIsPartialSleepAllowed = true;
+		}
+	}
+
 	FPBDIsland* FPBDIslandManager::MergeNodeIslands(FPBDIslandParticle* Node0, FPBDIslandParticle* Node1)
 	{
 		// NOTE: We only register the islands for merging here. The actual merge takes place later in ProcessMerges.
@@ -2436,7 +2447,6 @@ namespace Chaos::Private
 
 			// This gets set to false again next tick if there are moving kinematics in the island
 			Island->Flags.bIsSleepAllowed = true;
-			Island->Flags.bIsPartialSleepAllowed = true;
 		}
 	}
 
diff --git a/Engine/Source/Runtime/Experimental/Chaos/Public/Chaos/Island/IslandManager.h b/Engine/Source/Runtime/Experimental/Chaos/Public/Chaos/Island/IslandManager.h
--- a/Engine/Source/Runtime/Experimental/Chaos/Public/Chaos/Island/IslandManager.h
+++ b/Engine/Source/Runtime/Experimental/Chaos/Public/Chaos/Island/IslandManager.h
@@ -581,6 +581,7 @@ namespace Chaos::Private
 		CHAOS_API void WakeNodeIslands(const FPBDIslandParticle* Node);
 		CHAOS_API void EnqueueIslandCheckSleep(FPBDIsland* Island, const bool bIsSleepAllowed);
 		CHAOS_API void UpdateIslandPartialSleepAllowed(FPBDIsland* Island, const FGeometryParticleHandle* Particle);
+		CHAOS_API void ResetIslandPartialSleepAllowed();
 
 		// Island merge set management
 		CHAOS_API FPBDIsland* MergeNodeIslands(FPBDIslandParticle* Node0, FPBDIslandParticle* Node1);
