Partial sleep unit tests:
- The unit tests test each sub-feature of partial island sleep in isolation (e.g. momentum propagation, post-step waking).
- I hard coded the "off" default value in the unit test files locally so we can change the defaults directly in IslandManager.cpp without breaking the unit tests.

#FailedValidations Preflight

diff --git a/Engine/Source/Programs/HeadlessChaos/Private/HeadlessChaosTestGraphEvolution.cpp b/Engine/Source/Programs/HeadlessChaos/Private/HeadlessChaosTestGraphEvolution.cpp
--- a/Engine/Source/Programs/HeadlessChaos/Private/HeadlessChaosTestGraphEvolution.cpp
+++ b/Engine/Source/Programs/HeadlessChaos/Private/HeadlessChaosTestGraphEvolution.cpp
@@ -17,6 +17,83 @@ namespace ChaosTest
 
 	class GraphEvolutionTests : public ::testing::TestWithParam<bool>
 	{
+		void SetUp()
+		{
+			SetPartialIslandSleepUnitTestDefaults();
+		}
+
+		void TearDown()
+		{
+			RestorePartialIslandSleepEditorDefaults();
+		}
+
+		void SetPartialIslandSleepUnitTestDefaults()
+		{
+			CVarPostStepWake = IConsoleManager::Get().FindConsoleVariable(TEXT("p.Chaos.Solver.Sleep.PostStepWakeThreshold"), false);
+			check(CVarPostStepWake && CVarPostStepWake->IsVariableFloat());
+			PostStepWake_EditorDefault = CVarPostStepWake->GetFloat();
+			CVarPostStepWake->Set(PostStepWake_UnitTestDefault);
+
+			CVarMomentumPropagation = IConsoleManager::Get().FindConsoleVariable(TEXT("p.Chaos.Solver.Sleep.MomentumPropagation"), false);
+			check(CVarMomentumPropagation && CVarMomentumPropagation->IsVariableFloat());
+			MomentumPropagation_EditorDefault = CVarMomentumPropagation->GetFloat();
+			CVarMomentumPropagation->Set(MomentumPropagation_UnitTestDefault);
+
+			CVarPartialSleepCollisionConstraintsOnly = IConsoleManager::Get().FindConsoleVariable(TEXT("p.Chaos.Solver.Sleep.PartialSleepCollisionConstraintsOnly"), false);
+			check(CVarPartialSleepCollisionConstraintsOnly && CVarPartialSleepCollisionConstraintsOnly->IsVariableBool());
+			bPartialSleepCollisionConstraintsOnly_EditorDefault = CVarPartialSleepCollisionConstraintsOnly->GetBool();
+			CVarPartialSleepCollisionConstraintsOnly->Set(bPartialSleepCollisionConstraintsOnly_UnitTestDefault);
+
+			CVarPartialWakePreIntegration = IConsoleManager::Get().FindConsoleVariable(TEXT("p.Chaos.Solver.Sleep.PartialWakePreIntegration"), false);
+			check(CVarPartialWakePreIntegration && CVarPartialWakePreIntegration->IsVariableBool());
+			bPartialWakePreIntegration_EditorDefault = CVarPartialWakePreIntegration->GetBool();
+			CVarPartialWakePreIntegration->Set(bPartialWakePreIntegration_UnitTestDefault);
+
+			CVarLinearWakeThresholdMultiplier = IConsoleManager::Get().FindConsoleVariable(TEXT("p.Chaos.Solver.Sleep.Defaults.LinearWakeThresholdMultiplier"), false);
+			check(CVarLinearWakeThresholdMultiplier && CVarLinearWakeThresholdMultiplier->IsVariableFloat());
+			LinearWakeThresholdMultiplier_EditorDefault = CVarLinearWakeThresholdMultiplier->GetFloat();
+			CVarLinearWakeThresholdMultiplier->Set(LinearWakeThresholdMultiplier_UnitTestDefault);
+
+			CVarAngularWakeThresholdMultiplier = IConsoleManager::Get().FindConsoleVariable(TEXT("p.Chaos.Solver.Sleep.Defaults.AngularWakeThresholdMultiplier"), false);
+			check(CVarAngularWakeThresholdMultiplier && CVarAngularWakeThresholdMultiplier->IsVariableFloat());
+			AngularWakeThresholdMultiplier_EditorDefault = CVarAngularWakeThresholdMultiplier->GetFloat();
+			CVarAngularWakeThresholdMultiplier->Set(AngularWakeThresholdMultiplier_UnitTestDefault);
+		}
+
+		void RestorePartialIslandSleepEditorDefaults()
+		{
+			CVarPostStepWake->Set(PostStepWake_EditorDefault);
+			CVarMomentumPropagation->Set(MomentumPropagation_EditorDefault);
+			CVarPartialSleepCollisionConstraintsOnly->Set(bPartialSleepCollisionConstraintsOnly_EditorDefault);
+			CVarPartialWakePreIntegration->Set(bPartialWakePreIntegration_EditorDefault);
+			CVarLinearWakeThresholdMultiplier->Set(LinearWakeThresholdMultiplier_EditorDefault);
+			CVarAngularWakeThresholdMultiplier->Set(AngularWakeThresholdMultiplier_EditorDefault);
+		}
+
+		// Default values for unit testing
+		FRealSingle PostStepWake_UnitTestDefault = 0.0f;
+		FRealSingle MomentumPropagation_UnitTestDefault = 0.0f;
+		bool bPartialSleepCollisionConstraintsOnly_UnitTestDefault = false;
+		bool bPartialWakePreIntegration_UnitTestDefault = false;
+		FRealSingle LinearWakeThresholdMultiplier_UnitTestDefault = 1.0f;
+		FRealSingle AngularWakeThresholdMultiplier_UnitTestDefault = 1.0f;
+
+		// Storage for default values editor hard-coded in IslandManager.cpp
+		FRealSingle PostStepWake_EditorDefault;
+		FRealSingle MomentumPropagation_EditorDefault;
+		bool bPartialSleepCollisionConstraintsOnly_EditorDefault;
+		bool bPartialWakePreIntegration_EditorDefault;
+		FRealSingle LinearWakeThresholdMultiplier_EditorDefault;
+		FRealSingle AngularWakeThresholdMultiplier_EditorDefault;
+
+	public:
+		// Storage of console variable pointer
+		IConsoleVariable* CVarPostStepWake;
+		IConsoleVariable* CVarMomentumPropagation;
+		IConsoleVariable* CVarPartialSleepCollisionConstraintsOnly;
+		IConsoleVariable* CVarPartialWakePreIntegration;
+		IConsoleVariable* CVarLinearWakeThresholdMultiplier;
+		IConsoleVariable* CVarAngularWakeThresholdMultiplier;
 	};
 
 	class FGraphEvolutionTest
@@ -1443,7 +1520,7 @@ namespace ChaosTest
 
 
 	// Test the sparse array repeatable index assignment. See TSparseArray::SortFreeList()
-	GTEST_TEST(SparseArrayTests, TestSortFreeList)
+	TEST_P(GraphEvolutionTests, SparseArrayTests_TestSortFreeList)
 	{
 		TSparseArray<int32> Values;
 
@@ -1494,8 +1571,6 @@ namespace ChaosTest
 	// Teleporting means explicitly updating the transform of the particle.
 	TEST_P(GraphEvolutionTests, TestConstraintGraph_TeleportSleeping)
 	{
-		IConsoleVariable* CVarPartialWakePreIntegration = IConsoleManager::Get().FindConsoleVariable(TEXT("p.Chaos.Solver.Sleep.PartialWakePreIntegration"), false);
-		check(CVarPartialWakePreIntegration && CVarPartialWakePreIntegration->IsVariableBool());
 		const bool bPartialWakePreIntegrationDefault = CVarPartialWakePreIntegration->GetBool();
 		CVarPartialWakePreIntegration->Set(true);
 
@@ -1630,8 +1705,6 @@ namespace ChaosTest
 		if (Test.CVarPartialSleeping->GetBool())
 		{
 			// Turn off mommentum propagation to test level-based waking only
-			IConsoleVariable* CVarMomentumPropagation = IConsoleManager::Get().FindConsoleVariable(TEXT("p.Chaos.Solver.Sleep.MomentumPropagation"), false);
-			check(CVarMomentumPropagation && CVarMomentumPropagation->IsVariableFloat());
 			const FRealSingle MomentumPropagationDefault = CVarMomentumPropagation->GetFloat();
 			CVarMomentumPropagation->Set(0);
 
@@ -1707,8 +1780,6 @@ namespace ChaosTest
 		if (Test.CVarPartialSleeping->GetBool())
 		{
 			// Turn off mommentum propagation to test level-based waking only
-			IConsoleVariable* CVarMomentumPropagation = IConsoleManager::Get().FindConsoleVariable(TEXT("p.Chaos.Solver.Sleep.MomentumPropagation"), false);
-			check(CVarMomentumPropagation && CVarMomentumPropagation->IsVariableFloat());
 			const FRealSingle MomentumPropagationDefault = CVarMomentumPropagation->GetFloat();
 			CVarMomentumPropagation->Set(0);
 
@@ -1793,8 +1864,6 @@ namespace ChaosTest
 		if (Test.CVarPartialSleeping->GetBool())
 		{
 			// Turn off mommentum propagation to test level-based waking only
-			IConsoleVariable* CVarMomentumPropagation = IConsoleManager::Get().FindConsoleVariable(TEXT("p.Chaos.Solver.Sleep.MomentumPropagation"), false);
-			check(CVarMomentumPropagation && CVarMomentumPropagation->IsVariableFloat());
 			const FRealSingle MomentumPropagationDefault = CVarMomentumPropagation->GetFloat();
 			CVarMomentumPropagation->Set(0);
 
@@ -1897,8 +1966,6 @@ namespace ChaosTest
 		if (Test.CVarPartialSleeping->GetBool())
 		{
 			// Turn off mommentum propagation to test level-based waking only
-			IConsoleVariable* CVarMomentumPropagation = IConsoleManager::Get().FindConsoleVariable(TEXT("p.Chaos.Solver.Sleep.MomentumPropagation"), false);
-			check(CVarMomentumPropagation && CVarMomentumPropagation->IsVariableFloat());
 			const FRealSingle MomentumPropagationDefault = CVarMomentumPropagation->GetFloat();
 			CVarMomentumPropagation->Set(0);
 
@@ -2008,8 +2075,6 @@ namespace ChaosTest
 	TEST_P(GraphEvolutionTests, TestConstraintGraph_WakeUpPropagation_AddToBlockWithoutSpacingSleeping)
 	{
 		// Turn off mommentum propagation to test level-based waking only
-		IConsoleVariable* CVarMomentumPropagation = IConsoleManager::Get().FindConsoleVariable(TEXT("p.Chaos.Solver.Sleep.MomentumPropagation"), false);
-		check(CVarMomentumPropagation && CVarMomentumPropagation->IsVariableFloat());
 		const FRealSingle MomentumPropagationDefault = CVarMomentumPropagation->GetFloat();
 		CVarMomentumPropagation->Set(0);
 
@@ -2139,8 +2204,6 @@ namespace ChaosTest
 	TEST_P(GraphEvolutionTests, TestConstraintGraph_WakeUpPropagation_AddToBlockWithoutSpacingAwake)
 	{
 		// Turn off mommentum propagation to test level-based waking only
-		IConsoleVariable* CVarMomentumPropagation = IConsoleManager::Get().FindConsoleVariable(TEXT("p.Chaos.Solver.Sleep.MomentumPropagation"), false);
-		check(CVarMomentumPropagation && CVarMomentumPropagation->IsVariableFloat());
 		const FRealSingle MomentumPropagationDefault = CVarMomentumPropagation->GetFloat();
 		CVarMomentumPropagation->Set(0);
 
@@ -2515,8 +2578,6 @@ namespace ChaosTest
 		const bool bEnablePartialSleep = true;
 		const FRealSingle MomentumPropagation = GetParam() ? 0.5f : 0.025f;
 
-		IConsoleVariable* CVarMomentumPropagation = IConsoleManager::Get().FindConsoleVariable(TEXT("p.Chaos.Solver.Sleep.MomentumPropagation"), false);
-		check(CVarMomentumPropagation && CVarMomentumPropagation->IsVariableFloat());
 		const FRealSingle MomentumPropagationDefault = CVarMomentumPropagation->GetFloat();
 		CVarMomentumPropagation->Set(MomentumPropagation);
 
diff --git a/Engine/Source/Programs/HeadlessChaos/Private/HeadlessChaosTestSleeping.cpp b/Engine/Source/Programs/HeadlessChaos/Private/HeadlessChaosTestSleeping.cpp
--- a/Engine/Source/Programs/HeadlessChaos/Private/HeadlessChaosTestSleeping.cpp
+++ b/Engine/Source/Programs/HeadlessChaos/Private/HeadlessChaosTestSleeping.cpp
@@ -14,6 +14,83 @@ namespace ChaosTest
 
 	class SleepingTests : public ::testing::TestWithParam<bool>
 	{
+		void SetUp()
+		{
+			SetPartialIslandSleepUnitTestDefaults();
+		}
+
+		void TearDown()
+		{
+			RestorePartialIslandSleepEditorDefaults();
+		}
+
+		void SetPartialIslandSleepUnitTestDefaults()
+		{
+			CVarPostStepWake = IConsoleManager::Get().FindConsoleVariable(TEXT("p.Chaos.Solver.Sleep.PostStepWakeThreshold"), false);
+			check(CVarPostStepWake && CVarPostStepWake->IsVariableFloat());
+			PostStepWake_EditorDefault = CVarPostStepWake->GetFloat();
+			CVarPostStepWake->Set(PostStepWake_UnitTestDefault);
+
+			CVarMomentumPropagation = IConsoleManager::Get().FindConsoleVariable(TEXT("p.Chaos.Solver.Sleep.MomentumPropagation"), false);
+			check(CVarMomentumPropagation && CVarMomentumPropagation->IsVariableFloat());
+			MomentumPropagation_EditorDefault = CVarMomentumPropagation->GetFloat();
+			CVarMomentumPropagation->Set(MomentumPropagation_UnitTestDefault);
+
+			CVarPartialSleepCollisionConstraintsOnly = IConsoleManager::Get().FindConsoleVariable(TEXT("p.Chaos.Solver.Sleep.PartialSleepCollisionConstraintsOnly"), false);
+			check(CVarPartialSleepCollisionConstraintsOnly && CVarPartialSleepCollisionConstraintsOnly->IsVariableBool());
+			bPartialSleepCollisionConstraintsOnly_EditorDefault = CVarPartialSleepCollisionConstraintsOnly->GetBool();
+			CVarPartialSleepCollisionConstraintsOnly->Set(bPartialSleepCollisionConstraintsOnly_UnitTestDefault);
+
+			CVarPartialWakePreIntegration = IConsoleManager::Get().FindConsoleVariable(TEXT("p.Chaos.Solver.Sleep.PartialWakePreIntegration"), false);
+			check(CVarPartialWakePreIntegration && CVarPartialWakePreIntegration->IsVariableBool());
+			bPartialWakePreIntegration_EditorDefault = CVarPartialWakePreIntegration->GetBool();
+			CVarPartialWakePreIntegration->Set(bPartialWakePreIntegration_UnitTestDefault);
+
+			CVarLinearWakeThresholdMultiplier = IConsoleManager::Get().FindConsoleVariable(TEXT("p.Chaos.Solver.Sleep.Defaults.LinearWakeThresholdMultiplier"), false);
+			check(CVarLinearWakeThresholdMultiplier && CVarLinearWakeThresholdMultiplier->IsVariableFloat());
+			LinearWakeThresholdMultiplier_EditorDefault = CVarLinearWakeThresholdMultiplier->GetFloat();
+			CVarLinearWakeThresholdMultiplier->Set(LinearWakeThresholdMultiplier_UnitTestDefault);
+
+			CVarAngularWakeThresholdMultiplier = IConsoleManager::Get().FindConsoleVariable(TEXT("p.Chaos.Solver.Sleep.Defaults.AngularWakeThresholdMultiplier"), false);
+			check(CVarAngularWakeThresholdMultiplier && CVarAngularWakeThresholdMultiplier->IsVariableFloat());
+			AngularWakeThresholdMultiplier_EditorDefault = CVarAngularWakeThresholdMultiplier->GetFloat();
+			CVarAngularWakeThresholdMultiplier->Set(AngularWakeThresholdMultiplier_UnitTestDefault);
+		}
+
+		void RestorePartialIslandSleepEditorDefaults()
+		{
+			CVarPostStepWake->Set(PostStepWake_EditorDefault);
+			CVarMomentumPropagation->Set(MomentumPropagation_EditorDefault);
+			CVarPartialSleepCollisionConstraintsOnly->Set(bPartialSleepCollisionConstraintsOnly_EditorDefault);
+			CVarPartialWakePreIntegration->Set(bPartialWakePreIntegration_EditorDefault);
+			CVarLinearWakeThresholdMultiplier->Set(LinearWakeThresholdMultiplier_EditorDefault);
+			CVarAngularWakeThresholdMultiplier->Set(AngularWakeThresholdMultiplier_EditorDefault);
+		}
+
+		// Default values for unit testing
+		FRealSingle PostStepWake_UnitTestDefault = 0.0f;
+		FRealSingle MomentumPropagation_UnitTestDefault = 0.0f;
+		bool bPartialSleepCollisionConstraintsOnly_UnitTestDefault = false;
+		bool bPartialWakePreIntegration_UnitTestDefault = false;
+		FRealSingle LinearWakeThresholdMultiplier_UnitTestDefault = 1.0f;
+		FRealSingle AngularWakeThresholdMultiplier_UnitTestDefault = 1.0f;
+
+		// Storage for default values editor hard-coded in IslandManager.cpp
+		FRealSingle PostStepWake_EditorDefault;
+		FRealSingle MomentumPropagation_EditorDefault;
+		bool bPartialSleepCollisionConstraintsOnly_EditorDefault;
+		bool bPartialWakePreIntegration_EditorDefault;
+		FRealSingle LinearWakeThresholdMultiplier_EditorDefault;
+		FRealSingle AngularWakeThresholdMultiplier_EditorDefault;
+
+	public:
+		// Storage of console variable pointer
+		IConsoleVariable* CVarPostStepWake;
+		IConsoleVariable* CVarMomentumPropagation;
+		IConsoleVariable* CVarPartialSleepCollisionConstraintsOnly;
+		IConsoleVariable* CVarPartialWakePreIntegration;
+		IConsoleVariable* CVarLinearWakeThresholdMultiplier;
+		IConsoleVariable* CVarAngularWakeThresholdMultiplier;
 	};
 
 	class FSleepingTest
@@ -361,8 +438,6 @@ namespace ChaosTest
 		const bool bPartialSleepEnabled = GetParam();
 
 		// Turn off mommentum propagation to test level-based waking only
-		IConsoleVariable* CVarMomentumPropagation = IConsoleManager::Get().FindConsoleVariable(TEXT("p.Chaos.Solver.Sleep.MomentumPropagation"), false);
-		check(CVarMomentumPropagation && CVarMomentumPropagation->IsVariableFloat());
 		const FRealSingle MomentumPropagationDefault = CVarMomentumPropagation->GetFloat();
 		CVarMomentumPropagation->Set(0);
 
@@ -417,8 +492,6 @@ namespace ChaosTest
 		const bool bPartialSleepEnabled = GetParam();
 
 		// Turn off mommentum propagation to test level-based waking only
-		IConsoleVariable* CVarMomentumPropagation = IConsoleManager::Get().FindConsoleVariable(TEXT("p.Chaos.Solver.Sleep.MomentumPropagation"), false);
-		check(CVarMomentumPropagation && CVarMomentumPropagation->IsVariableFloat());
 		const FRealSingle MomentumPropagationDefault = CVarMomentumPropagation->GetFloat();
 		CVarMomentumPropagation->Set(0);
 
@@ -471,8 +544,6 @@ namespace ChaosTest
 		const bool bPartialSleepEnabled = GetParam();
 
 		// Turn off mommentum propagation to test level-based waking only
-		IConsoleVariable* CVarMomentumPropagation = IConsoleManager::Get().FindConsoleVariable(TEXT("p.Chaos.Solver.Sleep.MomentumPropagation"), false);
-		check(CVarMomentumPropagation && CVarMomentumPropagation->IsVariableFloat());
 		const FRealSingle MomentumPropagationDefault = CVarMomentumPropagation->GetFloat();
 		CVarMomentumPropagation->Set(0);
 
@@ -522,13 +593,9 @@ namespace ChaosTest
 		const bool bPartialSleepEnabled = GetParam();
 
 		// Turn off momentum propagation to test level-based waking only
-		IConsoleVariable* CVarMomentumPropagation = IConsoleManager::Get().FindConsoleVariable(TEXT("p.Chaos.Solver.Sleep.MomentumPropagation"), false);
-		check(CVarMomentumPropagation && CVarMomentumPropagation->IsVariableFloat());
 		const FRealSingle MomentumPropagationDefault = CVarMomentumPropagation->GetFloat();
 		CVarMomentumPropagation->Set(0);
 
-		IConsoleVariable* CVarPartialWakePreIntegration = IConsoleManager::Get().FindConsoleVariable(TEXT("p.Chaos.Solver.Sleep.PartialWakePreIntegration"), false);
-		check(CVarPartialWakePreIntegration && CVarPartialWakePreIntegration->IsVariableBool());
 		const bool bPartialWakePreIntegrationDefault = CVarPartialWakePreIntegration->GetBool();
 		CVarPartialWakePreIntegration->Set(true);
 
@@ -592,13 +659,9 @@ namespace ChaosTest
 		const bool bPartialSleepEnabled = GetParam();
 
 		// Turn off momentum propagation to test level-based waking only
-		IConsoleVariable* CVarMomentumPropagation = IConsoleManager::Get().FindConsoleVariable(TEXT("p.Chaos.Solver.Sleep.MomentumPropagation"), false);
-		check(CVarMomentumPropagation && CVarMomentumPropagation->IsVariableFloat());
 		const FRealSingle MomentumPropagationDefault = CVarMomentumPropagation->GetFloat();
 		CVarMomentumPropagation->Set(0);
 
-		IConsoleVariable* CVarPartialWakePreIntegration = IConsoleManager::Get().FindConsoleVariable(TEXT("p.Chaos.Solver.Sleep.PartialWakePreIntegration"), false);
-		check(CVarPartialWakePreIntegration && CVarPartialWakePreIntegration->IsVariableBool());
 		const bool bPartialWakePreIntegrationDefault = CVarPartialWakePreIntegration->GetBool();
 		CVarPartialWakePreIntegration->Set(true);
 
@@ -639,8 +702,6 @@ namespace ChaosTest
 		const FReal Dimension = 200;
 		const bool bPartialSleepEnabled = GetParam();
 
-		IConsoleVariable* CVarPartialWakePreIntegration = IConsoleManager::Get().FindConsoleVariable(TEXT("p.Chaos.Solver.Sleep.PartialWakePreIntegration"), false);
-		check(CVarPartialWakePreIntegration && CVarPartialWakePreIntegration->IsVariableBool());
 		const bool bPartialWakePreIntegrationDefault = CVarPartialWakePreIntegration->GetBool();
 		CVarPartialWakePreIntegration->Set(true);
 
@@ -690,13 +751,9 @@ namespace ChaosTest
 		const bool bPartialSleepEnabled = true;
 		const FRealSingle MomentumPropagation = GetParam() ? 0.5f : 0.0f;
 
-		IConsoleVariable* CVarMomentumPropagation = IConsoleManager::Get().FindConsoleVariable(TEXT("p.Chaos.Solver.Sleep.MomentumPropagation"), false);
-		check(CVarMomentumPropagation && CVarMomentumPropagation->IsVariableFloat());
 		const FRealSingle MomentumPropagationDefault = CVarMomentumPropagation->GetFloat();
 		CVarMomentumPropagation->Set(MomentumPropagation);
 
-		IConsoleVariable* CVarPartialWakePreIntegration = IConsoleManager::Get().FindConsoleVariable(TEXT("p.Chaos.Solver.Sleep.PartialWakePreIntegration"), false);
-		check(CVarPartialWakePreIntegration && CVarPartialWakePreIntegration->IsVariableBool());
 		const bool bPartialWakePreIntegrationDefault = CVarPartialWakePreIntegration->GetBool();
 		CVarPartialWakePreIntegration->Set(true);
 
@@ -769,8 +826,6 @@ namespace ChaosTest
 		const bool bPartialSleepEnabled = true;
 		const bool bPartialIslandSleepAllowed = GetParam();
 
-		IConsoleVariable* CVarPartialWakePreIntegration = IConsoleManager::Get().FindConsoleVariable(TEXT("p.Chaos.Solver.Sleep.PartialWakePreIntegration"), false);
-		check(CVarPartialWakePreIntegration && CVarPartialWakePreIntegration->IsVariableBool());
 		const bool bPartialWakePreIntegrationDefault = CVarPartialWakePreIntegration->GetBool();
 		CVarPartialWakePreIntegration->Set(true);
 
@@ -846,8 +901,6 @@ namespace ChaosTest
 		const bool bPartialSleepEnabled = GetParam();
 
 		// Turn off mommentum propagation to test level-based waking only
-		IConsoleVariable* CVarMomentumPropagation = IConsoleManager::Get().FindConsoleVariable(TEXT("p.Chaos.Solver.Sleep.MomentumPropagation"), false);
-		check(CVarMomentumPropagation && CVarMomentumPropagation->IsVariableFloat());
 		const FRealSingle MomentumPropagationDefault = CVarMomentumPropagation->GetFloat();
 		CVarMomentumPropagation->Set(0);
 
@@ -922,8 +975,6 @@ namespace ChaosTest
 		const bool bPartialSleepEnabled = GetParam();
 
 		// Turn off mommentum propagation to test level-based waking only
-		IConsoleVariable* CVarMomentumPropagation = IConsoleManager::Get().FindConsoleVariable(TEXT("p.Chaos.Solver.Sleep.MomentumPropagation"), false);
-		check(CVarMomentumPropagation && CVarMomentumPropagation->IsVariableFloat());
 		const FRealSingle MomentumPropagationDefault = CVarMomentumPropagation->GetFloat();
 		CVarMomentumPropagation->Set(0);
 
@@ -998,8 +1049,6 @@ namespace ChaosTest
 		const bool bAllowPartialSleep = GetParam();
 
 		// Turn off mommentum propagation to test level-based waking only
-		IConsoleVariable* CVarMomentumPropagation = IConsoleManager::Get().FindConsoleVariable(TEXT("p.Chaos.Solver.Sleep.MomentumPropagation"), false);
-		check(CVarMomentumPropagation && CVarMomentumPropagation->IsVariableFloat());
 		const FRealSingle MomentumPropagationDefault = CVarMomentumPropagation->GetFloat();
 		CVarMomentumPropagation->Set(0);
 
@@ -1059,8 +1108,6 @@ namespace ChaosTest
 		const FReal Dimension = 200;
 		const bool bPartialSleepEnabled = true;
 		const FRealSingle MomentumPropagation = GetParam() ? 0.5f : 0.1f;
-		IConsoleVariable* CVarMomentumPropagation = IConsoleManager::Get().FindConsoleVariable(TEXT("p.Chaos.Solver.Sleep.MomentumPropagation"), false);
-		check(CVarMomentumPropagation && CVarMomentumPropagation->IsVariableFloat());
 		const FRealSingle MomentumPropagationDefault = CVarMomentumPropagation->GetFloat();
 		CVarMomentumPropagation->Set(MomentumPropagation);
 		// Create a stack of cubes
@@ -1113,8 +1160,6 @@ namespace ChaosTest
 		const FVec3 Dimensions = FVec3(100, 200, 50);
 		const bool bPartialSleepEnabled = true;
 		const FRealSingle MomentumPropagation = GetParam() ? 0.5f : 0.1f;
-		IConsoleVariable* CVarMomentumPropagation = IConsoleManager::Get().FindConsoleVariable(TEXT("p.Chaos.Solver.Sleep.MomentumPropagation"), false);
-		check(CVarMomentumPropagation && CVarMomentumPropagation->IsVariableFloat());
 		const FRealSingle MomentumPropagationDefault = CVarMomentumPropagation->GetFloat();
 		CVarMomentumPropagation->Set(MomentumPropagation);
 		// Create a stack of cubes
@@ -1181,10 +1226,8 @@ namespace ChaosTest
 		const bool bPartialSleepEnabled = true;
 		const FRealSingle PostStepWakeThreshold = GetParam() ? 10.0f : 0.0f;
 
-		IConsoleVariable* CVarPostStepWakeThreshold = IConsoleManager::Get().FindConsoleVariable(TEXT("p.Chaos.Solver.Sleep.PostStepWakeThreshold"), false);
-		check(CVarPostStepWakeThreshold && CVarPostStepWakeThreshold->IsVariableFloat());
-		const FRealSingle PostStepWakeThresholdDefault = CVarPostStepWakeThreshold->GetFloat();
-		CVarPostStepWakeThreshold->Set(PostStepWakeThreshold);
+		const FRealSingle PostStepWakeThresholdDefault = CVarPostStepWake->GetFloat();
+		CVarPostStepWake->Set(PostStepWakeThreshold);
 
 		// Create a stack of cubes
 		FSleepingTest Test(bPartialSleepEnabled);
@@ -1235,7 +1278,7 @@ namespace ChaosTest
 			EXPECT_TRUE(Test.ParticleHandles[4]->ObjectState() == EObjectStateType::Sleeping);
 		}
 
-		CVarPostStepWakeThreshold->Set(PostStepWakeThresholdDefault);
+		CVarPostStepWake->Set(PostStepWakeThresholdDefault);
 	}
 
 	// Sleep all particles except one which are connected by joint constraints.
@@ -1243,8 +1286,6 @@ namespace ChaosTest
 	TEST_P(SleepingTests, PartiallySleepingIslandWithJoints)
 	{
 		const bool bPartialSleepCollisionConstraintsOnly = GetParam();
-		IConsoleVariable* CVarPartialSleepCollisionConstraintsOnly = IConsoleManager::Get().FindConsoleVariable(TEXT("p.Chaos.Solver.Sleep.PartialSleepCollisionConstraintsOnly"), false);
-		check(CVarPartialSleepCollisionConstraintsOnly && CVarPartialSleepCollisionConstraintsOnly->IsVariableBool());
 		const bool bPartialSleepCollisionConstraintsOnlyDefault = CVarPartialSleepCollisionConstraintsOnly->GetBool();
 		CVarPartialSleepCollisionConstraintsOnly->Set(bPartialSleepCollisionConstraintsOnly);
 
@@ -1286,8 +1327,6 @@ namespace ChaosTest
 	TEST_P(SleepingTests, PartiallySleepingIslandWithJoints_2)
 	{
 		const bool bPartialSleepCollisionConstraintsOnly = GetParam();
-		IConsoleVariable* CVarPartialSleepCollisionConstraintsOnly = IConsoleManager::Get().FindConsoleVariable(TEXT("p.Chaos.Solver.Sleep.PartialSleepCollisionConstraintsOnly"), false);
-		check(CVarPartialSleepCollisionConstraintsOnly && CVarPartialSleepCollisionConstraintsOnly->IsVariableBool());
 		const bool bPartialSleepCollisionConstraintsOnlyDefault = CVarPartialSleepCollisionConstraintsOnly->GetBool();
 		CVarPartialSleepCollisionConstraintsOnly->Set(bPartialSleepCollisionConstraintsOnly);
 
@@ -1330,8 +1369,6 @@ namespace ChaosTest
 	TEST_P(SleepingTests, FullySleepingIslandWithJoints)
 	{
 		const bool bPartialSleepCollisionConstraintsOnly = GetParam();
-		IConsoleVariable* CVarPartialSleepCollisionConstraintsOnly = IConsoleManager::Get().FindConsoleVariable(TEXT("p.Chaos.Solver.Sleep.PartialSleepCollisionConstraintsOnly"), false);
-		check(CVarPartialSleepCollisionConstraintsOnly && CVarPartialSleepCollisionConstraintsOnly->IsVariableBool());
 		const bool bPartialSleepCollisionConstraintsOnlyDefault = CVarPartialSleepCollisionConstraintsOnly->GetBool();
 		CVarPartialSleepCollisionConstraintsOnly->Set(bPartialSleepCollisionConstraintsOnly);
 
