Fix BugHawk detected issue:
- Partial island sleep now remains on if MinimumMotionlessRatio == 0 and all particles moving
- Partial island sleep now remains on if MinimumMotionlessRatio == 1 and all particles motionless

diff --git a/Engine/Source/Programs/HeadlessChaos/Private/HeadlessChaosTestGraphEvolution.cpp b/Engine/Source/Programs/HeadlessChaos/Private/HeadlessChaosTestGraphEvolution.cpp
--- a/Engine/Source/Programs/HeadlessChaos/Private/HeadlessChaosTestGraphEvolution.cpp
+++ b/Engine/Source/Programs/HeadlessChaos/Private/HeadlessChaosTestGraphEvolution.cpp
@@ -2642,7 +2642,7 @@ namespace ChaosTest
 		CVarMomentumPropagation->Set(MomentumPropagationDefault);
 	}
 
-	TEST_P(GraphEvolutionTests, TestConstraintGraph_SwitchBewtweenFullAndPartialIslandSleep)
+	TEST_P(GraphEvolutionTests, TestConstraintGraph_MinimumMotionlessRatio)
 	{
 		const bool bEnablePartialSleep = true;
 		const FRealSingle MinMotionlessRatio = GetParam() ? 0.5f : 0.1f;
@@ -2740,5 +2740,63 @@ namespace ChaosTest
 		CVarMinMotionlessRatio->Set(MinMotionlessRatioDefault);
 	}
 
+	TEST_P(GraphEvolutionTests, TestConstraintGraph_MinimumMotionlessRatio_Extremes)
+	{
+		const bool bEnablePartialSleep = true;
+		// MinMotionlessRatio == 0.0f => 0% of particles need to be motionless for partial island sleep to be active
+		// MinMotionlessRatio == 1.0f => 100% of particles need to be motionless for partial island sleep to be active
+		const FRealSingle MinMotionlessRatio = GetParam() ? 1.0f : 0.0f;
+
+		const FRealSingle MinMotionlessRatioDefault = CVarMinMotionlessRatio->GetFloat();
+		CVarMinMotionlessRatio->Set(MinMotionlessRatio);
+
+		const IConsoleVariable* CVarSleepModeSwitchCounterThreshold = IConsoleManager::Get().FindConsoleVariable(TEXT("p.Chaos.Solver.Sleep.PartialIslandSleep.SleepModeSwitchCounterThreshold"), false);
+		check(CVarSleepModeSwitchCounterThreshold && CVarSleepModeSwitchCounterThreshold->IsVariableInt());
+		const int32 SleepModeSwitchCounterThreshold = CVarSleepModeSwitchCounterThreshold->GetInt();
+		const int32 Steps = SleepModeSwitchCounterThreshold + 1;
+
+		FGraphEvolutionTest Test(5, bEnablePartialSleep);
+		Test.MakeChain();
+
+		FParticleDynamics Dynamics;
+		Dynamics.SetAcceleration(FVec3(0, 0, 0));
+		Dynamics.SetAngularAcceleration(FVec3(0, 0, 0));
+		if (MinMotionlessRatio == 0.0f)
+		{
+			Dynamics.SetLinearImpulseVelocity(FVec3(0, 0.1, 0));
+		}
+		else // MinMotionlessRatio == 1.0f
+		{
+			Dynamics.SetLinearImpulseVelocity(FVec3(0, 0, 0));
+		}
+		Dynamics.SetAngularImpulseVelocity(FVec3(0, 0, 0));
+
+		for (FPBDRigidParticleHandle* Particle : Test.ParticleHandles)
+		{
+			Test.Evolution.SetParticleObjectState(Particle, EObjectStateType::Dynamic);
+			Test.Evolution.SetParticleDynamics(Particle, Dynamics);
+		}
+
+		for (int32 I = 0; I < Steps; ++I)
+		{
+			Test.Advance();
+
+			EXPECT_TRUE(Test.IslandManager->GetIsland(0)->IsPartialSleepEnabled());
+		}
+
+		if (MinMotionlessRatio == 0.0f)
+		{
+			// Partial island sleeping remains enabled since all particles are in motion
+			EXPECT_TRUE(Test.IslandManager->GetIsland(0)->IsPartialSleepEnabled());
+		}
+		else // MinMotionlessRatio == 1.0f
+		{
+			// Partial island sleeping remains enabled since all particles are in motionless
+			EXPECT_TRUE(Test.IslandManager->GetIsland(0)->IsPartialSleepEnabled());
+		}
+
+		CVarMinMotionlessRatio->Set(MinMotionlessRatioDefault);
+	}
+
 	INSTANTIATE_TEST_SUITE_P(, GraphEvolutionTests,	testing::Bool());
 }
diff --git a/Engine/Source/Runtime/Experimental/Chaos/Private/Chaos/Island/IslandManager.cpp b/Engine/Source/Runtime/Experimental/Chaos/Private/Chaos/Island/IslandManager.cpp
--- a/Engine/Source/Runtime/Experimental/Chaos/Private/Chaos/Island/IslandManager.cpp
+++ b/Engine/Source/Runtime/Experimental/Chaos/Private/Chaos/Island/IslandManager.cpp
@@ -92,7 +92,7 @@ namespace Chaos::CVars
 
 	/** Cvar to set the minimum ratio of particles which needs to be motionless to activate partial island sleep. */
 	FRealSingle ChaosSolverMinMotionlessRatio = 0.1f;
-	FAutoConsoleVariableRef CVarChaosSolverMinMotionlessRatio(TEXT("p.Chaos.Solver.Sleep.PartialIslandSleep.MinimumMotionlessRatio"), ChaosSolverMinMotionlessRatio, TEXT("[EXPERIMENTAL] Ratio of the particles which need to be motionless in an island for partial island sleep to become active.[def:0.1]"));
+	FAutoConsoleVariableRef CVarChaosSolverMinMotionlessRatio(TEXT("p.Chaos.Solver.Sleep.PartialIslandSleep.MinimumMotionlessRatio"), ChaosSolverMinMotionlessRatio, TEXT("[EXPERIMENTAL] Ratio of the particles which need to be motionless in an island for partial island sleep to become active (0 to always activate partial island sleep).[def:0.1]"));
 
 	/** Cvar to set the minimum number of ticks the ratio of motionless particles needs to be exceeded before switching between partial and full island sleep. */
 	int32 ChaosSolverSleepModeSwitchThreshold = 10;
@@ -3052,11 +3052,11 @@ namespace Chaos::Private
 		{
 			const FReal RatioMotionless = MotionlessCount / (FReal)Island->GetNumParticles();
 
-			const bool bIsAboveMinMotionlessRatio = RatioMotionless > CVars::ChaosSolverMinMotionlessRatio;
+			const bool bIsBelowMinMotionlessRatio = RatioMotionless < CVars::ChaosSolverMinMotionlessRatio;
 			const bool bExceedSleepCounterSwitchThreshold = Island->SleepTypeSwitchCounter >= CVars::ChaosSolverSleepModeSwitchThreshold;
 
-			const bool bPotentialSwitchToPartialIslandSleep = !Island->Flags.bEnoughMotionlessNodesForPartialSleep && bIsAboveMinMotionlessRatio;
-			const bool bPotentialSwitchToFullIslandSleep = !!Island->Flags.bEnoughMotionlessNodesForPartialSleep && !bIsAboveMinMotionlessRatio;
+			const bool bPotentialSwitchToPartialIslandSleep = !Island->Flags.bEnoughMotionlessNodesForPartialSleep && !bIsBelowMinMotionlessRatio;
+			const bool bPotentialSwitchToFullIslandSleep = !!Island->Flags.bEnoughMotionlessNodesForPartialSleep && bIsBelowMinMotionlessRatio;
 
 			// Switch sleep type: flip the switch and reset the counter
 			const bool bSwitchSleepType = (bPotentialSwitchToPartialIslandSleep || bPotentialSwitchToFullIslandSleep) && bExceedSleepCounterSwitchThreshold;
