Bug fix: Support teleporting in partial island sleep with kinematic obstacles
- Kinematic obstacles lead to levels that do not correspond with horizontal height
- The fix is to update VSmooth with with the velocity equivalent of the distance the particle was teleported and wake particles through momentum propagation

diff --git a/Engine/Source/Programs/HeadlessChaos/Private/HeadlessChaosTestSleeping.cpp b/Engine/Source/Programs/HeadlessChaos/Private/HeadlessChaosTestSleeping.cpp
--- a/Engine/Source/Programs/HeadlessChaos/Private/HeadlessChaosTestSleeping.cpp
+++ b/Engine/Source/Programs/HeadlessChaos/Private/HeadlessChaosTestSleeping.cpp
@@ -786,14 +786,13 @@ namespace ChaosTest
 
 		Test.Advance();
 
-		if (MomentumPropagation)
+		if (MomentumPropagation > 0.0f)
 		{
 			// Stack we teleport out of
-			// Only one particle above the one which is teleported will wake up (level-based waking).
-			// Momentum transfer does not affect the stack we teleport out of.
-			EXPECT_TRUE(Test.ParticleHandles[0]->ObjectState() == EObjectStateType::Sleeping);
-			EXPECT_TRUE(Test.ParticleHandles[1]->ObjectState() == EObjectStateType::Sleeping);
-			EXPECT_TRUE(Test.ParticleHandles[2]->ObjectState() == EObjectStateType::Sleeping);
+			// Everything wakes up if momentum transfer is enabled
+			EXPECT_TRUE(Test.ParticleHandles[0]->ObjectState() == EObjectStateType::Dynamic);
+			EXPECT_TRUE(Test.ParticleHandles[1]->ObjectState() == EObjectStateType::Dynamic);
+			EXPECT_TRUE(Test.ParticleHandles[2]->ObjectState() == EObjectStateType::Dynamic);
 			EXPECT_TRUE(Test.ParticleHandles[3]->ObjectState() == EObjectStateType::Dynamic);
 			EXPECT_TRUE(Test.ParticleHandles[4]->ObjectState() == EObjectStateType::Dynamic);
 
@@ -805,7 +804,7 @@ namespace ChaosTest
 			EXPECT_TRUE(Test.ParticleHandles[8]->ObjectState() == EObjectStateType::Dynamic);
 			EXPECT_TRUE(Test.ParticleHandles[9]->ObjectState() == EObjectStateType::Dynamic);
 		}
-		else
+		else // MomentumPropagation == 0.0f
 		{
 			// Stack we teleport out of
 			// Only one particle above the one which is teleported will wake up (level-based waking).
@@ -902,6 +901,88 @@ namespace ChaosTest
 		CVarPartialWakePreIntegration->Set(bPartialWakePreIntegrationDefault);
 	}
 
+	TEST_P(SleepingTests, TeleportCubeFromOneStackToAnotherWithKinematicObstacle)
+	{
+		const int32 NumOfStackedObjects = 6;
+		const FReal Dimension = 200;
+		const bool bPartialSleepEnabled = true;
+		const FRealSingle MomentumPropagation = GetParam() ? 0.5f : 0.0f;
+
+		const FRealSingle MomentumPropagationDefault = CVarMomentumPropagation->GetFloat();
+		CVarMomentumPropagation->Set(MomentumPropagation);
+
+		const bool bPartialWakePreIntegrationDefault = CVarPartialWakePreIntegration->GetBool();
+		CVarPartialWakePreIntegration->Set(true);
+
+		// Create two identical stacks of cubes which don't touch
+		FSleepingTest Test(bPartialSleepEnabled);
+		Test.MakeStackOfCubes(NumOfStackedObjects, /*Size=*/Dimension);
+		Test.MakeStackOfCubes(NumOfStackedObjects, /*Size=*/Dimension, /*Pos=*/FVec3(0, 500, 0));
+
+		// Make the top of the static kinematic for levels not to correspond to horizontal height
+		Test.ParticleHandles[5]->SetObjectState(EObjectStateType::Static);
+		Test.ParticleHandles[11]->SetObjectState(EObjectStateType::Static);
+
+		// Simulate until all particles are sleeping
+		Test.AdvanceUntilSleeping();
+
+		//// Check that everything is sleeping
+		//for (FRigidBodyHandle_External* ParticleHandle : Test.ParticleHandles)
+		//{
+		//	EXPECT_TRUE(ParticleHandle->ObjectState() == EObjectStateType::Sleeping);
+		//}
+
+		// Teleport one of the cubes out of the stack into the other stack
+		const FVec3 NewX = Test.ParticleHandles[2]->GetX() + FVec3(0, 500, 400);
+		Test.ParticleHandles[2]->SetX(NewX, true);
+
+		Test.Advance();
+
+		if (MomentumPropagation > 0.0f)
+		{
+			// Stack we teleport out of
+			// Everything wakes up if momentum transfer is enabled
+			EXPECT_TRUE(Test.ParticleHandles[0]->ObjectState() == EObjectStateType::Dynamic);
+			EXPECT_TRUE(Test.ParticleHandles[1]->ObjectState() == EObjectStateType::Dynamic);
+			EXPECT_TRUE(Test.ParticleHandles[2]->ObjectState() == EObjectStateType::Dynamic);
+			EXPECT_TRUE(Test.ParticleHandles[3]->ObjectState() == EObjectStateType::Dynamic);
+			EXPECT_TRUE(Test.ParticleHandles[4]->ObjectState() == EObjectStateType::Dynamic);
+			EXPECT_TRUE(Test.ParticleHandles[5]->ObjectState() == EObjectStateType::Static);
+
+			// Stack we teleport into
+			// Everything wakes up if momentum transfer is enabled
+			EXPECT_TRUE(Test.ParticleHandles[6]->ObjectState() == EObjectStateType::Dynamic);
+			EXPECT_TRUE(Test.ParticleHandles[7]->ObjectState() == EObjectStateType::Dynamic);
+			EXPECT_TRUE(Test.ParticleHandles[8]->ObjectState() == EObjectStateType::Dynamic);
+			EXPECT_TRUE(Test.ParticleHandles[9]->ObjectState() == EObjectStateType::Dynamic);
+			EXPECT_TRUE(Test.ParticleHandles[10]->ObjectState() == EObjectStateType::Dynamic);
+			EXPECT_TRUE(Test.ParticleHandles[11]->ObjectState() == EObjectStateType::Static);
+		}
+		else // MomentumPropagation == 0.0f
+		{
+			// Stack we teleport out of
+			// Only one particle above the one which is teleported will wake up (level-based waking).
+			EXPECT_TRUE(Test.ParticleHandles[0]->ObjectState() == EObjectStateType::Sleeping);
+			EXPECT_TRUE(Test.ParticleHandles[1]->ObjectState() == EObjectStateType::Dynamic);
+			EXPECT_TRUE(Test.ParticleHandles[2]->ObjectState() == EObjectStateType::Dynamic);
+			EXPECT_TRUE(Test.ParticleHandles[3]->ObjectState() == EObjectStateType::Dynamic);
+			EXPECT_TRUE(Test.ParticleHandles[4]->ObjectState() == EObjectStateType::Sleeping);
+			EXPECT_TRUE(Test.ParticleHandles[5]->ObjectState() == EObjectStateType::Static);
+
+			// Stack we teleport into
+			// Only one particle above the one which is teleported into the stack will wake up (level-based waking).
+			EXPECT_TRUE(Test.ParticleHandles[6]->ObjectState() == EObjectStateType::Sleeping);
+			EXPECT_TRUE(Test.ParticleHandles[7]->ObjectState() == EObjectStateType::Sleeping);
+			EXPECT_TRUE(Test.ParticleHandles[8]->ObjectState() == EObjectStateType::Dynamic);
+			EXPECT_TRUE(Test.ParticleHandles[9]->ObjectState() == EObjectStateType::Dynamic);
+			EXPECT_TRUE(Test.ParticleHandles[10]->ObjectState() == EObjectStateType::Dynamic);
+			EXPECT_TRUE(Test.ParticleHandles[11]->ObjectState() == EObjectStateType::Static);
+		}
+
+		CVarPartialWakePreIntegration->Set(bPartialWakePreIntegrationDefault);
+		CVarMomentumPropagation->Set(MomentumPropagationDefault);
+	}
+
 	// Uses a traditional brick wall where every 2nd row is shifted (see SleepingTests::MakeBrickWall)
 	TEST_P(SleepingTests, CollideWithBrickWall)
 	{
diff --git a/Engine/Source/Runtime/Experimental/Chaos/Private/Chaos/Island/IslandManager.cpp b/Engine/Source/Runtime/Experimental/Chaos/Private/Chaos/Island/IslandManager.cpp
--- a/Engine/Source/Runtime/Experimental/Chaos/Private/Chaos/Island/IslandManager.cpp
+++ b/Engine/Source/Runtime/Experimental/Chaos/Private/Chaos/Island/IslandManager.cpp
@@ -910,7 +910,11 @@ namespace Chaos::Private
 		{
 			if (FPBDIslandParticle* Node = GetGraphNode(Particle))
 			{
-				Node->Flags.bWasUpdated = true;
+				FPBDIsland* Island = Node->Island;
+				if (Island && Island->IsPartialSleepEnabled())
+				{
+					Node->Flags.bWasUpdated = true;
+				}
 			}
 		}
 	}
@@ -921,7 +925,11 @@ namespace Chaos::Private
 		{
 			if (FPBDIslandParticle* Node = GetGraphNode(Particle))
 			{
-				return Node->Flags.bWasTransformSetExplicitly;
+				FPBDIsland* Island = Node->Island;
+				if (Island && Island->IsPartialSleepEnabled())
+				{
+					return Node->Flags.bWasTransformSetExplicitly;
+				}
 			}
 		}
 		return false;
@@ -933,7 +941,11 @@ namespace Chaos::Private
 		{
 			if (FPBDIslandParticle* Node = GetGraphNode(Particle))
 			{
-				Node->Flags.bWasTransformSetExplicitly = true;
+				FPBDIsland* Island = Node->Island;
+				if (Island && Island->IsPartialSleepEnabled())
+				{
+					Node->Flags.bWasTransformSetExplicitly = true;
+				}
 			}
 		}
 	}
diff --git a/Engine/Source/Runtime/Experimental/Chaos/Private/Chaos/PBDRigidsEvolutionGBF.cpp b/Engine/Source/Runtime/Experimental/Chaos/Private/Chaos/PBDRigidsEvolutionGBF.cpp
--- a/Engine/Source/Runtime/Experimental/Chaos/Private/Chaos/PBDRigidsEvolutionGBF.cpp
+++ b/Engine/Source/Runtime/Experimental/Chaos/Private/Chaos/PBDRigidsEvolutionGBF.cpp
@@ -1828,7 +1828,9 @@ void FPBDRigidsEvolutionGBF::OnParticleMoved(FGeometryParticleHandle* InParticle
 	// If this is a teleport, we assume other state has been set separately (V, W, etc) if required, which 
 	// would include resetting any sleep-related state. If this is not a teleport, we want to prevent 
 	// sleeping if the user keeps moving the body without changing the velocity.
-	if (!bIsTeleport)
+	// If the particle transform was set explicitly (contains check for partial island sleep), 
+	// we also want to manually set VSmooth to make use of momentum propagation in FPBDIslandManager::UpdateExplicitSleep.
+	if (!bIsTeleport || GetIslandManager().WasParticleTransformSetExplicitly(InParticle))
 	{
 		if (FPBDRigidParticleHandle* Rigid = InParticle->CastToRigidParticle())
 		{
