PR #13878: Fix Intersects comparing squared with non squared in Sphere.h

diff --git a/Engine/Source/Programs/HeadlessChaos/Private/HeadlessChaosTestImplicits.cpp b/Engine/Source/Programs/HeadlessChaos/Private/HeadlessChaosTestImplicits.cpp
--- a/Engine/Source/Programs/HeadlessChaos/Private/HeadlessChaosTestImplicits.cpp
+++ b/Engine/Source/Programs/HeadlessChaos/Private/HeadlessChaosTestImplicits.cpp
@@ -570,6 +570,31 @@ namespace ChaosTest {
 			EXPECT_VECTOR_NEAR(Subject2.FindClosestPoint(InputPoint), FVector(2, 2, 4), 0.001);
 			EXPECT_VECTOR_NEAR(Subject2.FindClosestPoint(FVec3(2, 2, 3.5)), FVector(2, 2, 4), 0.001);
 		}
+
+		// Intersects
+		{
+			// Overlap
+			{
+				Chaos::FSphere Sphere0(FVec3(0), 2);
+				Chaos::FSphere Sphere1(FVec3(4, 0, 0), 3);
+				EXPECT_TRUE(Sphere0.Intersects(Sphere1));
+				EXPECT_TRUE(Sphere1.Intersects(Sphere0));
+			}
+			// Tangent
+			{
+				Chaos::FSphere Sphere0(FVec3(0), 2);
+				Chaos::FSphere Sphere1(FVec3(5, 0, 0), 3);
+				EXPECT_TRUE(Sphere0.Intersects(Sphere1));
+				EXPECT_TRUE(Sphere1.Intersects(Sphere0));
+			}
+			// Separating
+			{
+				Chaos::FSphere Sphere0(FVec3(0), 2);
+				Chaos::FSphere Sphere1(FVec3(6, 0, 0), 3);
+				EXPECT_FALSE(Sphere0.Intersects(Sphere1));
+				EXPECT_FALSE(Sphere1.Intersects(Sphere0));
+			}
+		}
 	}
 	
 	/* Cylinder Helpers */
diff --git a/Engine/Source/Runtime/Experimental/Chaos/Public/Chaos/Sphere.h b/Engine/Source/Runtime/Experimental/Chaos/Public/Chaos/Sphere.h
--- a/Engine/Source/Runtime/Experimental/Chaos/Public/Chaos/Sphere.h
+++ b/Engine/Source/Runtime/Experimental/Chaos/Public/Chaos/Sphere.h
@@ -89,9 +89,9 @@ namespace Chaos
 
 		bool Intersects(const TSphere<T, d>& Other) const
 		{
-			FRealSingle CenterDistance = FVec3f::DistSquared(Other.GetCenterf(), GetCenterf());
+			FRealSingle CenterDistanceSquared = FVec3f::DistSquared(Other.GetCenterf(), GetCenterf());
 			FRealSingle RadialSum = Other.GetRadiusf() + GetRadiusf();
-			return RadialSum >= CenterDistance;
+			return RadialSum * RadialSum >= CenterDistanceSquared;
 		}
 
 		TVector<T, d> FindClosestPoint(const TVector<T, d>& StartPoint, const T Thickness = (T)0) const
