Chaos - add PT raycast to match the PT particle sweep functionality (for internal use only)

#rnx

diff --git a/Engine/Source/Programs/HeadlessChaos/Private/HeadlessChaosTestSimSweep.cpp b/Engine/Source/Programs/HeadlessChaos/Private/HeadlessChaosTestSimSweep.cpp
--- a/Engine/Source/Programs/HeadlessChaos/Private/HeadlessChaosTestSimSweep.cpp
+++ b/Engine/Source/Programs/HeadlessChaos/Private/HeadlessChaosTestSimSweep.cpp
@@ -2,13 +2,16 @@
 
 #include "HeadlessChaos.h"
 #include "HeadlessChaosTestEvolution.h"
+#include "Chaos/Collision/SimRaycast.h"
 #include "Chaos/Collision/SimSweep.h"
 
+PRAGMA_DISABLE_INTERNAL_WARNINGS
+
 namespace ChaosTest
 {
 	using namespace Chaos;
 
-	GTEST_TEST(SimOverlapTests, TestBoxBoxOverlap)
+	GTEST_TEST(SimSweepTests, TestBoxBoxOverlap)
 	{
 		const FVec3 FloorSize = FVec3(10000, 10000, 100);
 		const FVec3 BoxPos = FVec3(0, 0, 200);
@@ -34,4 +37,254 @@ namespace ChaosTest
 			EXPECT_TRUE((Overlaps[0].HitParticle == Box) || (Overlaps[1].HitParticle == Box));
 		}
 	}
-}
\ No newline at end of file
+
+	GTEST_TEST(SimSweepTests, TestBoxBoxSweep)
+	{
+		const FVec3 FloorSize = FVec3(10000, 10000, 100);
+		const FVec3 BoxPos = FVec3(0, 0, 500);
+		const FRotation3 BoxRot = FRotation3::FromIdentity();
+		const FVec3 BoxSize = FVec3(100);
+		const FReal BoxMass = 100;
+
+		FEvolutionTest Test;
+
+		FGeometryParticleHandle* Floor = Test.AddParticleFloor(FloorSize);
+		FPBDRigidParticleHandle* Box = Test.AddParticleBox(BoxPos, BoxRot, BoxSize, BoxMass);
+
+		Test.Tick();
+
+		FVec3 SweepStartPos = FVec3(0,0, 500);
+		FVec3 SweepDir = FVec3(0, 0, -1);
+		FRotation3 SweepRot = FRotation3::FromIdentity();
+		FIgnoreCollisionManager* IgnoreCollisionManager = nullptr;
+
+		// 100cm length sweep from Z=500 should miss
+		Private::FSimSweepParticleHit Hit100;
+		bool bIsHit100 = Private::SimSweepParticleFirstHit(
+			Test.GetEvolution().GetSpatialAcceleration(), IgnoreCollisionManager, 
+			Box, 
+			SweepStartPos, SweepRot, SweepDir, 100, 
+			Hit100);
+		EXPECT_FALSE(bIsHit100);
+		EXPECT_FALSE(Hit100.IsHit());
+
+		// 1000cm downward sweep from Z=500 should hit with Z=50, for a sweep distance of D=450
+		Private::FSimSweepParticleHit Hit1000;
+		bool bIsHit1000 = Private::SimSweepParticleFirstHit(
+			Test.GetEvolution().GetSpatialAcceleration(), IgnoreCollisionManager,
+			Box,
+			SweepStartPos,SweepRot, SweepDir, 1000,
+			Hit1000);
+		EXPECT_TRUE(bIsHit1000);
+		EXPECT_TRUE(Hit1000.IsHit());
+		if (bIsHit1000)
+		{
+			EXPECT_NEAR(Hit1000.HitDistance, 450.0f, UE_KINDA_SMALL_NUMBER);
+		}
+	}
+
+	// Raycast against the floor (a cube)
+	GTEST_TEST(SimSweepTests, TestBoxRaycastSimple)
+	{
+		const FVec3 FloorSize = FVec3(10000, 10000, 100);
+
+		FEvolutionTest Test;
+
+		FGeometryParticleHandle* Floor = Test.AddParticleFloor(FloorSize);
+		Test.GetEvolution().SetParticleTransform(Floor, FVec3(0,0,100), FRotation3::FromIdentity(), false, false);
+
+		Test.Tick();
+
+		Private::FSimRaycastRay Ray;
+		Ray.StartPos = FVec3f(0,0,500);
+		Ray.Dir = FVec3f(0,0,-1);
+		Ray.Length = 1000;
+		
+		Private::FSimRaycastHit Hit;
+
+		Filter::FQueryFilterData QueryFilter = Filter::FQueryObjectFilterBuilder().SetObjectTypes(0xFFFFFFFF).Build();
+
+		bool bAnyHits = Private::SimRaycastBatchFirstHits(
+			Test.GetEvolution().GetSpatialAcceleration(),
+			MakeArrayView<Private::FSimRaycastRay>(&Ray, 1),
+			QueryFilter,
+			/*Simple*/true, /*Complex*/false,
+			MakeArrayView<Private::FSimRaycastHit>(&Hit, 1));
+
+		EXPECT_TRUE(bAnyHits);
+		EXPECT_TRUE(Hit.IsHit());
+		EXPECT_NEAR(Hit.HitDistance, 350.0f, UE_SMALL_NUMBER);
+		EXPECT_NEAR(Hit.HitPosition.Z, 150.0f, UE_SMALL_NUMBER);
+	}
+
+	// Raycast against a cube and the floor, with filtering to ignore one or the other
+	GTEST_TEST(SimSweepTests, TestBoxRaycastSimpleFiltered)
+	{
+		const FVec3 FloorSize = FVec3(10000, 10000, 100);
+		const FVec3 BoxPos = FVec3(0, 0, 500);
+		const FRotation3 BoxRot = FRotation3::FromIdentity();
+		const FVec3 BoxSize = FVec3(100);
+		const FReal BoxMass = 100;
+
+		FEvolutionTest Test;
+
+		FGeometryParticleHandle* Floor = Test.AddParticleFloor(FloorSize);
+		FPBDRigidParticleHandle* Box = Test.AddParticleBox(BoxPos, BoxRot, BoxSize, BoxMass);
+		Test.GetEvolution().SetParticleObjectState(Box, EObjectStateType::Kinematic);
+
+		// Floor -> CollisionChannel 0
+		Filter::FShapeFilterData FloorFilterData = Filter::FShapeFilterBuilder().SetCollisionChannelIndex(0).SetBlockChannelMasks(0xFFFFFFFF, 0).Build();
+		Floor->ShapeInstances()[0]->SetShapeFilterData(FloorFilterData);
+
+		// Box -> CollisionChannel 1
+		Filter::FShapeFilterData BoxFilterData = Filter::FShapeFilterBuilder().SetCollisionChannelIndex(1).SetBlockChannelMasks(0xFFFFFFFF, 0).Build();
+		Box->ShapeInstances()[0]->SetShapeFilterData(BoxFilterData);
+
+		Test.Tick();
+
+		Private::FSimRaycastRay Rays[2];
+		// Downward Ray
+		Rays[0].StartPos = FVec3f(0, 0, 1000);
+		Rays[0].Dir = FVec3f(0, 0, -1);
+		Rays[0].Length = 2000;
+		// Upward Ray
+		Rays[1].StartPos = FVec3f(0, 0, -1000);
+		Rays[1].Dir = FVec3f(0, 0, 1);
+		Rays[1].Length = 2000;
+
+		// Detect Collision Channel 0
+		{
+			Private::FSimRaycastHit Hits[2];
+
+			Filter::FQueryFilterData QueryFilter = Filter::FQueryFilterBuilder::CreateLegacyObjectTypeFilter(1 << 0, false);
+
+			bool bAnyHits = Private::SimRaycastBatchFirstHits(
+				Test.GetEvolution().GetSpatialAcceleration(),
+				MakeArrayView<Private::FSimRaycastRay>(Rays, UE_ARRAY_COUNT(Rays)),
+				QueryFilter,
+				/*Simple*/true, /*Complex*/false,
+				MakeArrayView<Private::FSimRaycastHit>(Hits, UE_ARRAY_COUNT(Hits)));
+
+			// Downward Ray - Floor (filtered Box)
+			EXPECT_TRUE(Hits[0].IsHit());
+			EXPECT_EQ(Hits[0].HitParticle, Floor);
+
+			// Upward Ray - Floor (First encountered)
+			EXPECT_TRUE(Hits[1].IsHit());
+			EXPECT_EQ(Hits[1].HitParticle, Floor);
+		}
+
+		// Detect Collision Channel 1
+		{
+			Private::FSimRaycastHit Hits[2];
+
+			Filter::FQueryFilterData QueryFilter = Filter::FQueryFilterBuilder::CreateLegacyObjectTypeFilter(1 << 1, false);
+
+			bool bAnyHits = Private::SimRaycastBatchFirstHits(
+				Test.GetEvolution().GetSpatialAcceleration(),
+				MakeArrayView<Private::FSimRaycastRay>(Rays, UE_ARRAY_COUNT(Rays)),
+				QueryFilter,
+				/*Simple*/true, /*Complex*/false,
+				MakeArrayView<Private::FSimRaycastHit>(Hits, UE_ARRAY_COUNT(Hits)));
+
+			// Downward Ray - Box (First encountered)
+			EXPECT_TRUE(Hits[0].IsHit());
+			EXPECT_EQ(Hits[0].HitParticle, Box);
+
+			// Upward Ray - Box (filtered Floor)
+			EXPECT_TRUE(Hits[1].IsHit());
+			EXPECT_EQ(Hits[1].HitParticle, Box);
+		}
+	}
+
+	// Raycast against complex (in a scene with no complex shapes)
+	GTEST_TEST(SimSweepTests, TestBoxRaycastComplex)
+	{
+		const FVec3 FloorSize = FVec3(10000, 10000, 100);
+
+		FEvolutionTest Test;
+
+		FGeometryParticleHandle* Floor = Test.AddParticleFloor(FloorSize);
+		Test.GetEvolution().SetParticleTransform(Floor, FVec3(0, 0, 100), FRotation3::FromIdentity(), false, false);
+
+		Test.Tick();
+
+		Private::FSimRaycastRay Ray;
+		Ray.StartPos = FVec3f(0, 0, 500);
+		Ray.Dir = FVec3f(0, 0, -1);
+		Ray.Length = 1000;
+
+		Private::FSimRaycastHit Hit;
+
+		Filter::FQueryFilterData QueryFilter = Filter::FQueryFilterBuilder::CreateLegacyObjectTypeFilter(0xFFFFFFFF, false);
+
+		bool bAnyHits = Private::SimRaycastBatchFirstHits(
+			Test.GetEvolution().GetSpatialAcceleration(),
+			MakeArrayView<Private::FSimRaycastRay>(&Ray, 1),
+			QueryFilter,
+			/*Simple*/false, /*Complex*/true,
+			MakeArrayView<Private::FSimRaycastHit>(&Hit, 1));
+
+		EXPECT_FALSE(bAnyHits);
+		EXPECT_FALSE(Hit.IsHit());
+	}
+
+	// Multiple raycasts as a batch
+	GTEST_TEST(SimSweepTests, TestRaycastFirst)
+	{
+		const FVec3 FloorSize = FVec3(10000, 10000, 100);
+		const FVec3 FloorPos = FVec3(0, 0, 100);
+		const FVec3 BoxPos = FVec3(0, 0, 500);
+		const FRotation3 BoxRot = FRotation3::FromIdentity();
+		const FVec3 BoxSize = FVec3(100);
+		const FReal BoxMass = 100;
+
+		FEvolutionTest Test;
+
+		FPBDRigidParticleHandle* Box = Test.AddParticleBox(BoxPos, BoxRot, BoxSize, BoxMass);
+		Test.GetEvolution().SetParticleObjectState(Box, EObjectStateType::Kinematic);
+
+		FGeometryParticleHandle* Floor = Test.AddParticleFloor(FloorSize);
+		Test.GetEvolution().SetParticleTransform(Floor, FVec3(0, 0, 100), FRotation3::FromIdentity(), false, false);
+
+		Test.Tick();
+
+		Private::FSimRaycastRay Rays[2];
+		// Downward Ray
+		Rays[0].StartPos = FVec3f(0, 0, 1000);
+		Rays[0].Dir = FVec3f(0, 0, -1);
+		Rays[0].Length = 2000;
+		// Upward Ray
+		Rays[1].StartPos = FVec3f(0, 0, -1000);
+		Rays[1].Dir = FVec3f(0, 0, 1);
+		Rays[1].Length = 2000;
+
+		Private::FSimRaycastHit Hits[2];
+
+		Filter::FQueryFilterData QueryFilter = Filter::FQueryFilterBuilder::CreateLegacyObjectTypeFilter(0xFFFFFFFF, false);
+
+		bool bAnyHits = Private::SimRaycastBatchFirstHits(
+			Test.GetEvolution().GetSpatialAcceleration(),
+			MakeArrayView<Private::FSimRaycastRay>(Rays, UE_ARRAY_COUNT(Rays)),
+			QueryFilter,
+			/*Simple*/true, /*Complex*/false,
+			MakeArrayView<Private::FSimRaycastHit>(Hits, UE_ARRAY_COUNT(Hits)));
+
+		EXPECT_TRUE(bAnyHits);
+
+		// Downward Ray
+		EXPECT_TRUE(Hits[0].IsHit());
+		EXPECT_EQ(Hits[0].HitParticle, Box);
+		EXPECT_NEAR(Hits[0].HitDistance, 450.0f, UE_SMALL_NUMBER);
+		EXPECT_NEAR(Hits[0].HitPosition.Z, 550.0f, UE_SMALL_NUMBER);
+
+		// Upward Ray
+		EXPECT_TRUE(Hits[1].IsHit());
+		EXPECT_EQ(Hits[1].HitParticle, Floor);
+		EXPECT_NEAR(Hits[1].HitDistance, 1050.0f, UE_SMALL_NUMBER);
+		EXPECT_NEAR(Hits[1].HitPosition.Z, 50.0f, UE_SMALL_NUMBER);
+	}
+
+}
+
+PRAGMA_ENABLE_INTERNAL_WARNINGS
diff --git a/Engine/Source/Runtime/Experimental/Chaos/Private/Chaos/Collision/SimRaycast.cpp b/Engine/Source/Runtime/Experimental/Chaos/Private/Chaos/Collision/SimRaycast.cpp
new file mode 100644
--- /dev/null
+++ b/Engine/Source/Runtime/Experimental/Chaos/Private/Chaos/Collision/SimRaycast.cpp
@@ -0,0 +1,171 @@
+// Copyright Epic Games, Inc. All Rights Reserved.
+#include "Chaos/Collision/SimRaycast.h"
+#include "Chaos/CollisionFilterData.h"
+#include "Chaos/ParticleHandle.h"
+
+namespace Chaos::Private
+{
+	bool SimRaycastImplicit(const FSimRaycastRay& InRay, const FImplicitObject* InImplicit, const FRigidTransform3& InImplicitTransform, FSimRaycastHit& InOutHit)
+	{
+		// @todo(chaos): could optimize uniform scale rather than just unit scale...
+		const bool bIsImplicitUnitScale = FVec3::IsNearlyEqual(InImplicitTransform.GetScale3D(), FVec3(1, 1, 1), UE_SMALL_NUMBER);
+
+		// Transform the raycast into the space of the implicit
+		FVec3 LocalStartPos;
+		FVec3 LocalDir;
+		FReal LocalLength;
+		if (bIsImplicitUnitScale)
+		{
+			LocalStartPos = InImplicitTransform.InverseTransformPositionNoScale(InRay.StartPos);
+			LocalDir = InImplicitTransform.InverseTransformVectorNoScale(InRay.Dir);
+			LocalLength = InRay.Length;
+		}
+		else
+		{
+			const FVec3 LocalDelta = InImplicitTransform.InverseTransformVector(InRay.Dir * InRay.Length);
+			LocalStartPos = InImplicitTransform.InverseTransformPosition(InRay.StartPos);
+			LocalDir = FVec3::Zero();
+			LocalLength = LocalDelta.Size();
+			if (LocalLength > 0)
+			{
+				LocalDir = LocalDelta / LocalLength;
+			}
+		}
+
+		FVec3 LocalHitPos, LocalHitNormal;
+		FReal LocalHitDistance;
+		int32 HitFaceIndex;
+		if (InImplicit->Raycast(LocalStartPos, LocalDir, LocalLength, 0.0, LocalHitDistance, LocalHitPos, LocalHitNormal, HitFaceIndex))
+		{
+			// If we have an existing hit, ignore hits after it
+			const FReal HitTOI = (LocalLength > 0) ? LocalHitDistance / LocalLength : 0;
+			const bool bIsCloserHit = !InOutHit.IsHit() || (HitTOI < InOutHit.HitTOI);
+			if (bIsCloserHit)
+			{
+				// Save output in world space
+				if (bIsImplicitUnitScale)
+				{
+					InOutHit.HitPosition = InImplicitTransform.TransformPositionNoScale(LocalHitPos);
+					InOutHit.HitNormal = InImplicitTransform.TransformVectorNoScale(LocalHitNormal);
+					InOutHit.HitDistance = LocalHitDistance;
+				}
+				else
+				{
+					InOutHit.HitPosition = InImplicitTransform.TransformPosition(LocalHitPos);
+					InOutHit.HitNormal = InImplicitTransform.TransformNormal(LocalHitNormal);
+					InOutHit.HitDistance = HitTOI * InRay.Length;
+				}
+
+				// Hit normal and face normal always the same for raycasts.
+				InOutHit.HitFaceNormal = InOutHit.HitNormal;
+
+				// NOTE: Particle and Shape are unknown at this level, and must be filled in by the caller.
+				InOutHit.HitParticle = nullptr;
+				InOutHit.HitShape = nullptr;
+				InOutHit.HitGeometry = InImplicit;
+				InOutHit.HitFaceIndex = HitFaceIndex;
+				InOutHit.HitTOI = HitTOI;
+
+				return true;
+			}
+		}
+
+		return false;
+	}
+
+	FSimQueryDataParticleFilter::FSimQueryDataParticleFilter(const Filter::FQueryFilterData& InFilterData)
+		: FilterData(Filter::FQueryFilterBuilder::GetLegacyQueryFilter(InFilterData))
+	{
+	}
+
+	bool FSimQueryDataParticleFilter::ApplyFilter(const FGeometryParticleHandle* InParticleHandle) const
+	{
+		return true;
+	}
+
+	FSimQueryDataShapeFilter::FSimQueryDataShapeFilter(const Filter::FQueryFilterData& InFilterData, bool bInDetectSimple, bool bInDetectComplex)
+		: FilterData(Filter::FQueryFilterBuilder::GetLegacyQueryFilter(InFilterData))
+	{
+		Flags.bDetectSimple = bInDetectSimple;
+		Flags.bDetectComplex = bInDetectComplex;
+	}
+
+	bool FSimQueryDataShapeFilter::ApplyFilter(const FPerShapeData* Shape, const FImplicitObject* Implicit) const
+	{
+		// Run the shape Simple/Complex filter
+		const EImplicitObjectType ImplicitType = GetInnerType(Implicit->GetType());
+		const bool bIsMesh = (ImplicitType == ImplicitObjectType::TriangleMesh);
+		const bool bIsComplex = bIsMesh || (Shape->GetCollisionTraceType() == EChaosCollisionTraceFlag::Chaos_CTF_UseSimpleAsComplex);
+		const bool bIsSimple = !bIsMesh || (Shape->GetCollisionTraceType() == EChaosCollisionTraceFlag::Chaos_CTF_UseComplexAsSimple);
+		const bool bCanDetect = ((bIsComplex && !!Flags.bDetectComplex) || (bIsSimple && !!Flags.bDetectSimple));
+		if (!bCanDetect)
+		{
+			return false;
+		}
+
+		// Run the channel filter
+		// NOTE: PrePreSimFilterImp return true if we want to skip
+		return !PrePreSimFilterImp(FilterData, Shape->GetSimData());
+	}
+
+	bool SimRaycastBatchFirstHits(
+		ISpatialAcceleration<FAccelerationStructureHandle, FReal, 3>* SpatialAcceleration,
+		const TArrayView<const Private::FSimRaycastRay>& InRays,
+		const Filter::FQueryFilterData& QueryFilterData,
+		const bool bDetectSimple, const bool bDetectComplex,
+		const TArrayView<Private::FSimRaycastHit>& OutHits)
+	{
+		if (!ensureMsgf(InRays.Num() == OutHits.Num(), TEXT("Mismatched input and output arrays")))
+		{
+			return 0;
+		}
+
+		bool bAnyHits = false;
+
+		// Use default query filtering (but against shape SimFilter, not their QueryFilter)
+		FSimQueryDataParticleFilter ParticleFilter = FSimQueryDataParticleFilter(QueryFilterData);
+		FSimQueryDataShapeFilter ShapeFilter = FSimQueryDataShapeFilter(QueryFilterData, bDetectSimple, bDetectComplex);
+
+		// Raycast against the overlapping particle shape
+		const auto& RaycastParticleShape = 
+			[&InRays, &OutHits, &bAnyHits](const FSimOverlapParticleShape& Overlap) -> void
+			{
+				const FRigidTransform3 ParticleTransform = FConstGenericParticleHandle(Overlap.HitParticle)->GetTransformPQ();
+
+				const auto& RaycastLeafImplicit = 
+					[&Overlap, &ParticleTransform, &InRays, &OutHits, &bAnyHits](const FImplicitObject* Implicit, const FRigidTransform3& ImplicitTransform, const int32 RootObjectIndex, const int32 ObjectIndex, const int32 LeafObjectIndex) -> void
+					{
+						// Calculate the world-space transform of the implicit object
+						const FRigidTransform3 ParticleImplicitTransform = ImplicitTransform * ParticleTransform;
+
+						// Cast all the rays against the ImplicitObject
+						for (int32 RayIndex = 0; RayIndex < InRays.Num(); ++RayIndex)
+						{
+							if (SimRaycastImplicit(InRays[RayIndex], Implicit, ParticleImplicitTransform, OutHits[RayIndex]))
+							{
+								OutHits[RayIndex].HitParticle = Overlap.HitParticle;
+								OutHits[RayIndex].HitShape = Overlap.HitShape;
+								bAnyHits = true;
+							}
+						}
+					};
+
+				// Find all overlapping shapes and call the raycast lambda on them
+				Overlap.HitShape->GetGeometry()->VisitLeafObjects(RaycastLeafImplicit);
+			};
+
+		// Calculate the query bounds (from all rays)
+		FAABB3 QueryBounds = FAABB3::EmptyAABB();
+		for (const FSimRaycastRay& Ray : InRays)
+		{
+			QueryBounds.GrowToInclude(Ray.StartPos);
+			QueryBounds.GrowToInclude(Ray.StartPos + Ray.Length * Ray.Dir);
+		}
+
+		// Visit all shapes in the Spatial Acceleration that overlap any of the rays and run the raycasts
+		SimOverlapBounds(SpatialAcceleration, QueryBounds, ParticleFilter, ShapeFilter, RaycastParticleShape);
+
+		return bAnyHits;
+	}
+
+} // namespace Chaos::Private
\ No newline at end of file
diff --git a/Engine/Source/Runtime/Experimental/Chaos/Public/Chaos/Collision/SimRaycast.h b/Engine/Source/Runtime/Experimental/Chaos/Public/Chaos/Collision/SimRaycast.h
new file mode 100644
--- /dev/null
+++ b/Engine/Source/Runtime/Experimental/Chaos/Public/Chaos/Collision/SimRaycast.h
@@ -0,0 +1,112 @@
+// Copyright Epic Games, Inc. All Rights Reserved.
+#pragma once
+
+#include "Chaos/ParticleHandleFwd.h"
+#include "Chaos/Collision/SimSweep.h"
+
+namespace Chaos::Filter
+{
+	struct FQueryFilterData;
+}
+
+namespace Chaos::Private
+{
+	/**
+	* Raycast input
+	*/
+	class FSimRaycastRay
+	{
+	public:
+		FVec3 StartPos;
+		FVec3 Dir;
+		FReal Length;
+	};
+
+	/**
+	* Raycast output
+	*/
+	using FSimRaycastHit = FSimSweepParticleHit;
+
+	/*
+	* Standard particle filter for use with Sim Raycasts that uses FQueryFilterData to specify what to exclude.
+	* @todo(chaos): Currently accepts all.
+	*/
+	class FSimQueryDataParticleFilter
+	{
+	public:
+		CHAOS_API FSimQueryDataParticleFilter(const Filter::FQueryFilterData& InFilterData);
+
+		bool operator()(const FGeometryParticleHandle* InParticleHandle) const
+		{
+			return ApplyFilter(InParticleHandle);
+		}
+
+	private:
+		CHAOS_API bool ApplyFilter(const FGeometryParticleHandle* InParticleHandle) const;
+
+		FCollisionFilterData FilterData;
+	};
+
+	/*
+	* Standard shape filter for use with Sim Raycasts that uses FQueryFilterData to specify what to exclude.
+	*/
+	class FSimQueryDataShapeFilter
+	{
+	public:
+		CHAOS_API FSimQueryDataShapeFilter(const Filter::FQueryFilterData& InFilterData, bool bInDetectSimple, bool bInDetectComplex);
+
+		bool operator()(const FPerShapeData* InShape, const FImplicitObject* InImplicit) const
+		{
+			return ApplyFilter(InShape, InImplicit);
+		}
+
+	private:
+		CHAOS_API bool ApplyFilter(const FPerShapeData* InShape, const FImplicitObject* InImplicit) const;
+
+		union FFlags
+		{
+			FFlags() : Bits(0) {}
+			struct
+			{
+				uint8 bDetectSimple : 1;
+				uint8 bDetectComplex : 1;
+			};
+			uint8 Bits;
+		};
+
+		FCollisionFilterData FilterData;
+		FFlags Flags;
+	};
+
+	/**
+	* Cast a ray against an implicit object with the specified transform.
+	* @param InRay The ray info
+	* @param InImplicit The implicit object to raycast against
+	* @param InImplicitTransform The world space transform of the implicit object
+	* @param InOutHit The raycast output. If InOutHit alreday contains a valid hit, then it will be left unchanged if the new hit is not before it
+	* @return true if the ray hits the implicit object, false otherwise
+	*/
+	extern CHAOS_API bool SimRaycastImplicit(const FSimRaycastRay& InRay, const FImplicitObject* InImplicit, const FRigidTransform3& InImplicitTransform, FSimRaycastHit& InOutHit);
+
+	/**
+	* Cast a set of rays. This is intended to be used as an example - see the implementation for
+	* how to write your own raycasts. This version only visits the SpatialAcceleration once
+	* using the aggregate bounds of all of the rays in InRays. It is well suited to running
+	* multiple closely located short raycasts, but will likely underperform for long off-axis raycasts
+	* because of the way it collects objects in the raycast bounds.
+	* 
+	* @param InRays The batch of rays to cast. They should be fairly close together with a "reasonable" bounds for perf reasons.
+	* @param QueryFilterData The filter to reject shapes.
+	* @param bDetectSimple Whether to detect simple collision shapes.
+	* @param bDetectComplex Whether to detect complex collision shapes.
+	* @param OutHits The results of the raycasts. Must contain one and only one entry per input ray.
+	* @return true is any of the rays hit something, false otherwise.
+	*/
+	extern CHAOS_API bool SimRaycastBatchFirstHits(
+		ISpatialAcceleration<FAccelerationStructureHandle, FReal, 3>* SpatialAcceleration,
+		const TArrayView<const Private::FSimRaycastRay>& InRays,
+		const Filter::FQueryFilterData& QueryFilterData,
+		const bool bDetectSimple, const bool bDetectComplex,
+		const TArrayView<Private::FSimRaycastHit>& OutHits);
+
+}
