Partial island sleep - Allow 3 values for CVar to prevent partial sleep of particles connected by non-collision constraints
- 0: partial island sleep allowed between all types of constraints
- -1: always propagate partial wake events through non-collision constraints, i.e. the entire ragdoll wakes up if one limb is woken up.
- 1: ensure that a non-collision constraint can never connect a sleeping and an awake particle, i.e. the entire ragdoll remains awake if one limb is awake.

diff --git a/Engine/Source/Programs/HeadlessChaos/Private/HeadlessChaosTestGraphEvolution.cpp b/Engine/Source/Programs/HeadlessChaos/Private/HeadlessChaosTestGraphEvolution.cpp
--- a/Engine/Source/Programs/HeadlessChaos/Private/HeadlessChaosTestGraphEvolution.cpp
+++ b/Engine/Source/Programs/HeadlessChaos/Private/HeadlessChaosTestGraphEvolution.cpp
@@ -40,9 +40,9 @@ namespace ChaosTest
 			CVarMomentumPropagation->Set(MomentumPropagation_UnitTestDefault);
 
 			CVarPartialSleepCollisionConstraintsOnly = IConsoleManager::Get().FindConsoleVariable(TEXT("p.Chaos.Solver.Sleep.PartialIslandSleep.CollisionConstraintsOnly"), false);
-			check(CVarPartialSleepCollisionConstraintsOnly && CVarPartialSleepCollisionConstraintsOnly->IsVariableBool());
-			bPartialSleepCollisionConstraintsOnly_EditorDefault = CVarPartialSleepCollisionConstraintsOnly->GetBool();
-			CVarPartialSleepCollisionConstraintsOnly->Set(bPartialSleepCollisionConstraintsOnly_UnitTestDefault);
+			check(CVarPartialSleepCollisionConstraintsOnly && CVarPartialSleepCollisionConstraintsOnly->IsVariableInt());
+			PartialSleepCollisionConstraintsOnly_EditorDefault = CVarPartialSleepCollisionConstraintsOnly->GetInt();
+			CVarPartialSleepCollisionConstraintsOnly->Set(PartialSleepCollisionConstraintsOnly_UnitTestDefault);
 
 			CVarPartialWakePreIntegration = IConsoleManager::Get().FindConsoleVariable(TEXT("p.Chaos.Solver.Sleep.PartialIslandSleep.PartialWakePreIntegration"), false);
 			check(CVarPartialWakePreIntegration && CVarPartialWakePreIntegration->IsVariableBool());
@@ -74,7 +74,7 @@ namespace ChaosTest
 		{
 			CVarPostStepWake->Set(PostStepWake_EditorDefault);
 			CVarMomentumPropagation->Set(MomentumPropagation_EditorDefault);
-			CVarPartialSleepCollisionConstraintsOnly->Set(bPartialSleepCollisionConstraintsOnly_EditorDefault);
+			CVarPartialSleepCollisionConstraintsOnly->Set(PartialSleepCollisionConstraintsOnly_EditorDefault);
 			CVarPartialWakePreIntegration->Set(bPartialWakePreIntegration_EditorDefault);
 			CVarMinMotionlessRatio->Set(MinimumMotionlessThreshold_EditorDefault);
 			CVarLinearWakeThresholdMultiplier->Set(LinearWakeThresholdMultiplier_EditorDefault);
@@ -85,7 +85,7 @@ namespace ChaosTest
 		// Default values for unit testing
 		FRealSingle PostStepWake_UnitTestDefault = 0.0f;
 		FRealSingle MomentumPropagation_UnitTestDefault = 0.0f;
-		bool bPartialSleepCollisionConstraintsOnly_UnitTestDefault = false;
+		int32 PartialSleepCollisionConstraintsOnly_UnitTestDefault = 0;
 		bool bPartialWakePreIntegration_UnitTestDefault = false;
 		FRealSingle MinimumMotionlessThreshold_UnitTestDefault = 0.0f;
 		FRealSingle LinearWakeThresholdMultiplier_UnitTestDefault = 1.0f;
@@ -95,7 +95,7 @@ namespace ChaosTest
 		// Storage for default values editor hard-coded in IslandManager.cpp
 		FRealSingle PostStepWake_EditorDefault;
 		FRealSingle MomentumPropagation_EditorDefault;
-		bool bPartialSleepCollisionConstraintsOnly_EditorDefault;
+		int32 PartialSleepCollisionConstraintsOnly_EditorDefault;
 		bool bPartialWakePreIntegration_EditorDefault;
 		FRealSingle MinimumMotionlessThreshold_EditorDefault;
 		FRealSingle LinearWakeThresholdMultiplier_EditorDefault;
diff --git a/Engine/Source/Programs/HeadlessChaos/Private/HeadlessChaosTestSleeping.cpp b/Engine/Source/Programs/HeadlessChaos/Private/HeadlessChaosTestSleeping.cpp
--- a/Engine/Source/Programs/HeadlessChaos/Private/HeadlessChaosTestSleeping.cpp
+++ b/Engine/Source/Programs/HeadlessChaos/Private/HeadlessChaosTestSleeping.cpp
@@ -38,9 +38,9 @@ namespace ChaosTest
 			CVarMomentumPropagation->Set(MomentumPropagation_UnitTestDefault);
 
 			CVarPartialSleepCollisionConstraintsOnly = IConsoleManager::Get().FindConsoleVariable(TEXT("p.Chaos.Solver.Sleep.PartialIslandSleep.CollisionConstraintsOnly"), false);
-			check(CVarPartialSleepCollisionConstraintsOnly && CVarPartialSleepCollisionConstraintsOnly->IsVariableBool());
-			bPartialSleepCollisionConstraintsOnly_EditorDefault = CVarPartialSleepCollisionConstraintsOnly->GetBool();
-			CVarPartialSleepCollisionConstraintsOnly->Set(bPartialSleepCollisionConstraintsOnly_UnitTestDefault);
+			check(CVarPartialSleepCollisionConstraintsOnly && CVarPartialSleepCollisionConstraintsOnly->IsVariableInt());
+			PartialSleepCollisionConstraintsOnly_EditorDefault = CVarPartialSleepCollisionConstraintsOnly->GetInt();
+			CVarPartialSleepCollisionConstraintsOnly->Set(PartialSleepCollisionConstraintsOnly_UnitTestDefault);
 
 			CVarPartialWakePreIntegration = IConsoleManager::Get().FindConsoleVariable(TEXT("p.Chaos.Solver.Sleep.PartialIslandSleep.PartialWakePreIntegration"), false);
 			check(CVarPartialWakePreIntegration && CVarPartialWakePreIntegration->IsVariableBool());
@@ -72,7 +72,7 @@ namespace ChaosTest
 		{
 			CVarPostStepWake->Set(PostStepWake_EditorDefault);
 			CVarMomentumPropagation->Set(MomentumPropagation_EditorDefault);
-			CVarPartialSleepCollisionConstraintsOnly->Set(bPartialSleepCollisionConstraintsOnly_EditorDefault);
+			CVarPartialSleepCollisionConstraintsOnly->Set(PartialSleepCollisionConstraintsOnly_EditorDefault);
 			CVarPartialWakePreIntegration->Set(bPartialWakePreIntegration_EditorDefault);
 			CVarMinMotionlessRatio->Set(MinimumMotionlessThreshold_EditorDefault);
 			CVarLinearWakeThresholdMultiplier->Set(LinearWakeThresholdMultiplier_EditorDefault);
@@ -83,7 +83,7 @@ namespace ChaosTest
 		// Default values for unit testing
 		FRealSingle PostStepWake_UnitTestDefault = 0.0f;
 		FRealSingle MomentumPropagation_UnitTestDefault = 0.0f;
-		bool bPartialSleepCollisionConstraintsOnly_UnitTestDefault = false;
+		int32 PartialSleepCollisionConstraintsOnly_UnitTestDefault = 0;
 		bool bPartialWakePreIntegration_UnitTestDefault = false;
 		FRealSingle MinimumMotionlessThreshold_UnitTestDefault = 0.0f;
 		FRealSingle LinearWakeThresholdMultiplier_UnitTestDefault = 1.0f;
@@ -93,7 +93,7 @@ namespace ChaosTest
 		// Storage for default values editor hard-coded in IslandManager.cpp
 		FRealSingle PostStepWake_EditorDefault;
 		FRealSingle MomentumPropagation_EditorDefault;
-		bool bPartialSleepCollisionConstraintsOnly_EditorDefault;
+		int32 PartialSleepCollisionConstraintsOnly_EditorDefault;
 		bool bPartialWakePreIntegration_EditorDefault;
 		FRealSingle MinimumMotionlessThreshold_EditorDefault;
 		FRealSingle LinearWakeThresholdMultiplier_EditorDefault;
@@ -1385,9 +1385,9 @@ namespace ChaosTest
 	// Since joint constraints don't support partial sleeping, the entire island wakes up.
 	TEST_P(SleepingTests, PartiallySleepingIslandWithJoints)
 	{
-		const bool bPartialSleepCollisionConstraintsOnly = GetParam();
-		const bool bPartialSleepCollisionConstraintsOnlyDefault = CVarPartialSleepCollisionConstraintsOnly->GetBool();
-		CVarPartialSleepCollisionConstraintsOnly->Set(bPartialSleepCollisionConstraintsOnly);
+		const int32 PartialSleepCollisionConstraintsOnly = GetParam() ? 1 : 0;
+		const int32 PartialSleepCollisionConstraintsOnlyDefault = CVarPartialSleepCollisionConstraintsOnly->GetInt();
+		CVarPartialSleepCollisionConstraintsOnly->Set(PartialSleepCollisionConstraintsOnly);
 
 		const bool bPartialSleepEnabled = true;
 		FSleepingTest Test(bPartialSleepEnabled);
@@ -1401,7 +1401,7 @@ namespace ChaosTest
 
 		Test.Advance();
 
-		if (bPartialSleepCollisionConstraintsOnly)
+		if (PartialSleepCollisionConstraintsOnly == 1)
 		{
 			// All particles wake up since the connecting constraints don't support partial sleeping
 			for (int32 Id = 1; Id < Test.ParticleHandles.Num(); ++Id)
@@ -1410,7 +1410,7 @@ namespace ChaosTest
 				EXPECT_TRUE(ParticleHandle->ObjectState() == EObjectStateType::Dynamic);
 			}
 		}
-		else
+		else if (PartialSleepCollisionConstraintsOnly == 0)
 		{
 			// Partial sleep supported, no state change
 			EXPECT_TRUE(Test.ParticleHandles[0]->ObjectState() == EObjectStateType::Static);
@@ -1419,16 +1419,15 @@ namespace ChaosTest
 			EXPECT_TRUE(Test.ParticleHandles[3]->ObjectState() == EObjectStateType::Sleeping);
 			EXPECT_TRUE(Test.ParticleHandles[4]->ObjectState() == EObjectStateType::Dynamic);
 		}
-		CVarPartialSleepCollisionConstraintsOnly->Set(bPartialSleepCollisionConstraintsOnlyDefault);
+		CVarPartialSleepCollisionConstraintsOnly->Set(PartialSleepCollisionConstraintsOnlyDefault);
 	}
 
-	// TODO: Test with CVars::bChaosSolverPartialSleepCollisionConstraintsOnly == true
 	// Same as PartiallySleepingIslandWithJoints, just assign a different particle to be awake.
 	TEST_P(SleepingTests, PartiallySleepingIslandWithJoints_2)
 	{
-		const bool bPartialSleepCollisionConstraintsOnly = GetParam();
-		const bool bPartialSleepCollisionConstraintsOnlyDefault = CVarPartialSleepCollisionConstraintsOnly->GetBool();
-		CVarPartialSleepCollisionConstraintsOnly->Set(bPartialSleepCollisionConstraintsOnly);
+		const int32 PartialSleepCollisionConstraintsOnly = GetParam() ? 1 : 0;
+		const int32 PartialSleepCollisionConstraintsOnlyDefault = CVarPartialSleepCollisionConstraintsOnly->GetInt();
+		CVarPartialSleepCollisionConstraintsOnly->Set(PartialSleepCollisionConstraintsOnly);
 
 		const bool bPartialSleepEnabled = true;
 		FSleepingTest Test(bPartialSleepEnabled);
@@ -1442,7 +1441,7 @@ namespace ChaosTest
 
 		Test.Advance();
 
-		if (bPartialSleepCollisionConstraintsOnly)
+		if (PartialSleepCollisionConstraintsOnly == 1)
 		{
 			// All particles wake up since the connecting constraints don't support partial sleeping
 			for (int32 Id = 1; Id < Test.ParticleHandles.Num(); ++Id)
@@ -1451,7 +1450,7 @@ namespace ChaosTest
 				EXPECT_TRUE(ParticleHandle->ObjectState() == EObjectStateType::Dynamic);
 			}
 		}
-		else
+		else if (PartialSleepCollisionConstraintsOnly == 0)
 		{
 			// Partial sleep supported, no state change
 			EXPECT_TRUE(Test.ParticleHandles[0]->ObjectState() == EObjectStateType::Static);
@@ -1460,7 +1459,7 @@ namespace ChaosTest
 			EXPECT_TRUE(Test.ParticleHandles[3]->ObjectState() == EObjectStateType::Sleeping);
 			EXPECT_TRUE(Test.ParticleHandles[4]->ObjectState() == EObjectStateType::Sleeping);
 		}
-		CVarPartialSleepCollisionConstraintsOnly->Set(bPartialSleepCollisionConstraintsOnlyDefault);
+		CVarPartialSleepCollisionConstraintsOnly->Set(PartialSleepCollisionConstraintsOnlyDefault);
 	}
 
 	// TODO: Test with CVars::bChaosSolverPartialSleepCollisionConstraintsOnly == true
@@ -1468,9 +1467,9 @@ namespace ChaosTest
 	// Since all constraints already sleep, the entire island remains asleep.
 	TEST_P(SleepingTests, FullySleepingIslandWithJoints)
 	{
-		const bool bPartialSleepCollisionConstraintsOnly = GetParam();
-		const bool bPartialSleepCollisionConstraintsOnlyDefault = CVarPartialSleepCollisionConstraintsOnly->GetBool();
-		CVarPartialSleepCollisionConstraintsOnly->Set(bPartialSleepCollisionConstraintsOnly);
+		const int32 PartialSleepCollisionConstraintsOnly = GetParam() ? 1 : 0;
+		const int32 PartialSleepCollisionConstraintsOnlyDefault = CVarPartialSleepCollisionConstraintsOnly->GetInt();
+		CVarPartialSleepCollisionConstraintsOnly->Set(PartialSleepCollisionConstraintsOnly);
 
 		const bool bPartialSleepEnabled = true;
 		FSleepingTest Test(bPartialSleepEnabled);
@@ -1490,7 +1489,7 @@ namespace ChaosTest
 			const FRigidBodyHandle_External* ParticleHandle = Test.ParticleHandles[Id];
 			EXPECT_TRUE(ParticleHandle->ObjectState() == EObjectStateType::Sleeping);
 		}
-		CVarPartialSleepCollisionConstraintsOnly->Set(bPartialSleepCollisionConstraintsOnlyDefault);
+		CVarPartialSleepCollisionConstraintsOnly->Set(PartialSleepCollisionConstraintsOnlyDefault);
 	}
 
 	INSTANTIATE_TEST_SUITE_P(, SleepingTests, testing::Bool());
diff --git a/Engine/Source/Runtime/Experimental/Chaos/Private/Chaos/Island/IslandManager.cpp b/Engine/Source/Runtime/Experimental/Chaos/Private/Chaos/Island/IslandManager.cpp
--- a/Engine/Source/Runtime/Experimental/Chaos/Private/Chaos/Island/IslandManager.cpp
+++ b/Engine/Source/Runtime/Experimental/Chaos/Private/Chaos/Island/IslandManager.cpp
@@ -83,9 +83,12 @@ namespace Chaos::CVars
 	FRealSingle ChaosSolverMomentumPropagation = 0.3f;
 	FAutoConsoleVariableRef CVarChaosSolverMomentumPropagation(TEXT("p.Chaos.Solver.Sleep.PartialIslandSleep.MomentumPropagation"), ChaosSolverMomentumPropagation, TEXT("[EXPERIMENTAL] Ratio of the impact momentum transferred from one sleeping particle to another through a constraint during partial wake event (0 to disable, 1 to propagate without dissipation).[def:0.3]"));
 
-	/** Cvar to enable/disable partial sleep for collision constraints only, i.e. the sleep state of particles connect by other constraint types must match (partial island sleeping only) */
-	bool bChaosSolverPartialSleepCollisionConstraintsOnly = true;
-	FAutoConsoleVariableRef CVarChaosSolverPartialSleepCollisionConstraintsOnly(TEXT("p.Chaos.Solver.Sleep.PartialIslandSleep.CollisionConstraintsOnly"), bChaosSolverPartialSleepCollisionConstraintsOnly, TEXT("[EXPERIMENTAL] Safeguard only: Only collision constraints permit partial sleep, i.e. the sleep state of particles connect by other constraint types must match.[def:true]"));
+	/** Cvar to enable/disable partial sleep for collision constraints only, i.e. the sleep state of particles connect by other constraint types must match (partial island sleeping only)
+	* 0: never prevent partial sleeping for non-collision constraints, i.e. all particles are allowed to sleep independent of their neighbors' sleep state and constraint types. 
+	* 1: propagate partial wake events through non-collision constraints, AND post-solver re-waking of sleeping particles connected to an awake particle by non-collision constraint
+	* -1: propagate partial wake events through non-collision constraints, NO post-solver re-waking of sleeping particles connected to an awake particle by non-collision constraint */
+	int32 ChaosSolverPartialSleepCollisionConstraintsOnly = 1;
+	FAutoConsoleVariableRef CVarChaosSolverPartialSleepCollisionConstraintsOnly(TEXT("p.Chaos.Solver.Sleep.PartialIslandSleep.CollisionConstraintsOnly"), ChaosSolverPartialSleepCollisionConstraintsOnly, TEXT("[EXPERIMENTAL] Safeguard only: Prevent partial island sleep between non-collision constraints (0: off, 1: re-wake sleeping particles connected to awake particles by non-collision constraints, -1: propagate wake events through non-collision constraints).[def:1]"));
 
 	/** Cvar to enable partial waking pre-integration (UpdateExplicitSleep). */
 	bool bChaosSolverPartialWakePreIntegration = true;
@@ -625,10 +628,10 @@ namespace Chaos::Private
 		return NextVisitEpoch++;
 	}
 
-	bool FPBDIslandManager::ShouldAssignLevels() const
+	bool FPBDIslandManager::ShouldAssignLevels(const FPBDIsland* Island) const
 	{
 		// Partial island sleeping always needs levels
-		return bAssignLevels || CVars::bChaosSolverPartialIslandSleep;
+		return bAssignLevels || Island->IsPartialSleepEnabled();
 	}
 
 	void FPBDIslandManager::SetIsDeterministic(const bool bInIsDeterministic)
@@ -2160,7 +2163,7 @@ namespace Chaos::Private
 			if (!Island->Flags.bIsSleeping && (!!Island->Flags.bItemsAdded || !!Island->Flags.bItemsRemoved))
 			{
 				// We only need to assign levels if shock propagation or partial island sleeping is enabled (but we always sort)
-				if (ShouldAssignLevels())
+				if (ShouldAssignLevels(Island))
 				{
 					AssignIslandLevels(Island);
 				}
@@ -2705,7 +2708,7 @@ namespace Chaos::Private
 		SleepNonMovingParticles(Island, Dt);
 
 		bool bHaveParticlesBeenWokenUp = false;
-		if (CVars::bChaosSolverPartialSleepCollisionConstraintsOnly)
+		if (CVars::ChaosSolverPartialSleepCollisionConstraintsOnly >= 1)
 		{
 			// If this island contains any constraint which don't support partial sleeping,
 			// we need to traverse the constraints to make sure that connected particles are either both sleeping or both awake.
@@ -3308,7 +3311,7 @@ namespace Chaos::Private
 						const bool bIsTriggerNode = NodeIndex < NumNodesInQueue;
 						const bool bIsConfirmedPropagation = bIsTriggerNode && (Edge->Flags.bIsNew || bTriggerNodesAlwaysPropagate);
 						const bool bIsUpwardPropagation = OtherNode->Level >= Edge->Level;	// Propagation upward through the levels, i.e. away from closest kinematic
-						const bool bEdgePropagatesPartialWake = CVars::bChaosSolverPartialSleepCollisionConstraintsOnly && Edge->GetConstraint() && !Edge->GetConstraint()->SupportsPartialIslandSleeping();
+						const bool bEdgePropagatesPartialWake = (CVars::ChaosSolverPartialSleepCollisionConstraintsOnly != 0) && Edge->GetConstraint() && !Edge->GetConstraint()->SupportsPartialIslandSleeping();
 
 						// 1) We process one of the trigger nodes and propagate through a new constraint (if bTriggerNodesAlwaysPropagate == false).
 						// 2) We propagate upward in the constraint graph levels (away from kinematic).
diff --git a/Engine/Source/Runtime/Experimental/Chaos/Public/Chaos/Island/IslandManager.h b/Engine/Source/Runtime/Experimental/Chaos/Public/Chaos/Island/IslandManager.h
--- a/Engine/Source/Runtime/Experimental/Chaos/Public/Chaos/Island/IslandManager.h
+++ b/Engine/Source/Runtime/Experimental/Chaos/Public/Chaos/Island/IslandManager.h
@@ -660,7 +660,7 @@ namespace Chaos::Private
 		CHAOS_API void UpdateAllParticleSleepMetrics(const FRealSingle Dt);
 
 		CHAOS_API int32 GetNextVisitEpoch();
-		CHAOS_API bool ShouldAssignLevels() const;
+		CHAOS_API bool ShouldAssignLevels(const FPBDIsland* Island) const;
 
 		// Ideally the island manager would not know about the particles, materials etc, but
 		// it is currently required for particle sleep management
