BugHawk comment fix in unit test

diff --git a/Engine/Source/Programs/HeadlessChaos/Private/HeadlessChaosTestGraphEvolution.cpp b/Engine/Source/Programs/HeadlessChaos/Private/HeadlessChaosTestGraphEvolution.cpp
--- a/Engine/Source/Programs/HeadlessChaos/Private/HeadlessChaosTestGraphEvolution.cpp
+++ b/Engine/Source/Programs/HeadlessChaos/Private/HeadlessChaosTestGraphEvolution.cpp
@@ -2830,7 +2830,7 @@ namespace ChaosTest
 
 		if (bInParallel)
 		{
-			// Multithreaded access not supported, this will crash in debug but pass in development.
+			// Multithreaded access not supported, this will crash in debug and development due to access detectors.
 			PhysicsParallelFor(Test.ParticleHandles.Num(),
 				[&](const int32 Index)
 				{
