Collision Filter Refactor Part 4.2: Cleaning up some more legacy usage of the old collision filter structure.
- Clean up the SQVisitor to ParticleHandle prepre filter path to use the new structures.
- Clean up some tests to use the new structures (and hide the old one more).
- Clean up CalcHitType to be the new structures (temporarily convert under the hood).

diff --git a/Engine/Source/Programs/HeadlessChaos/Private/HeadlessChaosTestPhysicsScene.cpp b/Engine/Source/Programs/HeadlessChaos/Private/HeadlessChaosTestPhysicsScene.cpp
--- a/Engine/Source/Programs/HeadlessChaos/Private/HeadlessChaosTestPhysicsScene.cpp
+++ b/Engine/Source/Programs/HeadlessChaos/Private/HeadlessChaosTestPhysicsScene.cpp
@@ -35,7 +35,7 @@ namespace ChaosTest {
 		bool bQueryGameThread = true; // Query game thread or physics thread data?
 
 		bool bUseQueryFilter = false;
-		FCollisionFilterData FilterData;
+		Chaos::Filter::FQueryFilterData FilterData;
 
 		FSimpleRaycastVisitor(const FVec3& InStart, bool bInQueryGameThread)
 			: Start(InStart)
@@ -44,13 +44,12 @@ namespace ChaosTest {
 		{
 		}
 
-		FSimpleRaycastVisitor(const FVec3& InStart, FCollisionFilterData& InFilterData, bool bInQueryGameThread)
+		FSimpleRaycastVisitor(const FVec3& InStart, const Chaos::Filter::FQueryFilterData& InFilterData, bool bInQueryGameThread)
 			: Start(InStart)
 			, bHit(false)
 			, bQueryGameThread(bInQueryGameThread)
 			, bUseQueryFilter(true)
 			, FilterData(InFilterData)
-			
 		{
 		}
 
@@ -58,7 +57,7 @@ namespace ChaosTest {
 		{
 			if (bUseQueryFilter)
 			{
-				return Payload.PrePreQueryFilter(&FilterData);
+				return Payload.PrePreFilter(FilterData);
 			}
 			return false;
 		}
@@ -143,6 +142,16 @@ namespace ChaosTest {
 		return HitBuffer;
 	}
 
+	bool PrePreFilter(const Chaos::Filter::FQueryFilterData& QueryFilter, const Chaos::Filter::FShapeFilterData& ShapeFilter)
+	{
+		PRAGMA_DISABLE_INTERNAL_WARNINGS
+		const FCollisionFilterData LegacyQueryFilter = Chaos::Filter::FQueryFilterBuilder::GetLegacyQueryFilter(QueryFilter);
+		Chaos::Filter::FCombinedShapeFilterData CombinedFilter(ShapeFilter, Chaos::Filter::FInstanceData());
+		const FCollisionFilterData LegacyShapeFilter = Chaos::Filter::FShapeFilterBuilder::GetLegacyShapeQueryFilter(CombinedFilter);
+		PRAGMA_ENABLE_INTERNAL_WARNINGS
+		return PrePreQueryFilterImp(LegacyQueryFilter, LegacyShapeFilter);
+	}
+
 	GTEST_TEST(EngineInterface, CreateAndReleaseActor)
 	{
 		FChaosScene Scene(nullptr, /*AsyncDt=*/-1);
@@ -479,27 +488,21 @@ namespace ChaosTest {
 		Chaos::Filter::FQueryTraceFilterBuilder Builder;
 		Builder.SetCollisionChannelIndex(7);
 		const Chaos::Filter::FQueryFilterData QueryFilterData = Builder.Build();
-		FCollisionFilterData QueryFilter = Chaos::Filter::FQueryFilterBuilder::GetLegacyQueryFilter(QueryFilterData);
 		PRAGMA_ENABLE_INTERNAL_WARNINGS
 
 		// Get collision data off shape
 		for (const TUniquePtr<Chaos::FPerShapeData>& Shape : Particle.ShapesArray())
 		{
-			const FCollisionData& CollisionData = Shape->GetCollisionData();
-			PRAGMA_DISABLE_INTERNAL_WARNINGS
-			const FCollisionFilterData QueryData = CollisionData.GetQueryData();
-			PRAGMA_ENABLE_INTERNAL_WARNINGS
-			EXPECT_EQ(QueryData.Word0, 0); // ensure query filter is defaulted to 0 (no query allowed at all)
-
+			const Chaos::Filter::FShapeFilterData& ShapeFilter = Shape->GetShapeFilterData();
 			// Verify query is filtered out with default collision data on shape
-			bool bFiltered = PrePreQueryFilterImp(QueryFilter, QueryData);
+			bool bFiltered = PrePreFilter(QueryFilterData, ShapeFilter);
 			EXPECT_EQ(bFiltered, true);
 		}
 
 		// Query against particle on game thread, should fail to hit due to particle filter being defaulted, no touch/block set.
 		{
 			bool bQueryGameThread = true;
-			FSimpleRaycastVisitor Visitor(Start, QueryFilter, bQueryGameThread);
+			FSimpleRaycastVisitor Visitor(Start, QueryFilterData, bQueryGameThread);
 			Scene.GetSpacialAcceleration()->Raycast(Start, Dir, Length, Visitor);
 			EXPECT_EQ(Visitor.bHit, false);
 		}
@@ -511,9 +514,6 @@ namespace ChaosTest {
 		ShapeFilterBuilder.SetOverlapChannelMask(TNumericLimits<uint64>::Max());
 		ShapeFilterBuilder.SetFilterFlags(Chaos::EFilterFlags::All);
 		const Chaos::Filter::FShapeFilterData NewParticleQueryFilter = ShapeFilterBuilder.Build();
-		// Temporarily get the legacy query filter from the new filter until PrePrepQueryFilter is updated.
-		Chaos::Filter::FCombinedShapeFilterData CombinedShapeFilter(NewParticleQueryFilter, Chaos::Filter::FInstanceData());
-		const FCollisionFilterData LegacyQueryFilter = Chaos::Filter::FShapeFilterBuilder::GetLegacyShapeQueryFilter(CombinedShapeFilter);
 		PRAGMA_ENABLE_INTERNAL_WARNINGS
 
 		for (const TUniquePtr<Chaos::FPerShapeData>& Shape : Particle.ShapesArray())
@@ -522,7 +522,7 @@ namespace ChaosTest {
 			Shape->SetShapeFilterData(NewParticleQueryFilter);
 
 			// Filter with new data, ensuring we pass and are not filtered out.
-			bool bFiltered = PrePreQueryFilterImp(QueryFilter, LegacyQueryFilter);
+			bool bFiltered = PrePreFilter(QueryFilterData, NewParticleQueryFilter);
 			EXPECT_EQ(bFiltered, false);
 		}
 
@@ -532,7 +532,7 @@ namespace ChaosTest {
 		// Query against particle on game thread, should hit with new filter.
 		{
 			bool bQueryGameThread = true;
-			FSimpleRaycastVisitor Visitor(Start, QueryFilter, bQueryGameThread);
+			FSimpleRaycastVisitor Visitor(Start, QueryFilterData, bQueryGameThread);
 			Scene.GetSpacialAcceleration()->Raycast(Start, Dir, Length, Visitor);
 			EXPECT_EQ(Visitor.bHit, true);
 		}
@@ -545,7 +545,7 @@ namespace ChaosTest {
 		// If this fails it means we did not update cached filter data in acceleration structure entry.
 		{
 			bool bQueryGameThread = false;
-			FSimpleRaycastVisitor Visitor(Start, QueryFilter, bQueryGameThread);
+			FSimpleRaycastVisitor Visitor(Start, QueryFilterData, bQueryGameThread);
 			Scene.GetSolver()->GetEvolution()->GetSpatialAcceleration()->Raycast(Start, Dir, Length, Visitor);
 			EXPECT_EQ(Visitor.bHit, true);
 		}
diff --git a/Engine/Source/Runtime/Engine/Classes/PhysicsEngine/ExternalSpatialAccelerationPayload.h b/Engine/Source/Runtime/Engine/Classes/PhysicsEngine/ExternalSpatialAccelerationPayload.h
--- a/Engine/Source/Runtime/Engine/Classes/PhysicsEngine/ExternalSpatialAccelerationPayload.h
+++ b/Engine/Source/Runtime/Engine/Classes/PhysicsEngine/ExternalSpatialAccelerationPayload.h
@@ -36,8 +36,12 @@ public:
 	bool operator!=(const FExternalSpatialAccelerationPayload& Other) const;
 	ENGINE_API Chaos::FUniqueIdx UniqueIdx() const;
 	ENGINE_API void Serialize(Chaos::FChaosArchive& Ar);
+	UE_DEPRECATED(5.8, "This is deprecated in favor of PrePreFilter")
 	bool PrePreQueryFilter(const void* QueryData) const { return false; }
+	UE_DEPRECATED(5.8, "This is deprecated in favor of PrePreFilter")
 	bool PrePreSimFilter(const void* SimData) const { return false; }
+	template <typename FilterDataType>
+	bool PrePreFilter(const FilterDataType&) const { return false; }
 	bool IsValid() const 
 	{
 		return CachedUniqueIdx.IsValid(); // Don't use this payload if the UniqueIdx is invalid
diff --git a/Engine/Source/Runtime/Engine/Private/Collision/CollisionConversions.cpp b/Engine/Source/Runtime/Engine/Private/Collision/CollisionConversions.cpp
--- a/Engine/Source/Runtime/Engine/Private/Collision/CollisionConversions.cpp
+++ b/Engine/Source/Runtime/Engine/Private/Collision/CollisionConversions.cpp
@@ -282,7 +282,6 @@ const FPhysicsActor* GetGTActor(const Chaos::FGeometryParticleHandle* PTActor)
 template <typename THitLocation>
 EConvertQueryResult ConvertQueryImpactHitImp(const UWorld* World, const THitLocation& Hit, FHitResult& OutResult, float CheckLength, const Chaos::Filter::FQueryFilterData& QueryFilter, const FVector& StartLoc, const FVector& EndLoc, const FPhysicsGeometry* Geom, const FTransform& QueryTM, bool bReturnFaceIndex, bool bReturnPhysMat)
 {
-	const FCollisionFilterData LegacyQueryFilter = Chaos::Filter::FQueryFilterBuilder::GetLegacyQueryFilter(QueryFilter);
 	using namespace ChaosInterface;
 
 	SCOPE_CYCLE_COUNTER(STAT_ConvertQueryImpactHit);
@@ -348,8 +347,8 @@ EConvertQueryResult ConvertQueryImpactHitImp(const UWorld* World, const THitLoca
 	const FPhysicsShape* GTShape = GetGTShape<bIsGTQuery>(pHitShape, GTActor);
 
 	// See if this is a 'blocking' hit
-	const FCollisionFilterData ShapeFilter = Chaos::Filter::FShapeFilterBuilder::GetLegacyShapeQueryFilter(GetCombinedShapeFilterData(HitShape));
-	const ECollisionQueryHitType HitType = FCollisionQueryFilterCallback::CalcQueryHitType(LegacyQueryFilter, ShapeFilter);
+	const Chaos::Filter::FShapeFilterData ShapeFilter = GetShapeFilterData(HitShape);
+	const ECollisionQueryHitType HitType = FCollisionQueryFilterCallback::CalcQueryHitType(QueryFilter, ShapeFilter);
 	OutResult.bBlockingHit = (HitType == ECollisionQueryHitType::Block);
 	OutResult.bStartPenetrating = bInitialOverlap;
 
@@ -562,9 +561,8 @@ static bool ConvertOverlappedShapeToImpactHit(const UWorld* World, const THitLoc
 	const auto& HitActor = *pHitActor;
 
 	// See if this is a 'blocking' hit
-	FCollisionFilterData ShapeFilter = Chaos::Filter::FShapeFilterBuilder::GetLegacyShapeQueryFilter(GetCombinedShapeFilterData(HitShape));
-	const FCollisionFilterData LegacyQueryFilter = Chaos::Filter::FQueryFilterBuilder::GetLegacyQueryFilter(QueryFilter);
-	const ECollisionQueryHitType HitType = FCollisionQueryFilterCallback::CalcQueryHitType(LegacyQueryFilter, ShapeFilter);
+	const Chaos::Filter::FShapeFilterData ShapeFilter = GetShapeFilterData(HitShape);
+	const ECollisionQueryHitType HitType = FCollisionQueryFilterCallback::CalcQueryHitType(QueryFilter, ShapeFilter);
 	const bool bBlockingHit = (HitType == ECollisionQueryHitType::Block);
 	OutResult.bBlockingHit = bBlockingHit;
 
@@ -780,11 +778,10 @@ static void AddUniqueOverlap(TArray<FOverlapResult>& OutOverlaps, const FOverlap
 bool IsBlocking(const FPhysicsShape& Shape, const Chaos::Filter::FQueryFilterData& QueryFilter)
 {
 	using namespace ChaosInterface;
-	const FCollisionFilterData LegacyQueryFilter = Chaos::Filter::FQueryFilterBuilder::GetLegacyQueryFilter(QueryFilter);
 
 	// See if this is a 'blocking' hit
-	const FCollisionFilterData ShapeFilter = Chaos::Filter::FShapeFilterBuilder::GetLegacyShapeQueryFilter(GetCombinedShapeFilterData(Shape));
-	const ECollisionQueryHitType HitType = FCollisionQueryFilterCallback::CalcQueryHitType(LegacyQueryFilter, ShapeFilter);
+	const Chaos::Filter::FShapeFilterData ShapeFilter = GetShapeFilterData(Shape);
+	const ECollisionQueryHitType HitType = FCollisionQueryFilterCallback::CalcQueryHitType(QueryFilter, ShapeFilter);
 	const bool bBlock = (HitType == ECollisionQueryHitType::Block);
 	return bBlock;
 }
diff --git a/Engine/Source/Runtime/Engine/Private/PhysicsEngine/CollisionQueryFilterCallback.cpp b/Engine/Source/Runtime/Engine/Private/PhysicsEngine/CollisionQueryFilterCallback.cpp
--- a/Engine/Source/Runtime/Engine/Private/PhysicsEngine/CollisionQueryFilterCallback.cpp
+++ b/Engine/Source/Runtime/Engine/Private/PhysicsEngine/CollisionQueryFilterCallback.cpp
@@ -8,8 +8,12 @@
 #include "Physics/Experimental/ChaosInterfaceWrapper.h"
 #include "Chaos/GeometryParticles.h"
 
-ECollisionQueryHitType FCollisionQueryFilterCallback::CalcQueryHitType(const FCollisionFilterData& QueryFilter, const FCollisionFilterData& ShapeFilter, bool bPreFilter)
+ECollisionQueryHitType FCollisionQueryFilterCallback::CalcQueryHitType(const FQueryFilterData& InQueryFilter, const FShapeFilterData& InShapeFilter, bool bPreFilter)
 {
+	const FCollisionFilterData QueryFilter = Chaos::Filter::FQueryFilterBuilder::GetLegacyQueryFilter(InQueryFilter);
+	Chaos::Filter::FCombinedShapeFilterData CombinedFilter(InShapeFilter, Chaos::Filter::FInstanceData());
+	const FCollisionFilterData ShapeFilter = Chaos::Filter::FShapeFilterBuilder::GetLegacyShapeQueryFilter(CombinedFilter);
+
 	ECollisionQuery QueryType = (ECollisionQuery)QueryFilter.Word0;
 	FMaskFilter QuerierMaskFilter;
 	PRAGMA_DISABLE_DEPRECATION_WARNINGS
@@ -186,9 +190,7 @@ ECollisionQueryHitType FCollisionQueryFilterCallback::PreFilterImp(const FQueryF
 		return (PreFilterReturnValue = ECollisionQueryHitType::None);
 	}
 	
-	const FCollisionFilterData LegacyQueryFilter = Chaos::Filter::FQueryFilterBuilder::GetLegacyQueryFilter(FilterData);
-	const FCollisionFilterData LegacyShapeFilter = Chaos::Filter::FShapeFilterBuilder::GetLegacyShapeQueryFilter(CombinedShapeFilter);
-	ECollisionQueryHitType Result = FCollisionQueryFilterCallback::CalcQueryHitType(LegacyQueryFilter, LegacyShapeFilter, true);
+	ECollisionQueryHitType Result = FCollisionQueryFilterCallback::CalcQueryHitType(FilterData, ShapeFilter, true);
 
 	if (Result == ECollisionQueryHitType::Touch && bIgnoreTouches)
 	{
diff --git a/Engine/Source/Runtime/Engine/Private/PhysicsEngine/CollisionQueryFilterCallback.h b/Engine/Source/Runtime/Engine/Private/PhysicsEngine/CollisionQueryFilterCallback.h
--- a/Engine/Source/Runtime/Engine/Private/PhysicsEngine/CollisionQueryFilterCallback.h
+++ b/Engine/Source/Runtime/Engine/Private/PhysicsEngine/CollisionQueryFilterCallback.h
@@ -60,7 +60,7 @@ public:
 	}
 	~FCollisionQueryFilterCallback() {}
 
-	static ECollisionQueryHitType CalcQueryHitType(const FCollisionFilterData& QueryFilter, const FCollisionFilterData& ShapeFilter, bool bPreFilter = false);
+	static ECollisionQueryHitType CalcQueryHitType(const FQueryFilterData& QueryFilter, const FShapeFilterData& ShapeFilter, bool bPreFilter = false);
 
 	virtual ECollisionQueryHitType PostFilter(const FQueryFilterData& FilterData, const ChaosInterface::FQueryHit& Hit) override;
 	virtual ECollisionQueryHitType PostFilter(const FQueryFilterData& FilterData, const ChaosInterface::FPTQueryHit& Hit) override;
diff --git a/Engine/Source/Runtime/Experimental/Chaos/Public/Chaos/ParticleHandle.h b/Engine/Source/Runtime/Experimental/Chaos/Public/Chaos/ParticleHandle.h
--- a/Engine/Source/Runtime/Experimental/Chaos/Public/Chaos/ParticleHandle.h
+++ b/Engine/Source/Runtime/Experimental/Chaos/Public/Chaos/ParticleHandle.h
@@ -292,6 +292,19 @@ public:
 		return false;
 	}
 
+	bool PrePreFilter(const Chaos::Filter::FQueryFilterData& QueryFilterData) const
+	{
+		if (bCanPrePreFilter)
+		{
+			PRAGMA_DISABLE_INTERNAL_WARNINGS
+			const FCollisionFilterData LegacyQueryFilter = Chaos::Filter::FQueryFilterBuilder::GetLegacyQueryFilter(QueryFilterData);
+			PRAGMA_ENABLE_INTERNAL_WARNINGS
+			return PrePreQueryFilterImp(LegacyQueryFilter, UnionQueryFilterData);
+		}
+		
+		return false;
+	}
+
 	void UpdateFrom(const FAccelerationStructureHandle& InOther)
 	{
 		UnionQueryFilterData.Word0 = InOther.UnionQueryFilterData.Word0;
diff --git a/Engine/Source/Runtime/PhysicsCore/Public/SQVisitor.h b/Engine/Source/Runtime/PhysicsCore/Public/SQVisitor.h
--- a/Engine/Source/Runtime/PhysicsCore/Public/SQVisitor.h
+++ b/Engine/Source/Runtime/PhysicsCore/Public/SQVisitor.h
@@ -56,7 +56,6 @@ struct TSQVisitor : public Chaos::ISpatialVisitor<TPayload, Chaos::FReal>
 		, HitBuffer(InHitBuffer)
 		, QueryFlags(InCommonParams.Flags)
 		, QueryFilterData(InCommonParams.FilterData)
-		, LegacyQueryFilterDataConcrete(Chaos::Filter::FQueryFilterBuilder::GetLegacyQueryFilter(InCommonParams.FilterData))
 		, QueryCallback(InCommonParams.GetQueryCallback())
 		, QueryCenter(StartTM.GetLocation())
 	{
@@ -83,7 +82,6 @@ struct TSQVisitor : public Chaos::ISpatialVisitor<TPayload, Chaos::FReal>
 		, HitBuffer(InHitBuffer)
 		, QueryFlags(InCommonParams.Flags)
 		, QueryFilterData(InCommonParams.FilterData)
-		, LegacyQueryFilterDataConcrete(Chaos::Filter::FQueryFilterBuilder::GetLegacyQueryFilter(InCommonParams.FilterData))
 		, QueryGeom(&InQueryGeom)
 		, QueryCallback(InCommonParams.GetQueryCallback())
 		, StartTM(InStartTM)
@@ -118,7 +116,6 @@ struct TSQVisitor : public Chaos::ISpatialVisitor<TPayload, Chaos::FReal>
 		, HitBuffer(InHitBuffer)
 		, QueryFlags(InCommonParams.Flags)
 		, QueryFilterData(InCommonParams.FilterData)
-		, LegacyQueryFilterDataConcrete(Chaos::Filter::FQueryFilterBuilder::GetLegacyQueryFilter(InCommonParams.FilterData))
 		, QueryGeom(&InQueryGeom)
 		, QueryCallback(InCommonParams.GetQueryCallback())
 		, StartTM(InWorldTM)
@@ -164,7 +161,7 @@ struct TSQVisitor : public Chaos::ISpatialVisitor<TPayload, Chaos::FReal>
 
 	virtual bool PrePreFilter(const TPayload& Payload) const override
 	{
-		return Payload.PrePreQueryFilter(&LegacyQueryFilterDataConcrete);
+		return Payload.PrePreFilter(QueryFilterData);
 	}
 
 	/** Return a pointer to the payload on which we are querying the acceleration structure */
@@ -525,7 +522,6 @@ private:
 	ChaosInterface::FSQHitBuffer<THitType>& HitBuffer;
 	const FChaosQueryFlags QueryFlags;
 	const Chaos::Filter::FQueryFilterData QueryFilterData;
-	const FCollisionFilterData LegacyQueryFilterDataConcrete;
 	const QueryGeometryType* QueryGeom = nullptr;
 	ICollisionQueryFilterCallback& QueryCallback;
 	const FTransform StartTM;
@@ -547,7 +543,6 @@ struct TBPVisitor : public Chaos::ISpatialVisitor<TPayload, Chaos::FReal>
 		, HitBuffer(InHitBuffer)
 		, QueryFlags(InCommonParams.Flags)
 		, QueryFilterData(InCommonParams.FilterData)
-		, LegacyQueryFilterDataConcrete(Chaos::Filter::FQueryFilterBuilder::GetLegacyQueryFilter(InCommonParams.FilterData))
 		, QueryGeom(&InQueryGeom)
 		, QueryCallback(InCommonParams.GetQueryCallback())
 		, StartTM(InWorldTM)
@@ -585,7 +580,7 @@ struct TBPVisitor : public Chaos::ISpatialVisitor<TPayload, Chaos::FReal>
 
 	virtual bool PrePreFilter(const TPayload& Payload) const override
 	{
-		return Payload.PrePreQueryFilter(&LegacyQueryFilterDataConcrete);
+		return Payload.PrePreFilter(QueryFilterData);
 	}
 
 	/** Return a pointer to the payload on which we are querying the acceleration structure */
@@ -676,7 +671,6 @@ private:
 	ChaosInterface::FSQHitBuffer<THitType>& HitBuffer;
 	const FChaosQueryFlags QueryFlags;
 	const Chaos::Filter::FQueryFilterData QueryFilterData;
-	const FCollisionFilterData LegacyQueryFilterDataConcrete;
 	const QueryGeometryType* QueryGeom;
 	ICollisionQueryFilterCallback& QueryCallback;
 	const FTransform StartTM;
