[Chaos] Fix serialization test to correctly initialise X/R fields before setting geometry. Setting geometry will update world space state so an uninitialised X or R can lead to invalid bounds being calculated and triggering some new checks for bounds validity

diff --git a/Engine/Source/Programs/HeadlessChaos/Private/HeadlessChaosTestSerialization.cpp b/Engine/Source/Programs/HeadlessChaos/Private/HeadlessChaosTestSerialization.cpp
--- a/Engine/Source/Programs/HeadlessChaos/Private/HeadlessChaosTestSerialization.cpp
+++ b/Engine/Source/Programs/HeadlessChaos/Private/HeadlessChaosTestSerialization.cpp
@@ -334,15 +334,13 @@ namespace ChaosTest
 
 			FGeometryParticles OriginalParticles;
 			OriginalParticles.AddParticles(2);
-			OriginalParticles.SetR(0, FRotation3::Identity);
-			OriginalParticles.SetR(1, FRotation3::Identity);
-			OriginalParticles.SetGeometry(0, OriginalSpheres[0]);
-			OriginalParticles.SetGeometry(1, OriginalSpheres[1]);
 			OriginalParticles.SetX(0, FVec3(100, 1, 2));
 			OriginalParticles.SetX(1, FVec3(0, 1, 2));
 			OriginalParticles.SetR(0, FRotation3::Identity);
 			OriginalParticles.SetR(1, FRotation3::Identity);
-
+			OriginalParticles.SetGeometry(0, OriginalSpheres[0]);
+			OriginalParticles.SetGeometry(1, OriginalSpheres[1]);
+			
 			TBoundingVolumeHierarchy<FGeometryParticles, TArray<int32>> OriginalBVH(OriginalParticles);
 
 			FMemoryWriter Ar(Data);
