[Chaos] Further unit test fixes to resolve uninitialised particle X values leading to indeterminate world bounds calculation when setting geometry

#Synced-CL 48214939

diff --git a/Engine/Source/Programs/HeadlessChaos/Private/HeadlessChaosTestSerialization.cpp b/Engine/Source/Programs/HeadlessChaos/Private/HeadlessChaosTestSerialization.cpp
--- a/Engine/Source/Programs/HeadlessChaos/Private/HeadlessChaosTestSerialization.cpp
+++ b/Engine/Source/Programs/HeadlessChaos/Private/HeadlessChaosTestSerialization.cpp
@@ -255,6 +255,8 @@ namespace ChaosTest
 		{
 			FGeometryParticles OriginalParticles;
 			OriginalParticles.AddParticles(2);
+			OriginalParticles.SetX(0, FVec3::ZeroVector);
+			OriginalParticles.SetX(1, FVec3::ZeroVector);
 			OriginalParticles.SetR(0, FRotation3::Identity);
 			OriginalParticles.SetR(1, FRotation3::Identity);
 			OriginalParticles.SetGeometry(0, OriginalSpheres[0]);
@@ -291,6 +293,8 @@ namespace ChaosTest
 		{
 			auto OriginalParticles = MakeUnique<FGeometryParticles>();
 			OriginalParticles->AddParticles(2);
+			OriginalParticles->SetX(0, FVec3::ZeroVector);
+			OriginalParticles->SetX(1, FVec3::ZeroVector);
 			OriginalParticles->SetR(0, FRotation3::Identity);
 			OriginalParticles->SetR(1, FRotation3::Identity);
 			OriginalParticles->SetGeometry(0, OriginalSpheres[0]);
