Fix headlesschaos rewind gtests

#Synced-CL 49101379

diff --git a/Engine/Source/Programs/HeadlessChaos/Private/HeadlessChaosTestRewind.cpp b/Engine/Source/Programs/HeadlessChaos/Private/HeadlessChaosTestRewind.cpp
--- a/Engine/Source/Programs/HeadlessChaos/Private/HeadlessChaosTestRewind.cpp
+++ b/Engine/Source/Programs/HeadlessChaos/Private/HeadlessChaosTestRewind.cpp
@@ -551,7 +551,7 @@ namespace ChaosTest {
 				FRewindCallbackTestHelper* Helper = RegisterCallbackHelper(Solver);
 				Helper->TriggerRewindFunc = [&Floor, &bHasResimmed](int32 PhysicsStep) -> int32
 				{
-					if (!bHasResimmed && Floor)
+					if (!bHasResimmed && Floor && Floor->GetPhysicsThreadAPI())
 					{
 						bHasResimmed = true;
 						return 0;
@@ -563,23 +563,28 @@ namespace ChaosTest {
 					{
 						const FReal Time = PhysicsStep * SimDt;
 
-						if (bIsResimming)
+						if (bIsResimming && RewindData)
 						{
 							if (PhysicsStep == 1)
 							{
 								RewindData->SpawnProxyIfNeeded(*Floor);
 							}
 						}
-
-						if (!bHasResimmed || Time < 4.5)
-						{
-							//simply movement without hitting floor because it's spawned too late
-							EXPECT_NEAR(Proxy->GetPhysicsThreadAPI()->X()[2], 14.5 - Time, 1e-2);
-						}
-						else
-						{
-							//floor spawned earlier so we hit it
-							EXPECT_GE(Proxy->GetPhysicsThreadAPI()->X()[2], 10 - 1e-2);
+						if(Proxy)
+						{ 
+							if (FRigidBodyHandle_Internal* PT0 = Proxy->GetPhysicsThreadAPI())
+							{ 
+								if (!bHasResimmed || Time < 4.5)
+								{
+									//simply movement without hitting floor because it's spawned too late
+									EXPECT_NEAR(PT0->X()[2], 14.5 - Time, 1e-2);
+								}
+								else
+								{
+									//floor spawned earlier so we hit it
+									EXPECT_GE(PT0->X()[2], 10 - 1e-2);
+								}
+							}
 						}
 				};
 
@@ -625,7 +630,7 @@ namespace ChaosTest {
 				FRewindCallbackTestHelper* Helper = RegisterCallbackHelper(Solver);
 				Helper->TriggerRewindFunc = [&SpawnedProxy, &SpawnedProxyNoCorrection, &bHasResimmed](int32 PhysicsStep) -> int32
 				{
-					if(!bHasResimmed && SpawnedProxy)
+					if(!bHasResimmed && SpawnedProxy && SpawnedProxy->GetPhysicsThreadAPI())
 					{
 						bHasResimmed = true;
 						return 0;
@@ -638,26 +643,29 @@ namespace ChaosTest {
 				{
 					const FReal Time = PhysicsStep * SimDt;
 
-					if (bIsResimming && SpawnedProxy)
-					{
-						if (PhysicsStep == 10)
-						{
-							RewindData->SpawnProxyIfNeeded(*SpawnedProxy);
-							SpawnedProxy->GetPhysicsThreadAPI()->SetX(FVec3(500, 0, 100.0));
-						}
-					}
+					if(SpawnedProxy && RewindData)
+					{ 
+						if(FRigidBodyHandle_Internal* PT0 = SpawnedProxy->GetPhysicsThreadAPI())
+						{ 
+							if (bIsResimming && (PhysicsStep == 10))
+							{
+								RewindData->SpawnProxyIfNeeded(*SpawnedProxy);
+								PT0->SetX(FVec3(500, 0, 100.0));
+							}
 
-					if (SpawnedProxy && SpawnedProxyNoCorrection)
-					{
-						auto PT0 = SpawnedProxy->GetPhysicsThreadAPI();
-						auto PT1 = SpawnedProxyNoCorrection->GetPhysicsThreadAPI();
-						
-						// After we've applied the correction and simmed past frame 10, we expect the first proxy to be "ahead" of the second one that didn't get corrected
-						if (bHasResimmed && PhysicsStep > 10)
-						{
-							EXPECT_LT(PT0->X()[2], PT1->X()[2]);
+							if (SpawnedProxyNoCorrection)
+							{
+								if (FRigidBodyHandle_Internal* PT1 = SpawnedProxyNoCorrection->GetPhysicsThreadAPI())  
+								{
+									// After we've applied the correction and simmed past frame 10, we expect the first proxy to be "ahead" of the second one that didn't get corrected
+									if (bHasResimmed && PhysicsStep > 10)
+									{
+										EXPECT_LT(PT0->X()[2], PT1->X()[2]);
+									}
+								}
+							}	
 						}
-					}					
+					}
 				};
 
 				const int32 LastGameStep = 32;
@@ -1900,7 +1908,7 @@ namespace ChaosTest {
 				};
 
 				const int32 LastGameStep = 20;
-				const int32 NumPhysSteps = FMath::TruncToInt(LastGameStep / SimDt);
+				const int32 NumPhysSteps = FMath::TruncToInt(LastGameStep / SimDt) - 1;
 
 				Solver->SetRewindCallback(MakeUnique<FRewindCallback>(NumPhysSteps, SimDt));
 				FSimCallbackHelper* SimCallback = Solver->template CreateAndRegisterSimCallbackObject_External<FSimCallbackHelper>();
@@ -1940,8 +1948,8 @@ namespace ChaosTest {
 							CurTime += SimDt;
 						}
 					}
-
-					EXPECT_EQ(Count, NumPhysSteps * 2);	//should have two results for each physics step since we rewound
+					const int32 NumSteps = NumPhysSteps * 2 + (LastGameStep - NumPhysSteps);
+					EXPECT_EQ(Count, NumSteps);	
 				}
 				
 
@@ -2057,7 +2065,7 @@ namespace ChaosTest {
 				};
 
 				const int32 LastGameStep = 20;
-				const int32 NumPhysSteps = FMath::TruncToInt(LastGameStep / SimDt);
+				const int32 NumPhysSteps = FMath::TruncToInt(LastGameStep / SimDt) - 1;
 
 				auto UniqueRewindCallback = MakeUnique<FRewindCallback>(NumPhysSteps);
 				auto RewindCallback = UniqueRewindCallback.Get();
@@ -2100,11 +2108,16 @@ namespace ChaosTest {
 						}
 						else
 						{
+							if ((Count == 2 * NumPhysSteps))
+							{
+								Correction = 0.0;
+							}
 							CurTime += SimDt;
 						}
 					}
 
-					EXPECT_EQ(Count, NumPhysSteps * 2);	//should have two results for each physics step since we rewound
+					const int32 NumSteps = NumPhysSteps * 2 + (LastGameStep - NumPhysSteps);
+					EXPECT_EQ(Count, NumSteps);	
 				}
 
 
@@ -2280,7 +2293,8 @@ namespace ChaosTest {
 			int32 FirstStepResim = INT_MAX;
 			int32 NumResimSteps = 0;
 
-			FRewindCallbackTestHelper* Helper = RegisterCallbackHelper(Solver, FMath::TruncToInt(LastGameStep / SimDt) );
+			const int32 NumPhysSteps = FMath::TruncToInt(LastGameStep / SimDt) - 1;
+			FRewindCallbackTestHelper* Helper = RegisterCallbackHelper(Solver, NumPhysSteps);
 			Helper->ProcessInputsFunc = [Proxy, SimDt, &ExpectedVZ, &ExpectedXZ, &FirstStepResim, &NumResimSteps](const int32 Step, const bool bResim)
 			{
 				if(bResim)
@@ -2334,7 +2348,7 @@ namespace ChaosTest {
 			}
 
 			EXPECT_EQ(FirstStepResim, 0);
-			EXPECT_EQ(NumResimSteps, LastGameStep / SimDt);
+			EXPECT_EQ(NumResimSteps, NumPhysSteps);
 
 			//no desync so should be empty
 #if REWIND_DESYNC
