HeadlessChaos: remove extra parentheses

#Synced-CL 51662528

diff --git a/Engine/Source/Programs/HeadlessChaos/Private/HeadlessChaosTestRewind.cpp b/Engine/Source/Programs/HeadlessChaos/Private/HeadlessChaosTestRewind.cpp
--- a/Engine/Source/Programs/HeadlessChaos/Private/HeadlessChaosTestRewind.cpp
+++ b/Engine/Source/Programs/HeadlessChaos/Private/HeadlessChaosTestRewind.cpp
@@ -2109,7 +2109,7 @@ namespace ChaosTest {
 						}
 						else
 						{
-							if ((Count == 2 * NumPhysSteps))
+							if (Count == 2 * NumPhysSteps)
 							{
 								Correction = 0.0;
 							}
