Fixing an issue with tri mesh sweep CCD where the TOI was incorrectly calculated when there were multiple hits.
Was regressed in CL 48292054.

diff --git a/Engine/Source/Programs/HeadlessChaos/Private/HeadlessChaosTestSweep.cpp b/Engine/Source/Programs/HeadlessChaos/Private/HeadlessChaosTestSweep.cpp
--- a/Engine/Source/Programs/HeadlessChaos/Private/HeadlessChaosTestSweep.cpp
+++ b/Engine/Source/Programs/HeadlessChaos/Private/HeadlessChaosTestSweep.cpp
@@ -931,6 +931,7 @@ namespace ChaosTest
 			FTriangleMeshImplicitObject::ParticlesType TrimeshParticles;
 			TArray<TVec3<int32>> Indices;
 			TArray<uint16> Materials;
+			TUniquePtr<TArray<int32>> ExternalFaceIndexMap;
 
 			void AutoBuildMaterials()
 			{
@@ -940,6 +941,15 @@ namespace ChaosTest
 					Materials.Emplace(0);
 				}
 			}
+
+			void AutoBuildExternalFaceIndices()
+			{
+				ExternalFaceIndexMap = MakeUnique<TArray<int32>>();
+				for (int32 I = 0; I < Indices.Num(); ++I)
+				{
+					ExternalFaceIndexMap->Add(I);
+				}
+			}
 		};
 
 		static FMeshInitData CreateSimpleTriInitData()
@@ -1764,6 +1774,277 @@ namespace ChaosTest
 		EXPECT_EQ(Result.bResult, false);
 	}
 
+	struct FSweepCCDFixture : public FTriangleMeshSweepFixture
+	{
+		FSweepCCDResultData SweepSphereCCD(FTriangleMeshImplicitObject& TriangleMesh, const FVec3& SweepStart, const FVec3& SweepDirection, const FQuat& SweepRotation = FQuat::Identity, const FReal SweepLength = 100)
+		{
+			const Chaos::FSphere Sphere(FVec3::ZeroVector, SphereRadius);
+			const FRigidTransform3 StartTM(SweepStart, SweepRotation);
+			const FVec3 NormalizedSweepDirection = SweepDirection.GetSafeNormal();
+			const FRigidTransform3 BToATM = StartTM.GetRelativeTransformNoScale(FRigidTransform3());
+
+			FSweepCCDResultData Result;
+			Result.bResult = TriangleMesh.SweepGeomCCD(Sphere, BToATM, NormalizedSweepDirection, SweepLength, IgnorePenetration, TargetPenetration, Result.TOI, Result.Phi, Result.Position, Result.Normal, Result.FaceId, Result.FaceNormal);
+			return Result;
+		}
+
+		FReal IgnorePenetration = 0;
+		FReal TargetPenetration = 0;
+	};
+
+	TEST_F(FSweepCCDFixture, SimpleTriMeshWithTwoHits_SweepCCD_ResultsAreExpected)
+	{
+		const FReal SweepLength = 1000;
+		SphereRadius = 15;
+
+		FMeshInitData InitData;
+		InitData.TrimeshParticles = FTriangleMeshImplicitObject::ParticlesType
+		({
+			FVec3(0, 0, 0), FVec3(10, 0, 0), FVec3(0, 10, 0),
+			FVec3(0, 0, 5), FVec3(10, 0, 5), FVec3(0, 10, 5),
+		});
+		InitData.Indices =
+		{
+			{0, 1, 2},
+			{3, 4, 5},
+		};
+		InitData.AutoBuildMaterials();
+		InitData.AutoBuildExternalFaceIndices();
+
+		FTriangleMeshImplicitObject TriangleMesh(MoveTemp(InitData.TrimeshParticles), MoveTemp(InitData.Indices), MoveTemp(InitData.Materials), MoveTemp(InitData.ExternalFaceIndexMap), nullptr, true);
+
+		FVec3 SweepStart = FVec3(2, 2, 50);
+		FVec3 SweepDir = FVec3(0, 0, -1);
+		FQuat SweepRotation = FQuat::Identity;
+
+		// First do a test with no target penetration
+		{
+			TargetPenetration = 0;
+			const double ExpectedTOI = (SweepStart.Z - SphereRadius - 5 + TargetPenetration) / SweepLength;
+
+			FSweepCCDResultData Result = SweepSphereCCD(TriangleMesh, SweepStart, SweepDir, SweepRotation, SweepLength);
+
+			EXPECT_EQ(Result.bResult, true);
+			EXPECT_NEAR(Result.TOI, ExpectedTOI, KINDA_SMALL_NUMBER);
+			EXPECT_NEAR(Result.Phi, 0, KINDA_SMALL_NUMBER);
+			EXPECT_VECTOR_NEAR(Result.Position, FVec3(2, 2, 5), KINDA_SMALL_NUMBER);
+			EXPECT_VECTOR_NEAR(Result.Normal, FVec3(0, 0, 1), KINDA_SMALL_NUMBER);
+			EXPECT_EQ(Result.FaceId, 1);
+		}
+
+		// Then do another test with target penetration. The same triangle should be hit, but with different TOI and PHI.
+		{
+			TargetPenetration = 6.3500003814697266;
+			const double ExpectedTOI = (SweepStart.Z - SphereRadius - 5 + TargetPenetration) / SweepLength;
+
+			FSweepCCDResultData Result = SweepSphereCCD(TriangleMesh, SweepStart, SweepDir, SweepRotation, SweepLength);
+
+			EXPECT_EQ(Result.bResult, true);
+			EXPECT_NEAR(Result.TOI, ExpectedTOI, KINDA_SMALL_NUMBER);
+			EXPECT_NEAR(Result.Phi, -TargetPenetration, KINDA_SMALL_NUMBER);
+			EXPECT_VECTOR_NEAR(Result.Position, FVec3(2, 2, 5), KINDA_SMALL_NUMBER);
+			EXPECT_VECTOR_NEAR(Result.Normal, FVec3(0, 0, 1), KINDA_SMALL_NUMBER);
+			EXPECT_EQ(Result.FaceId, 1);
+		}
+	}
+
+	TEST_F(FSweepCCDFixture, SavedTriMeshWithMultipleHits_SweepCCD_ResultsAreExpected)
+	{
+		// This is a mesh that showed a bug with CCD.
+		// The above test (SimpleTriMeshWithTwoHits_SweepCCD_ResultsAreExpected) is an extremely simplified version of this test.
+		// The root issue was that TOI was computed using the clipped distance when traversing the BVH instead of the original length, causing TOI to approach 1.
+		const FReal SweepLength = 109.97858833658735;
+		SphereRadius = 15.8750010;
+
+		FMeshInitData InitData;
+		InitData.TrimeshParticles = FTriangleMeshImplicitObject::ParticlesType
+		({
+			FVec3(-344.745148, -362.888794, -1.48117601e-06),
+			FVec3(-361.505798, -361.505798, -9.83252107e-07),
+			FVec3(-350.139618, -370.883759, -1.48008394e-06),
+			FVec3(-350.139618, -370.883759, 26.9759979),
+			FVec3(-337.094055, -377.964661, 26.9759979),
+			FVec3(-362.888794, -344.745148, -5.74946910e-07),
+			FVec3(-370.883759, -350.139587, -5.73328350e-07),
+			FVec3(-377.964661, -337.094025, -2.63812211e-07),
+			FVec3(-370.883759, -350.139587, 26.9760017),
+			FVec3(-361.505798, -361.505798, 26.9760017),
+			FVec3(-354.686462, -354.686462, -9.84611802e-07),
+			FVec3(-344.745148, -362.888794, 26.9759979),
+			FVec3(-354.686462, -354.686462, 26.9760017),
+			FVec3(-362.888794, -344.745148, 26.9760017),
+			FVec3(-349.114685, -369.364716, 28.7999992),
+			FVec3(-360.210144, -360.210114, 28.8000011),
+			FVec3(-345.770111, -364.407837, 28.7999992),
+			FVec3(-369.364716, -349.114655, 28.8000011),
+			FVec3(-364.407837, -345.770111, 28.8000011),
+			FVec3(-355.982147, -355.982147, 28.8000011),
+			FVec3(-337.094055, -377.964661, -2.05032484e-06),
+			FVec3(-320.754150, -372.984985, -2.68089730e-06),
+			FVec3(-322.677887, -382.439697, -2.68047484e-06),
+			FVec3(-322.677887, -382.439697, 26.9759979),
+			FVec3(-333.345245, -369.076477, -2.05111360e-06),
+			FVec3(-333.345245, -369.076477, 26.9759979),
+			FVec3(-320.754150, -372.984985, 26.9759979),
+			FVec3(-307.200012, -384.000000, 26.9759979),
+			FVec3(-306.717346, -374.399994, 26.9759979),
+			FVec3(-336.381775, -376.275909, 28.7999992),
+			FVec3(-322.312378, -380.643311, 28.7999992),
+			FVec3(-321.119629, -374.781403, 28.7999992),
+			FVec3(-334.057526, -370.765228, 28.7999992),
+			FVec3(-307.108276, -382.175995, 28.7999992),
+			FVec3(-306.809082, -376.224030, 28.7999992),
+			FVec3(-369.076477, -333.345215, -2.65682075e-07),
+			FVec3(-372.984955, -320.754150, -7.02783609e-08),
+			FVec3(-382.439697, -322.677887, -6.82030006e-08),
+			FVec3(-377.964661, -337.094025, 26.9760017),
+			FVec3(-369.076477, -333.345215, 26.9760017),
+			FVec3(-382.439697, -322.677887, 26.9760017),
+			FVec3(-384.000000, -307.200012, -6.82121053e-14),
+			FVec3(-374.399994, -306.717316, -5.29428446e-10),
+			FVec3(-372.984955, -320.754150, 26.9760017),
+			FVec3(-384.000000, -307.200012, 26.9760017),
+			FVec3(-374.399994, -306.717316, 26.9760017),
+			FVec3(-376.275909, -336.381744, 28.8000011),
+			FVec3(-370.765198, -334.057495, 28.8000011),
+			FVec3(-380.643280, -322.312378, 28.8000011),
+			FVec3(-374.781372, -321.119629, 28.8000011),
+			FVec3(-382.175995, -307.108276, 28.8000011),
+			FVec3(-376.224030, -306.809052, 28.8000011),
+		});
+		InitData.Indices =
+		{
+			{14, 9, 3},
+			{15, 14, 16},
+			{8, 9, 17},
+			{15, 18, 17},
+			{9, 14, 15},
+			{17, 9, 15},
+			{19, 15, 16},
+			{11, 19, 16},
+			{19, 11, 12},
+			{18, 15, 19},
+			{29, 3, 4},
+			{3, 29, 14},
+			{14, 29, 32},
+			{14, 32, 16},
+			{1, 0, 2},
+			{2, 3, 1},
+			{3, 2, 4},
+			{5, 1, 6},
+			{8, 6, 9},
+			{9, 6, 1},
+			{0, 1, 10},
+			{10, 11, 0},
+			{1, 5, 10},
+			{9, 1, 3},
+			{11, 10, 12},
+			{20, 4, 2},
+			{2, 24, 20},
+			{24, 2, 0},
+			{30, 4, 23},
+			{4, 30, 29},
+			{31, 29, 30},
+			{16, 25, 11},
+			{25, 16, 32},
+			{29, 31, 32},
+			{32, 26, 25},
+			{27, 30, 23},
+			{33, 31, 30},
+			{30, 27, 33},
+			{31, 33, 34},
+			{34, 26, 31},
+			{28, 26, 34},
+			{26, 32, 31},
+			{20, 21, 22},
+			{22, 23, 20},
+			{0, 25, 24},
+			{21, 20, 24},
+			{4, 20, 23},
+			{25, 0, 11},
+			{26, 24, 25},
+			{24, 26, 21},
+			{23, 22, 27},
+			{28, 21, 26},
+			{8, 7, 6},
+			{5, 12, 10},
+			{12, 5, 13},
+			{12, 13, 19},
+			{18, 19, 13},
+			{6, 35, 5},
+			{35, 6, 7},
+			{35, 13, 5},
+			{38, 7, 8},
+			{13, 35, 39},
+			{38, 8, 46},
+			{46, 8, 17},
+			{47, 17, 18},
+			{39, 18, 13},
+			{17, 47, 46},
+			{18, 39, 47},
+			{48, 38, 46},
+			{40, 38, 48},
+			{46, 49, 48},
+			{49, 46, 47},
+			{43, 47, 39},
+			{48, 44, 40},
+			{49, 50, 48},
+			{47, 43, 49},
+			{44, 48, 50},
+			{43, 51, 49},
+			{49, 51, 50},
+			{51, 43, 45},
+			{36, 7, 37},
+			{38, 37, 7},
+			{7, 36, 35},
+			{36, 39, 35},
+			{40, 37, 38},
+			{41, 36, 37},
+			{40, 41, 37},
+			{36, 41, 42},
+			{42, 43, 36},
+			{39, 36, 43},
+			{44, 41, 40},
+			{43, 42, 45},
+		};
+		InitData.AutoBuildMaterials();
+		InitData.AutoBuildExternalFaceIndices();
+		FTriangleMeshImplicitObject TriangleMesh(MoveTemp(InitData.TrimeshParticles), MoveTemp(InitData.Indices), MoveTemp(InitData.Materials), MoveTemp(InitData.ExternalFaceIndexMap), nullptr, true);
+
+		FVec3 SweepStart = FVec3(-277.57524933822191, -298.05159216244965, 12.770771045445827);
+		FVec3 SweepDir = FVec3(-0.92428685934217092, -0.38139225184986136, 0.015286707794476266);
+		FQuat Rotation(-0.00085530578740522696, -0.00082053338777325429, 0.70703569140501821, 0.70717683524049568);
+
+		const FReal ExpectedDistance = 81.848526000976562;
+		{
+			IgnorePenetration = 0;
+			TargetPenetration = 0;
+			FSweepCCDResultData Result = SweepSphereCCD(TriangleMesh, SweepStart, SweepDir, Rotation, SweepLength);
+
+			const FReal ExpectedTOI = (ExpectedDistance + TargetPenetration) / SweepLength;
+			EXPECT_EQ(Result.bResult, true);
+			EXPECT_NEAR(Result.TOI, ExpectedTOI, KINDA_SMALL_NUMBER);
+			EXPECT_NEAR(Result.Phi, 0, KINDA_SMALL_NUMBER);
+			EXPECT_VECTOR_NEAR(Result.Position, FVec3(-367.17901611328125, -336.84103393554688, 14.021965980529785), KINDA_SMALL_NUMBER);
+			EXPECT_VECTOR_NEAR(Result.Normal, FVec3(0.87888097408604637, 0.47704112337361676, -0.0000000000000000), KINDA_SMALL_NUMBER);
+			EXPECT_EQ(Result.FaceId, 61);
+		}
+		{
+			IgnorePenetration = 12.700000762939453;
+			TargetPenetration = 6.3500003814697266;
+			FSweepCCDResultData Result = SweepSphereCCD(TriangleMesh, SweepStart, SweepDir, Rotation, SweepLength);
+
+			const FReal ExpectedTOI = (ExpectedDistance + TargetPenetration) / SweepLength;
+			EXPECT_EQ(Result.bResult, true);
+			// Note: This test is slightly outside KINDA_SMALL_NUMBER in the epsilon check, as the way the TOI is computed is slightly different and loses a little numerical precision.
+			EXPECT_NEAR(Result.TOI, ExpectedTOI, 1.e-3f);
+			EXPECT_NEAR(Result.Phi, -TargetPenetration, KINDA_SMALL_NUMBER);
+			EXPECT_VECTOR_NEAR(Result.Position, FVec3(-367.17901611328125, -336.84103393554688, 14.021965980529785), KINDA_SMALL_NUMBER);
+			EXPECT_VECTOR_NEAR(Result.Normal, FVec3(0.87888097408604637, 0.47704112337361676, -0.0000000000000000), KINDA_SMALL_NUMBER);
+			EXPECT_EQ(Result.FaceId, 61);
+		}
+	}
+
 	TEST_F(FSphereSweepFixture, TransformedUnionContainingNonUniformScaledObject_SweepSphere_ResultsAreExpected)
 	{
 		// This test is to catch a bug where a Transformed union with a scaled object would 
diff --git a/Engine/Source/Runtime/Experimental/Chaos/Private/Chaos/TriangleMeshImplicitObject.cpp b/Engine/Source/Runtime/Experimental/Chaos/Private/Chaos/TriangleMeshImplicitObject.cpp
--- a/Engine/Source/Runtime/Experimental/Chaos/Private/Chaos/TriangleMeshImplicitObject.cpp
+++ b/Engine/Source/Runtime/Experimental/Chaos/Private/Chaos/TriangleMeshImplicitObject.cpp
@@ -1672,7 +1672,9 @@ struct FTriangleMeshSweepVisitorCCD
 
 			// Calculate the time to reach a depth of TargetPenetration
 			FReal TargetTOI, TargetPhi;
-			ComputeSweptContactTOIAndPhiAtTargetPenetration(DirDotNormal, ScaledCurLength, FReal(ScaledDistance), IgnorePenetration, TargetPenetration, TargetTOI, TargetPhi);
+			// Make sure to compute the TOI against the original length, not the current length.
+			const FReal ScaledLength = Length * LengthScale;
+			ComputeSweptContactTOIAndPhiAtTargetPenetration(DirDotNormal, ScaledLength, FReal(ScaledDistance), IgnorePenetration, TargetPenetration, TargetTOI, TargetPhi);
 
 			if (ScaledDistance < ScaledCurLength)
 			{
