[Parallel constraint solver] Phase 1 - solver partition manager greedy graph coloring & task-parallel constraint solve

FSolverPartitionManager (abstract interface):
- Subdivides the constraints into batches
- Determines which batches must run in series or can run in parallel.
- One partition manager per constraint container solver (i.e. constraint type)

FGreedyGraphColoring (derived from FSolverPartitionManager)
- Greedy graph coloring independent of islanding directly done on the constraint containers
- Constraints of the same color can be solved batch-parallel
- Constraints of different colors must be solved in series

FConstraintGroupSolver:
- Task-parallel solve of constraint batches
- Launches all constraint solves as tasks with a dependency hierarchy
- Dependency is provided by FSolverPartitionManager

#Synced-CL 52071161

diff --git a/Engine/Source/Programs/HeadlessChaos/Private/HeadlessChaosTestJoints.cpp b/Engine/Source/Programs/HeadlessChaos/Private/HeadlessChaosTestJoints.cpp
--- a/Engine/Source/Programs/HeadlessChaos/Private/HeadlessChaosTestJoints.cpp
+++ b/Engine/Source/Programs/HeadlessChaos/Private/HeadlessChaosTestJoints.cpp
@@ -344,7 +344,7 @@ namespace ChaosTest {
 	 * Pendulum with animated root.
 	 */
 	template <typename TEvolution>
-	void JointConstraint_ShortChainAnimated()
+	void JointConstraint_ShortChainAnimated(const bool bProjectionEnabled)
 	{
 		const int32 NumIterations = 10;
 		const FReal Gravity = 980;
@@ -389,6 +389,13 @@ namespace ChaosTest {
 
 		Test.Create();
 
+		for (FPBDJointConstraintHandle* Joint : Test.Joints)
+		{
+			FPBDJointSettings Settings = Joint->GetSettings();
+			Settings.bProjectionEnabled = bProjectionEnabled;
+			Joint->SetSettings(Settings);
+		}
+
 		const FVec3 Box2LocalSpaceJointPosition = Test.JointPositions[0] - Test.ParticlePositions[1];
 
 		FReal MaxDistanceError = 0.0f;
@@ -429,7 +436,7 @@ namespace ChaosTest {
 	 * Pendulum with animated root.
 	 */
 	template <typename TEvolution>
-	void JointConstraint_LongChainAnimated()
+	void JointConstraint_LongChainAnimated(const bool bProjectionEnabled)
 	{
 		const int NumParticles = 10;
 		const int32 NumIterations = 20;
@@ -475,6 +482,13 @@ namespace ChaosTest {
 
 		Test.Create();
 
+		for (FPBDJointConstraintHandle* Joint : Test.Joints)
+		{
+			FPBDJointSettings Settings = Joint->GetSettings();
+			Settings.bProjectionEnabled = bProjectionEnabled;
+			Joint->SetSettings(Settings);
+		}
+
 		const FVec3 Box2LocalSpaceJointPosition = Test.JointPositions[0] - Test.ParticlePositions[1];
 
 		FReal MaxDistanceError = 0.0f;
@@ -666,7 +680,7 @@ namespace ChaosTest {
 		{
 			LinearJointTest.JointPositions.Add(LinearJointTest.ParticlePositions[JointIndex]);
 			LinearJointTest.JointParticleIndices.Add({ JointIndex, JointIndex + 1 });
-			NonLinearJointTest.JointPositions.Add(LinearJointTest.ParticlePositions[JointIndex]);
+			NonLinearJointTest.JointPositions.Add(NonLinearJointTest.ParticlePositions[JointIndex]);
 			NonLinearJointTest.JointParticleIndices.Add({ JointIndex, JointIndex + 1 });
 		}
 
@@ -721,7 +735,28 @@ namespace ChaosTest {
 			const FReal ExpectedDistance = (LinearJointTest.ParticlePositions[ParticleIndex2] - LinearJointTest.ParticlePositions[ParticleIndex1]).Size();
 			const FReal LinearDistanceError = FMath::Abs(ExpectedDistance - LinearDistance);
 			const FReal NonLinearDistanceError = FMath::Abs(ExpectedDistance - NonLinearDistance);
-			EXPECT_LT(NonLinearDistanceError * 10, LinearDistanceError);
+			EXPECT_LT(NonLinearDistanceError, LinearDistanceError);
+
+			//// Uncomment for debugging
+			//printf("ParticleIndex1: %d\n", ParticleIndex1);
+			//printf("Linear Particle1 X: %f, %f, %f\n", LinearJointTest.GetParticle(ParticleIndex1)->GetX().X, LinearJointTest.GetParticle(ParticleIndex1)->GetX().Y, LinearJointTest.GetParticle(ParticleIndex1)->GetX().Z);
+			//printf("Linear Particle1 P: %f, %f, %f\n", LinearJointTest.GetParticle(ParticleIndex1)->CastToRigidParticle()->GetP().X, LinearJointTest.GetParticle(ParticleIndex1)->CastToRigidParticle()->GetP().Y, LinearJointTest.GetParticle(ParticleIndex1)->CastToRigidParticle()->GetP().Z);
+			//printf("Nonlinear Particle1 X: %f, %f, %f\n", NonLinearJointTest.GetParticle(ParticleIndex1)->GetX().X, NonLinearJointTest.GetParticle(ParticleIndex1)->GetX().Y, NonLinearJointTest.GetParticle(ParticleIndex1)->GetX().Z);
+			//printf("Nonlinear Particle1 P: %f, %f, %f\n", NonLinearJointTest.GetParticle(ParticleIndex1)->CastToRigidParticle()->GetP().X, NonLinearJointTest.GetParticle(ParticleIndex1)->CastToRigidParticle()->GetP().Y, NonLinearJointTest.GetParticle(ParticleIndex1)->CastToRigidParticle()->GetP().Z);
+			//printf("Expected Particle1: %f, %f, %f\n", LinearJointTest.ParticlePositions[ParticleIndex1].X, LinearJointTest.ParticlePositions[ParticleIndex1].Y, LinearJointTest.ParticlePositions[ParticleIndex1].Z);
+			//printf("ParticleIndex2: %d\n", ParticleIndex2);
+			//printf("Linear Particle2 X: %f, %f, %f\n", LinearJointTest.GetParticle(ParticleIndex2)->GetX().X, LinearJointTest.GetParticle(ParticleIndex2)->GetX().Y, LinearJointTest.GetParticle(ParticleIndex2)->GetX().Z);
+			//printf("Linear Particle2 P: %f, %f, %f\n", LinearJointTest.GetParticle(ParticleIndex2)->CastToRigidParticle()->GetP().X, LinearJointTest.GetParticle(ParticleIndex2)->CastToRigidParticle()->GetP().Y, LinearJointTest.GetParticle(ParticleIndex2)->CastToRigidParticle()->GetP().Z);
+			//printf("Nonlinear Particle2 X: %f, %f, %f\n", NonLinearJointTest.GetParticle(ParticleIndex2)->GetX().X, NonLinearJointTest.GetParticle(ParticleIndex2)->GetX().Y, NonLinearJointTest.GetParticle(ParticleIndex2)->GetX().Z);
+			//printf("Nonlinear Particle2 P: %f, %f, %f\n", NonLinearJointTest.GetParticle(ParticleIndex2)->CastToRigidParticle()->GetP().X, NonLinearJointTest.GetParticle(ParticleIndex2)->CastToRigidParticle()->GetP().Y, NonLinearJointTest.GetParticle(ParticleIndex2)->CastToRigidParticle()->GetP().Z);
+			//printf("Expected Particle2: %f, %f, %f\n", LinearJointTest.ParticlePositions[ParticleIndex2].X, LinearJointTest.ParticlePositions[ParticleIndex2].Y, LinearJointTest.ParticlePositions[ParticleIndex2].Z);
+			//printf("JointIndex: %d\n", JointIndex);
+			//printf("LinearDistanceError: %f\n", LinearDistanceError);
+			//printf("NonLinearDistanceError: %f\n", NonLinearDistanceError);
+			//printf("LinearDistance: %f\n", LinearDistance);
+			//printf("NonLinearDistance: %f\n", NonLinearDistance);
+			//printf("ExpectedDistance: %f\n", ExpectedDistance);
+			//printf("\n");
 		}
 	}
 
@@ -925,12 +960,20 @@ namespace ChaosTest {
 		JointConstraint_SingleAnimated<FPBDRigidsEvolutionGBF>();
 	}
 
-	GTEST_TEST(AllEvolutions, JointTests_TestShortJointChainWithAnimatedRoot) {
-		JointConstraint_ShortChainAnimated<FPBDRigidsEvolutionGBF>();
+	GTEST_TEST(AllEvolutions, JointTests_TestShortJointChainWithAnimatedRootProjectionEnabled) {
+		JointConstraint_ShortChainAnimated<FPBDRigidsEvolutionGBF>(true);
+	}	
+	
+	GTEST_TEST(AllEvolutions, JointTests_TestShortJointChainWithAnimatedRootProjectionDisabled) {
+		JointConstraint_ShortChainAnimated<FPBDRigidsEvolutionGBF>(false);
+	}
+
+	GTEST_TEST(AllEvolutions, JointTests_TestLongJointChainWithAnimatedRootProjectionEnabled) {
+		JointConstraint_LongChainAnimated<FPBDRigidsEvolutionGBF>(true);
 	}
 
-	GTEST_TEST(AllEvolutions, JointTests_TestLongJointChainWithAnimatedRoot) {
-		JointConstraint_LongChainAnimated<FPBDRigidsEvolutionGBF>();
+	GTEST_TEST(AllEvolutions, JointTests_TestLongJointChainWithAnimatedRootProjectionDisabled) {
+		JointConstraint_LongChainAnimated<FPBDRigidsEvolutionGBF>(false);
 	}
 
 	GTEST_TEST(AllEvolutions, JointTests_TestLongJointChainWithVaryingIterations) {
diff --git a/Engine/Source/Programs/HeadlessChaos/Private/HeadlessChaosTestSolverPartitions.cpp b/Engine/Source/Programs/HeadlessChaos/Private/HeadlessChaosTestSolverPartitions.cpp
new file mode 100644
--- /dev/null
+++ b/Engine/Source/Programs/HeadlessChaos/Private/HeadlessChaosTestSolverPartitions.cpp
@@ -0,0 +1,850 @@
+// Copyright Epic Games, Inc. All Rights Reserved.
+
+#include "HeadlessChaos.h"
+#include "HeadlessChaosTestUtility.h"
+
+#include "Chaos/Character/CharacterGroundConstraintContainer.h"
+#include "Chaos/Island/IslandManager.h"
+#include "Chaos/PBDRigidsSOAs.h"
+#include "Chaos/Evolution/ConstraintGroupSolver.h"
+#include "Chaos/Evolution/IterationSettings.h"
+
+#include "ChaosSolversModule.h"
+#include "PBDRigidsSolver.h"
+#include "PhysicsProxy/SingleParticlePhysicsProxy.h"
+#include "PhysicsProxy/JointConstraintProxy.h"
+#include "Chaos/PBDJointConstraintData.h"
+#include "HAL/ConsoleManager.h"
+#include "PhysicsInterfaceDeclaresCore.h"
+
+namespace ChaosTest
+{
+	using namespace Chaos;
+
+	class SolverPartitionTest : public ::testing::Test
+	{
+	protected:
+		SolverPartitionTest()
+			: SOAs(UniqueIndices)
+		{
+			Solver = Container.CreateSceneSolver(0);
+		}
+
+		void SetUp()
+		{
+			CVarGraphColoring = IConsoleManager::Get().FindConsoleVariable(TEXT("p.Chaos.SolverPartitionManager.Partitioning"), false);
+			check(CVarGraphColoring && CVarGraphColoring->IsVariableInt());
+			GraphColoring_EditorDefault = CVarGraphColoring->GetInt();
+			CVarGraphColoring->Set(GraphColoring_UnitTestDefault);
+
+			CVarBatchSize = IConsoleManager::Get().FindConsoleVariable(TEXT("p.Chaos.SolverPartitionManager.BatchSize"), false);
+			check(CVarBatchSize && CVarBatchSize->IsVariableInt());
+			BatchSize_EditorDefault = CVarBatchSize->GetInt();
+			CVarBatchSize->Set(BatchSize_UnitTestDefault);
+
+			CVarMinConstraintsForColoring = IConsoleManager::Get().FindConsoleVariable(TEXT("p.Chaos.SolverPartitionManager.MinConstraintsForColoring"), false);
+			check(CVarMinConstraintsForColoring && CVarMinConstraintsForColoring->IsVariableInt());
+			MinConstraintsForColoring_EditorDefault = CVarMinConstraintsForColoring->GetInt();
+			CVarMinConstraintsForColoring->Set(MinConstraintsForColoring_UnitTestDefault);
+		}
+
+		void TearDown()
+		{
+			CVarGraphColoring->Set(GraphColoring_EditorDefault);
+			CVarBatchSize->Set(BatchSize_EditorDefault);
+			CVarMinConstraintsForColoring->Set(MinConstraintsForColoring_EditorDefault);
+		}
+
+		// Default values for unit testing
+		int32 GraphColoring_UnitTestDefault = 1;
+		int32 BatchSize_UnitTestDefault = 1;
+		int32 MinConstraintsForColoring_UnitTestDefault = 1;
+
+		// Storage for default values hard-coded in SolverPartitionManager.cpp
+		int32 GraphColoring_EditorDefault;
+		int32 BatchSize_EditorDefault;
+		int32 MinConstraintsForColoring_EditorDefault;
+
+		// Storage of console variable pointer
+		IConsoleVariable* CVarGraphColoring;
+		IConsoleVariable* CVarBatchSize;
+		IConsoleVariable* CVarMinConstraintsForColoring;
+
+		FParticleUniqueIndicesMultithreaded UniqueIndices;
+		FPBDRigidsSOAs SOAs;
+		FCharacterGroundConstraintContainer Container;
+		TArray<FCharacterGroundConstraintSettings> Settings;
+		TArray<FCharacterGroundConstraintDynamicData> Data;
+		Private::FIterationSettings IterationSettings;
+		TUniquePtr<FConstraintContainerSolver> Solver;
+	};
+
+	// Test coloring of constraint graph
+	// Test 1
+	//	Ak - Bd - Cd - Dd - Ed
+	//	   0    1    0    1
+	TEST_F(SolverPartitionTest, GraphColoring_Chain)
+	{
+		TArray<FPBDRigidParticleHandle*> Particles = SOAs.CreateDynamicParticles(5);
+		Settings.SetNum(1);
+		Data.SetNum(1);
+
+		TArray<FCharacterGroundConstraintHandle*> Constraints;
+		for (int32 ParticleIndex = 0; ParticleIndex < Particles.Num() - 1; ++ParticleIndex)
+		{
+			Constraints.Add(Container.AddConstraint(Settings[0], Data[0], Particles[ParticleIndex], Particles[ParticleIndex + 1]));
+		}
+
+		Solver->AddConstraints();
+		EXPECT_EQ(Solver->GetNumConstraints(), Constraints.Num());
+
+		PRAGMA_DISABLE_INTERNAL_WARNINGS
+		const int32 MaxSerialBatches = Solver->PrepareSolverPartitions();
+		const TArray<TArray<int32>>& SerialIndices = Solver->GetConstraintsBatchIndices();
+		PRAGMA_ENABLE_INTERNAL_WARNINGS
+
+		EXPECT_EQ(SerialIndices.Num(), 2);
+		EXPECT_EQ(Solver->GetNumConstraints(), 4);
+
+		if (SerialIndices.Num() > 0 && SerialIndices[0].Num() > 0)
+			EXPECT_EQ(SerialIndices[0][0], 0);
+		if (SerialIndices.Num() > 1 && SerialIndices[1].Num() > 0)
+			EXPECT_EQ(SerialIndices[1][0], 2);
+	}
+
+	// Test coloring of constraint graph
+	// Test 1
+	//				    Gd
+	//				 0/	  \3
+	//		        Ed     Fd
+	//            1/ \2	 1/ \2
+	//            Bd   Cd   Dd
+	//            0\   |0   /0
+	//                 Ak
+	TEST_F(SolverPartitionTest, GraphColoring_WallWithSpacing)
+	{
+		TArray<FPBDRigidParticleHandle*> Particles = SOAs.CreateDynamicParticles(7);
+		Settings.SetNum(1);
+		Data.SetNum(1);
+
+		TArray<FCharacterGroundConstraintHandle*> Constraints;
+
+		Particles[0]->SetObjectStateLowLevel(EObjectStateType::Kinematic);
+
+		Constraints.Add(Container.AddConstraint(Settings[0], Data[0], Particles[0], Particles[1]));
+		Constraints.Add(Container.AddConstraint(Settings[0], Data[0], Particles[0], Particles[2]));
+		Constraints.Add(Container.AddConstraint(Settings[0], Data[0], Particles[0], Particles[3]));
+		Constraints.Add(Container.AddConstraint(Settings[0], Data[0], Particles[1], Particles[4]));
+		Constraints.Add(Container.AddConstraint(Settings[0], Data[0], Particles[2], Particles[4]));
+		Constraints.Add(Container.AddConstraint(Settings[0], Data[0], Particles[2], Particles[5]));
+		Constraints.Add(Container.AddConstraint(Settings[0], Data[0], Particles[3], Particles[5]));
+		Constraints.Add(Container.AddConstraint(Settings[0], Data[0], Particles[4], Particles[6]));
+		Constraints.Add(Container.AddConstraint(Settings[0], Data[0], Particles[5], Particles[6]));
+
+		Solver->AddConstraints();
+		EXPECT_EQ(Solver->GetNumConstraints(), Constraints.Num());
+
+		PRAGMA_DISABLE_INTERNAL_WARNINGS
+		const int32 MaxSerialBatches = Solver->PrepareSolverPartitions();
+		const TArray<TArray<int32>>& SerialIndices = Solver->GetConstraintsBatchIndices();
+		PRAGMA_ENABLE_INTERNAL_WARNINGS
+
+		EXPECT_EQ(SerialIndices.Num(), 4);
+		EXPECT_EQ(Solver->GetNumConstraints(), 9);
+
+		if (SerialIndices.Num() > 0 && SerialIndices[0].Num() > 0)
+			EXPECT_EQ(SerialIndices[0][0], 0);
+		if (SerialIndices.Num() > 1 && SerialIndices[1].Num() > 0)
+			EXPECT_EQ(SerialIndices[1][0], 4);
+		if (SerialIndices.Num() > 2 && SerialIndices[2].Num() > 0)
+			EXPECT_EQ(SerialIndices[2][0], 6);
+		if (SerialIndices.Num() > 3 && SerialIndices[3].Num() > 0)
+			EXPECT_EQ(SerialIndices[3][0], 8);
+	}
+
+	// Test coloring of constraint graph
+	// Test 1
+	//				    Gd
+	//				 1/	  \2
+	//		        Ed -0  Fd
+	//            2/  \3 4/  \1
+	//            Bd -1 Cd -2 Dd
+	//            0\   |0   /0
+	//                 Ak
+	TEST_F(SolverPartitionTest, GraphColoring_WallWithoutSpacing)
+	{
+		TArray<FPBDRigidParticleHandle*> Particles = SOAs.CreateDynamicParticles(7);
+		Settings.SetNum(1);
+		Data.SetNum(1);
+
+		TArray<FCharacterGroundConstraintHandle*> Constraints;
+
+		Particles[0]->SetObjectStateLowLevel(EObjectStateType::Kinematic);
+
+		Constraints.Add(Container.AddConstraint(Settings[0], Data[0], Particles[0], Particles[1]));
+		Constraints.Add(Container.AddConstraint(Settings[0], Data[0], Particles[0], Particles[2]));
+		Constraints.Add(Container.AddConstraint(Settings[0], Data[0], Particles[0], Particles[3]));
+		Constraints.Add(Container.AddConstraint(Settings[0], Data[0], Particles[1], Particles[2]));
+		Constraints.Add(Container.AddConstraint(Settings[0], Data[0], Particles[2], Particles[3]));
+		Constraints.Add(Container.AddConstraint(Settings[0], Data[0], Particles[1], Particles[4]));
+		Constraints.Add(Container.AddConstraint(Settings[0], Data[0], Particles[2], Particles[4]));
+		Constraints.Add(Container.AddConstraint(Settings[0], Data[0], Particles[2], Particles[5]));
+		Constraints.Add(Container.AddConstraint(Settings[0], Data[0], Particles[3], Particles[5]));
+		Constraints.Add(Container.AddConstraint(Settings[0], Data[0], Particles[4], Particles[5]));
+		Constraints.Add(Container.AddConstraint(Settings[0], Data[0], Particles[4], Particles[6]));
+		Constraints.Add(Container.AddConstraint(Settings[0], Data[0], Particles[5], Particles[6]));
+
+		Solver->AddConstraints();
+		EXPECT_EQ(Solver->GetNumConstraints(), Constraints.Num());
+
+		PRAGMA_DISABLE_INTERNAL_WARNINGS
+		const int32 MaxSerialBatches = Solver->PrepareSolverPartitions();
+		const TArray<TArray<int32>>& SerialIndices = Solver->GetConstraintsBatchIndices();
+		PRAGMA_ENABLE_INTERNAL_WARNINGS
+
+		EXPECT_EQ(SerialIndices.Num(), 5);
+		EXPECT_EQ(Solver->GetNumConstraints(), 12);
+
+		if (SerialIndices.Num() > 0 && SerialIndices[0].Num() > 0)
+			EXPECT_EQ(SerialIndices[0][0], 0);
+		if (SerialIndices.Num() > 1 && SerialIndices[1].Num() > 0)
+			EXPECT_EQ(SerialIndices[1][0], 4);
+		if (SerialIndices.Num() > 2 && SerialIndices[2].Num() > 0)
+			EXPECT_EQ(SerialIndices[2][0], 7);
+		if (SerialIndices.Num() > 3 && SerialIndices[3].Num() > 0)
+			EXPECT_EQ(SerialIndices[3][0], 10);
+		if (SerialIndices.Num() > 4 && SerialIndices[4].Num() > 0)
+			EXPECT_EQ(SerialIndices[4][0], 11);
+	}
+
+	// Test coloring of constraint graph
+	// Test 1
+	//			 Hd 0- Kd 2- Ld
+	//			  1|   1|   0|
+	//		     Ed 0- Fd 2- Gd
+	//            2|   3|   1|
+	//           Bd 1- Cd 2- Dd
+	//             0\  0|   /0
+	//                 Ak
+	//
+	TEST_F(SolverPartitionTest, GraphColoring_BlockWithoutSpacing)
+	{
+		TArray<FPBDRigidParticleHandle*> Particles = SOAs.CreateDynamicParticles(10);
+		Settings.SetNum(1);
+		Data.SetNum(1);
+
+		TArray<FCharacterGroundConstraintHandle*> Constraints;
+
+		Particles[0]->SetObjectStateLowLevel(EObjectStateType::Kinematic);
+
+		// Constraints with the ground
+		Constraints.Add(Container.AddConstraint(Settings[0], Data[0], Particles[0], Particles[1]));
+		Constraints.Add(Container.AddConstraint(Settings[0], Data[0], Particles[0], Particles[2]));
+		Constraints.Add(Container.AddConstraint(Settings[0], Data[0], Particles[0], Particles[3]));
+		// Horizontal constraints
+		Constraints.Add(Container.AddConstraint(Settings[0], Data[0], Particles[1], Particles[2]));
+		Constraints.Add(Container.AddConstraint(Settings[0], Data[0], Particles[2], Particles[3]));
+		// Vertical constraints
+		Constraints.Add(Container.AddConstraint(Settings[0], Data[0], Particles[1], Particles[4]));
+		Constraints.Add(Container.AddConstraint(Settings[0], Data[0], Particles[2], Particles[5]));
+		Constraints.Add(Container.AddConstraint(Settings[0], Data[0], Particles[3], Particles[6]));
+		// Horizontal constraints
+		Constraints.Add(Container.AddConstraint(Settings[0], Data[0], Particles[4], Particles[5]));
+		Constraints.Add(Container.AddConstraint(Settings[0], Data[0], Particles[5], Particles[6]));
+		// Vertical constraints
+		Constraints.Add(Container.AddConstraint(Settings[0], Data[0], Particles[4], Particles[7]));
+		Constraints.Add(Container.AddConstraint(Settings[0], Data[0], Particles[5], Particles[8]));
+		Constraints.Add(Container.AddConstraint(Settings[0], Data[0], Particles[6], Particles[9]));
+		// Horizontal constraints
+		Constraints.Add(Container.AddConstraint(Settings[0], Data[0], Particles[7], Particles[8]));
+		Constraints.Add(Container.AddConstraint(Settings[0], Data[0], Particles[8], Particles[9]));
+
+		Solver->AddConstraints();
+		EXPECT_EQ(Solver->GetNumConstraints(), Constraints.Num());
+
+		PRAGMA_DISABLE_INTERNAL_WARNINGS
+		const int32 MaxSerialBatches = Solver->PrepareSolverPartitions();
+		const TArray<TArray<int32>> SerialIndices = Solver->GetConstraintsBatchIndices();
+		PRAGMA_ENABLE_INTERNAL_WARNINGS
+
+		EXPECT_EQ(SerialIndices.Num(), 4);
+		EXPECT_EQ(Solver->GetNumConstraints(), 15);
+
+		if (SerialIndices.Num() > 0 && SerialIndices[0].Num() > 0)
+			EXPECT_EQ(SerialIndices[0][0], 0);
+		if (SerialIndices.Num() > 1 && SerialIndices[1].Num() > 0)
+			EXPECT_EQ(SerialIndices[1][0], 6);
+		if (SerialIndices.Num() > 2 && SerialIndices[2].Num() > 0)
+			EXPECT_EQ(SerialIndices[2][0], 10);
+		if (SerialIndices.Num() > 3 && SerialIndices[3].Num() > 0)
+			EXPECT_EQ(SerialIndices[3][0], 14);
+	}
+}
+
+namespace ChaosTest
+{
+	using namespace Chaos;
+
+	class PartitioningSolverTest : public ::testing::Test
+	{
+	public:
+		void SetUp()
+		{
+			CVarParallelSolve = IConsoleManager::Get().FindConsoleVariable(TEXT("p.Chaos.SolverPartitionManager.ParallelSolve"), false);
+			check(CVarParallelSolve && CVarParallelSolve->IsVariableInt());
+			ParallelSolve_EditorDefault = CVarParallelSolve->GetInt();
+			CVarParallelSolve->Set(ParallelSolve_UnitTestDefault);
+
+			CVarGraphColoring = IConsoleManager::Get().FindConsoleVariable(TEXT("p.Chaos.SolverPartitionManager.Partitioning"), false);
+			check(CVarGraphColoring && CVarGraphColoring->IsVariableInt());
+			GraphColoring_EditorDefault = CVarGraphColoring->GetInt();
+			CVarGraphColoring->Set(GraphColoring_UnitTestDefault);
+
+			CVarBatchSize = IConsoleManager::Get().FindConsoleVariable(TEXT("p.Chaos.SolverPartitionManager.BatchSize"), false);
+			check(CVarBatchSize && CVarBatchSize->IsVariableInt());
+			BatchSize_EditorDefault = CVarBatchSize->GetInt();
+			CVarBatchSize->Set(BatchSize_UnitTestDefault);
+
+			CVarMinConstraintsForColoring = IConsoleManager::Get().FindConsoleVariable(TEXT("p.Chaos.SolverPartitionManager.MinConstraintsForColoring"), false);
+			check(CVarMinConstraintsForColoring && CVarMinConstraintsForColoring->IsVariableInt());
+			MinConstraintsForColoring_EditorDefault = CVarMinConstraintsForColoring->GetInt();
+			CVarMinConstraintsForColoring->Set(MinConstraintsForColoring_UnitTestDefault);
+		}
+
+		void TearDown()
+		{
+			CVarParallelSolve->Set(ParallelSolve_EditorDefault);
+			CVarGraphColoring->Set(GraphColoring_EditorDefault);
+			CVarBatchSize->Set(BatchSize_EditorDefault);
+			CVarMinConstraintsForColoring->Set(MinConstraintsForColoring_EditorDefault);
+		}
+
+		// Default values for unit testing
+		int32 ParallelSolve_UnitTestDefault = 1;
+		int32 GraphColoring_UnitTestDefault = 1;
+		int32 BatchSize_UnitTestDefault = 1;
+		int32 MinConstraintsForColoring_UnitTestDefault = 1;
+
+		// Storage for default values hard-coded in SolverPartitionManager.cpp
+		int32 ParallelSolve_EditorDefault;
+		int32 GraphColoring_EditorDefault;
+		int32 BatchSize_EditorDefault;
+		int32 MinConstraintsForColoring_EditorDefault;
+
+		// Storage of console variable pointer
+		IConsoleVariable* CVarParallelSolve;
+		IConsoleVariable* CVarGraphColoring;
+		IConsoleVariable* CVarBatchSize;
+		IConsoleVariable* CVarMinConstraintsForColoring;
+
+
+	};
+
+	class FPartitioningSolverTest
+	{
+	public:
+		FPartitioningSolverTest()
+			: MaterialData()
+			, FloorProxy()
+			, ParticleProxies()
+			, ParticleHandles()
+			, JointHandles()
+			, TickCount(0)
+		{
+			// Create a new material before creating the solver
+			FMaterialHandle MaterialHandle = FPhysicalMaterialManager::Get().Create();
+			FPhysicsMaterial* PhysicsMaterial = MaterialHandle.Get();
+			PhysicsMaterial->Friction = 0.7;
+			PhysicsMaterial->StaticFriction = 0.0;
+			PhysicsMaterial->Restitution = 0.0;
+			PhysicsMaterial->Density = 1.0;
+			PhysicsMaterial->SleepingLinearThreshold = 1.0;
+			PhysicsMaterial->SleepingAngularThreshold = 0.05;
+			PhysicsMaterial->SleepCounterThreshold = 4;
+			MaterialData.Materials.Add(MaterialHandle);
+
+			Module = FChaosSolversModule::GetModule();
+			Solver = Module->CreateSolver(nullptr, /*AsyncDt=*/-1);
+			InitSolverSettings(Solver);
+			Solver->SetThreadingMode_External(EThreadingModeTemp::SingleThread);
+		}
+
+		~FPartitioningSolverTest()
+		{
+
+			for (FSingleParticlePhysicsProxy* ParticleProxy : ParticleProxies)
+			{
+				Solver->UnregisterObject(ParticleProxy);
+			}
+			for (FJointConstraint* Joint : JointHandles)
+			{
+				if (Joint)
+					Solver->UnregisterObject(Joint);
+			}
+			if (FloorProxy)
+				Solver->UnregisterObject(FloorProxy);
+			Module->DestroySolver(Solver);
+		}
+
+		FRigidBodyHandle_External* MakeFloor(const FVec3& Position = FVec3(0))
+		{
+			FloorProxy = FSingleParticlePhysicsProxy::Create(Chaos::FGeometryParticle::CreateParticle());
+			FRigidBodyHandle_External* FloorParticle = &FloorProxy->GetGameThreadAPI();
+			FImplicitObjectPtr FloorGeom = Chaos::FImplicitObjectPtr(new TBox<FReal, 3>(FVec3(-5000, -5000, -100), FVec3(5000, 5000, 0)));
+			FloorParticle->SetGeometry(FloorGeom);
+			FloorParticle->SetObjectState(EObjectStateType::Static);
+			Solver->RegisterObject(FloorProxy);
+			FloorParticle->SetX(Position);
+			ChaosTest::SetParticleSimDataToCollide({ FloorProxy->GetParticle_LowLevel() });
+			for (const TUniquePtr<FPerShapeData>& Shape : FloorParticle->ShapesArray())
+			{
+				Shape->SetMaterialData(MaterialData);
+			}
+			return FloorParticle;
+		}
+
+		FRigidBodyHandle_External* MakeCube(const FVec3& Position, const FReal Size = 200, const FReal Mass = 1)
+		{
+			FSingleParticlePhysicsProxy* CubeProxy = FSingleParticlePhysicsProxy::Create(Chaos::FPBDRigidParticle::CreateParticle());
+			FRigidBodyHandle_External* CubeParticle = &CubeProxy->GetGameThreadAPI();
+			const FReal HalfSize = Size / 2.0f;
+			FImplicitObjectPtr CollidingCubeGeom = Chaos::FImplicitObjectPtr(new TBox<FReal, 3>(FVec3(-HalfSize), FVec3(HalfSize)));
+			CubeParticle->SetGeometry(CollidingCubeGeom);
+			Solver->RegisterObject(CubeProxy);
+			CubeParticle->SetGravityEnabled(true);
+			CubeParticle->SetX(Position);
+			SetCubeInertiaTensor(*CubeParticle, /*Dimension=*/Size, /*Mass=*/Mass);
+			ChaosTest::SetParticleSimDataToCollide({ CubeProxy->GetParticle_LowLevel() });
+			for (const TUniquePtr<FPerShapeData>& Shape : CubeParticle->ShapesArray())
+			{
+				Shape->SetMaterialData(MaterialData);
+			}
+
+			ParticleProxies.Add(CubeProxy);
+			ParticleHandles.Add(CubeParticle);
+
+			return CubeParticle;
+		}
+
+		FRigidBodyHandle_External* MakeSphere(const FVec3& Position, const FReal Radius = 100, const FReal Mass = 1)
+		{
+			FSingleParticlePhysicsProxy* SphereProxy = FSingleParticlePhysicsProxy::Create(Chaos::FPBDRigidParticle::CreateParticle());
+			FRigidBodyHandle_External* SphereParticle = &SphereProxy->GetGameThreadAPI();
+			FImplicitObjectPtr CollidingSphereGeom = Chaos::FImplicitObjectPtr(new TSphere<FReal, 3>(FVec3(0), Radius));
+			SphereParticle->SetGeometry(CollidingSphereGeom);
+			Solver->RegisterObject(SphereProxy);
+			SphereParticle->SetGravityEnabled(true);
+			SphereParticle->SetX(Position);
+			SetSphereInertiaTensor(*SphereParticle, /*Radius=*/Radius, /*Mass=*/Mass);
+			ChaosTest::SetParticleSimDataToCollide({ SphereProxy->GetParticle_LowLevel() });
+			for (const TUniquePtr<FPerShapeData>& Shape : SphereParticle->ShapesArray())
+			{
+				Shape->SetMaterialData(MaterialData);
+			}
+
+			ParticleProxies.Add(SphereProxy);
+			ParticleHandles.Add(SphereParticle);
+
+			return SphereParticle;
+		}
+
+		FJointConstraint* AddJoint(const FProxyBasePair& InConstrainedParticles, const FTransformPair& InTransform)
+		{
+			FJointConstraint* Joint = new FJointConstraint();
+			Joint->SetParticleProxies(InConstrainedParticles);
+
+			FPBDJointSettings SettingsTemp = Joint->GetJointSettings();
+			SettingsTemp.ConnectorTransforms = InTransform;
+			Joint->SetJointSettings(SettingsTemp);
+
+			Solver->RegisterObject(Joint);
+
+			JointHandles.Add(Joint);
+			return Joint;
+		}
+
+		void MakeStackOfCubes(const int32 Num, const FReal Size = 200, const FVec3 Pos = FVec3(0))
+		{
+			const FReal InitialHeight = Size / 2.0f;
+			for (int32 Id = 0; Id < Num; ++Id)
+			{
+				const FReal Z = InitialHeight + Size * Id;
+				MakeCube(Pos + FVec3(0, 0, Z), Size);
+			}
+		}
+
+		void MakeChainOfSpheres(const int32 Num, const FReal Radius = 100, const FVec3 Pos = FVec3(0))
+		{
+			const int32 DistanceMultiplier = 3;
+			TArray<FSingleParticlePhysicsProxy*> Proxies;
+			for (int32 Id = 0; Id < Num; ++Id)
+			{
+				const FReal Z = -DistanceMultiplier * Radius * Id;
+				FRigidBodyHandle_External* SphereParticle = MakeSphere(Pos + FVec3(0, 0, Z), Radius);
+				Proxies.Add(SphereParticle->GetProxy());
+			}
+			for (int32 Id = 0; Id < Num - 1; ++Id)
+			{
+				FRigidTransform3 Transform1 = FRigidTransform3(FVec3(0, 0, -DistanceMultiplier * Radius * 0.5), FRotation3::FromIdentity());
+				FRigidTransform3 Transform2 = FRigidTransform3(FVec3(0, 0, DistanceMultiplier * Radius * 0.5), FRotation3::FromIdentity());
+				AddJoint({ Proxies[Id], Proxies[Id + 1] }, { Transform1 , Transform2 });
+			}
+
+		}
+
+		void Advance()
+		{
+			Solver->AdvanceAndDispatch_External(1 / 60.0);
+			Solver->UpdateGameThreadStructures();
+			++TickCount;
+		}
+
+		void AdvanceUntilSleeping(const int32 MaxIterations = 200)
+		{
+			const int32 MaxTickCount = TickCount + MaxIterations;
+			bool bIsSleeping = false;
+			while (!bIsSleeping && (TickCount < MaxTickCount))
+			{
+				Advance();
+
+				bIsSleeping = true;
+				for (FRigidBodyHandle_External* ParticleHandle : ParticleHandles)
+				{
+					// Check that none of the particles is dynamic
+					if (ParticleHandle->ObjectState() == EObjectStateType::Dynamic)
+					{
+						bIsSleeping = false;
+					}
+				}
+			}
+
+			EXPECT_TRUE(bIsSleeping);
+			EXPECT_LT(TickCount, MaxTickCount);
+		}
+
+		FMaterialData MaterialData;
+
+		FSingleParticlePhysicsProxy* FloorProxy;
+		TArray<FSingleParticlePhysicsProxy*> ParticleProxies;
+		TArray<FRigidBodyHandle_External*> ParticleHandles;
+		TArray< FJointConstraint*> JointHandles;
+
+		FPBDRigidsSolver* Solver;
+		FChaosSolversModule* Module;
+
+		int32 TickCount;
+	};
+
+	// Compares test results between serial and parallel solve.
+	// Note that we do expect equivalence since the execution order is the same.
+	TEST_F(PartitioningSolverTest, StackOfCubes_SerialVsParallelSolve)
+	{
+		const int32 NumOfCubes = 20;
+		const int32 ParallelModes = 2;
+		const int32 PartitioningTypes = 2;
+		const int32 NumSteps = 10;
+		const int32 NumIterations = 8;
+		//const int32 NumShockPropagationIterations = 0;
+
+		for (int32 PartitioningType = 0; PartitioningType < PartitioningTypes; ++PartitioningType)
+		{
+			TArray<TArray<FVec3>> FinalPosition;
+			FinalPosition.SetNum(ParallelModes);
+			for (int32 ParallelMode = 0; ParallelMode < ParallelModes; ++ParallelMode)
+			{
+				FinalPosition[ParallelMode].Reserve(NumOfCubes);
+
+				CVarGraphColoring->Set(PartitioningType);
+				CVarParallelSolve->Set(ParallelMode);
+
+				FPartitioningSolverTest Test;
+				Test.Solver->SetPositionIterations(NumIterations);
+				Test.Solver->SetVelocityIterations(0);
+				// Test.Solver->GetEvolution()->SetShockPropagationIterations(NumShockPropagationIterations, NumShockPropagationIterations);
+				Test.MakeFloor();
+				Test.MakeStackOfCubes(NumOfCubes);
+
+				for (int32 Step = 0; Step < NumSteps; ++Step)
+				{
+					Test.Advance();
+				}
+
+				for (FRigidBodyHandle_External* ParticleHandle : Test.ParticleHandles)
+				{
+					FinalPosition[ParallelMode].Add(ParticleHandle->GetX());
+				}
+			}
+
+
+			for (int32 ParallelModeId = 0; ParallelModeId < ParallelModes - 1; ++ParallelModeId)
+			{
+				//printf("ParallelModes: %d, %d\n", ParallelModeId, ParallelModeId + 1);
+				for (int32 Index = 0; Index < FinalPosition[ParallelModeId].Num(); ++Index)
+				{
+					EXPECT_VECTOR_NEAR(FinalPosition[ParallelModeId][Index], FinalPosition[ParallelModeId + 1][Index], UE_SMALL_NUMBER);
+				}
+			}
+		}
+	}
+
+	// Compares test results between serial and parallel solve.
+	// Note that we do expect equivalence since the execution order is the same.
+	TEST_F(PartitioningSolverTest, TwoStacksOfCubes_SerialVsParallelSolve)
+	{
+		const int32 NumOfCubes = 20;
+		const int32 ParallelModes = 2;
+		const int32 PartitioningTypes = 2;
+		const int32 NumSteps = 1;
+		const int32 NumIterations = 8;
+		const int32 NumShockPropagationIterations = 0;
+
+		for (int32 PartitioningType = 0; PartitioningType < PartitioningTypes; ++PartitioningType)
+		{
+			TArray<TArray<FVec3>> FinalPosition;
+			FinalPosition.SetNum(ParallelModes);
+			for (int32 ParallelMode = 0; ParallelMode < ParallelModes; ++ParallelMode)
+			{
+				FinalPosition[ParallelMode].Reserve(NumOfCubes);
+
+				CVarGraphColoring->Set(PartitioningType);
+				CVarParallelSolve->Set(ParallelMode);
+
+				FPartitioningSolverTest Test;
+				Test.Solver->SetPositionIterations(NumIterations);
+				//Test.Solver->SetVelocityIterations(0);
+				//Test.Solver->GetEvolution()->SetShockPropagationIterations(NumShockPropagationIterations, NumShockPropagationIterations);
+				Test.MakeFloor();
+				Test.MakeStackOfCubes(NumOfCubes, 200, FVec3(0, 0, 0));
+				Test.MakeStackOfCubes(NumOfCubes, 200, FVec3(600, 0, 0));
+
+				for (int32 Step = 0; Step < NumSteps; ++Step)
+				{
+					Test.Advance();
+				}
+
+				for (FRigidBodyHandle_External* ParticleHandle : Test.ParticleHandles)
+				{
+					FinalPosition[ParallelMode].Add(ParticleHandle->GetX());
+				}
+			}
+
+
+			for (int32 ParallelModeId = 0; ParallelModeId < ParallelModes - 1; ++ParallelModeId)
+			{
+				//printf("ParallelModes: %d, %d\n", ParallelModeId, ParallelModeId + 1);
+				for (int32 Index = 0; Index < FinalPosition[ParallelModeId].Num(); ++Index)
+				{
+					EXPECT_VECTOR_NEAR(FinalPosition[ParallelModeId][Index], FinalPosition[ParallelModeId + 1][Index], UE_SMALL_NUMBER);
+				}
+			}
+		}
+	}
+
+	// Compares test results between serial and parallel solve.
+	// Note that we do expect equivalence since the execution order is the same.
+	TEST_F(PartitioningSolverTest, ChainOfSpheres_SerialVsParallelSolve)
+	{
+		const int32 NumOfCubes = 20;
+		const int32 ParallelModes = 2;
+		const int32 PartitioningTypes = 2;
+		const int32 NumSteps = 10;
+		const int32 NumIterations = 8;
+		const int32 NumShockPropagationIterations = 0;
+
+		for (int32 PartitioningType = 0; PartitioningType < PartitioningTypes; ++PartitioningType)
+		{
+			TArray<TArray<FVec3>> FinalPosition;
+			FinalPosition.SetNum(ParallelModes);
+			for (int32 ParallelMode = 0; ParallelMode < ParallelModes; ++ParallelMode)
+			{
+				FinalPosition[ParallelMode].Reserve(NumOfCubes * 2);
+
+				CVarGraphColoring->Set(PartitioningType);
+				CVarParallelSolve->Set(ParallelMode);
+
+				FPartitioningSolverTest Test;
+				Test.Solver->SetPositionIterations(NumIterations);
+				//Test.Solver->SetVelocityIterations(0);
+				//Test.Solver->GetEvolution()->SetShockPropagationIterations(NumShockPropagationIterations, NumShockPropagationIterations);
+				Test.MakeChainOfSpheres(NumOfCubes, 200, FVec3(0, 0, 0));
+				Test.MakeChainOfSpheres(NumOfCubes, 200, FVec3(600, 0, 0));
+
+				for (int32 Step = 0; Step < NumSteps; ++Step)
+				{
+					Test.Advance();
+				}
+
+				for (FRigidBodyHandle_External* ParticleHandle : Test.ParticleHandles)
+				{
+					FinalPosition[ParallelMode].Add(ParticleHandle->GetX());
+				}
+			}
+
+
+			for (int32 ParallelModeId = 0; ParallelModeId < ParallelModes - 1; ++ParallelModeId)
+			{
+				//printf("ParallelModes: %d, %d\n", ParallelModeId, ParallelModeId + 1);
+				for (int32 Index = 0; Index < FinalPosition[ParallelModeId].Num(); ++Index)
+				{
+					EXPECT_VECTOR_NEAR(FinalPosition[ParallelModeId][Index], FinalPosition[ParallelModeId + 1][Index], UE_SMALL_NUMBER);
+				}
+			}
+		}
+	}
+
+	// Compares test results between no partitioning and graph coloring.
+	// Note that we do not expect equivalence since we change the execution order.
+	TEST_F(PartitioningSolverTest, StackOfCubes_NoPartitioningVsColoring)
+	{
+		const int32 NumOfCubes = 5;
+		const int32 ParallelModes = 2;
+		const int32 PartitioningTypes = 2;
+		const int32 NumSteps = 1;
+
+		for (int32 ParallelMode = 0; ParallelMode < ParallelModes; ++ParallelMode)
+		{
+			TArray<TArray<FVec3>> FinalPosition;
+			FinalPosition.SetNum(PartitioningTypes);
+
+			for (int32 PartitioningType = 0; PartitioningType < PartitioningTypes; ++PartitioningType)
+			{
+				FinalPosition[PartitioningType].Reserve(NumOfCubes);
+
+				CVarGraphColoring->Set(PartitioningType);
+				CVarParallelSolve->Set(ParallelMode);
+
+				FPartitioningSolverTest Test;
+				Test.Solver->SetPositionIterations(500);
+				Test.MakeFloor();
+				Test.MakeStackOfCubes(NumOfCubes);
+
+				for (int32 Step = 0; Step < NumSteps; ++Step)
+				{
+					Test.Advance();
+				}
+
+				for (FRigidBodyHandle_External* ParticleHandle : Test.ParticleHandles)
+				{
+					FinalPosition[PartitioningType].Add(ParticleHandle->GetX());
+				}
+			}
+
+			// Coloring changes the execution order
+			// Therefore, we expect the solver to take longer to converge, especially for high stacks.
+			// Nevertheless, for a small stack and many solver iterations, the result should be nearly identical.
+			for (int32 PartitioningType = 0; PartitioningType < PartitioningTypes - 1; ++PartitioningType)
+			{
+				//printf("PartitioningTypes: %d, %d\n", PartitioningType, PartitioningType + 1);
+				for (int32 Index = 0; Index < FinalPosition[PartitioningType].Num(); ++Index)
+				{
+					EXPECT_VECTOR_NEAR(FinalPosition[PartitioningType][Index], FinalPosition[PartitioningType + 1][Index], UE_KINDA_SMALL_NUMBER);
+				}
+			}
+		}
+	}
+
+	// Compares test results between no partitioning and graph coloring.
+	// Note that we do not expect equivalence since we change the execution order.
+	TEST_F(PartitioningSolverTest, TwoStacksOfCubes_NoPartitioningVsColoring)
+	{
+		const int32 NumOfCubes = 5;
+		const int32 ParallelModes = 2;
+		const int32 PartitioningTypes = 2;
+		const int32 NumSteps = 1;
+		const int32 NumIterations = 500;
+		const int32 NumShockPropagationIterations = 0;
+
+		for (int32 ParallelMode = 0; ParallelMode < ParallelModes; ++ParallelMode)
+		{
+			TArray<TArray<FVec3>> FinalPosition;
+			FinalPosition.SetNum(PartitioningTypes);
+			for (int32 PartitioningType = 0; PartitioningType < PartitioningTypes; ++PartitioningType)
+			{
+				FinalPosition[PartitioningType].Reserve(NumOfCubes * 2);
+
+				CVarGraphColoring->Set(PartitioningType);
+				CVarParallelSolve->Set(ParallelMode);
+
+				FPartitioningSolverTest Test;
+				Test.Solver->SetPositionIterations(NumIterations);
+				//Test.Solver->SetVelocityIterations(0);
+				//Test.Solver->GetEvolution()->SetShockPropagationIterations(NumShockPropagationIterations, NumShockPropagationIterations);
+				Test.MakeFloor();
+				Test.MakeStackOfCubes(NumOfCubes, 200, FVec3(0, 0, 0));
+				Test.MakeStackOfCubes(NumOfCubes, 200, FVec3(600, 0, 0));
+
+				for (int32 Step = 0; Step < NumSteps; ++Step)
+				{
+					Test.Advance();
+				}
+
+				for (FRigidBodyHandle_External* ParticleHandle : Test.ParticleHandles)
+				{
+					FinalPosition[PartitioningType].Add(ParticleHandle->GetX());
+				}
+			}
+
+
+			for (int32 PartitioningType = 0; PartitioningType < PartitioningTypes - 1; ++PartitioningType)
+			{
+				//printf("PartitioningType: %d, %d\n", PartitioningType, PartitioningType + 1);
+				for (int32 Index = 0; Index < FinalPosition[PartitioningType].Num(); ++Index)
+				{
+					EXPECT_VECTOR_NEAR(FinalPosition[PartitioningType][Index], FinalPosition[PartitioningType + 1][Index], UE_KINDA_SMALL_NUMBER);
+				}
+			}
+		}
+	}
+
+	// Compares test results between no partitioning and graph coloring.
+	// Note that we do not expect equivalence since we change the execution order.
+	TEST_F(PartitioningSolverTest, ChainOfSpheres_NoPartitioningVsColoring)
+	{
+		const int32 NumOfCubes = 20;
+		const int32 ParallelModes = 2;
+		const int32 PartitioningTypes = 2;
+		const int32 NumSteps = 10;
+		const int32 NumIterations = 8;
+		const int32 NumShockPropagationIterations = 0;
+
+		for (int32 ParallelMode = 0; ParallelMode < ParallelModes; ++ParallelMode)
+		{
+			TArray<TArray<FVec3>> FinalPosition;
+			FinalPosition.SetNum(PartitioningTypes);
+			for (int32 PartitioningType = 0; PartitioningType < PartitioningTypes; ++PartitioningType)
+			{
+				FinalPosition[PartitioningType].Reserve(NumOfCubes);
+
+				CVarGraphColoring->Set(PartitioningType);
+				CVarParallelSolve->Set(ParallelMode);
+
+				FPartitioningSolverTest Test;
+				Test.Solver->SetPositionIterations(NumIterations);
+				//Test.Solver->SetVelocityIterations(0);
+				//Test.Solver->GetEvolution()->SetShockPropagationIterations(NumShockPropagationIterations, NumShockPropagationIterations);
+				Test.MakeChainOfSpheres(NumOfCubes, 200, FVec3(0, 0, 0));
+				Test.MakeChainOfSpheres(NumOfCubes, 200, FVec3(600, 0, 0));
+
+				for (int32 Step = 0; Step < NumSteps; ++Step)
+				{
+					Test.Advance();
+				}
+
+				for (FRigidBodyHandle_External* ParticleHandle : Test.ParticleHandles)
+				{
+					FinalPosition[PartitioningType].Add(ParticleHandle->GetX());
+				}
+			}
+
+
+			for (int32 PartitioningType = 0; PartitioningType < PartitioningTypes - 1; ++PartitioningType)
+			{
+				//printf("ParallelModes: %d, %d\n", PartitioningType, PartitioningType + 1);
+				for (int32 Index = 0; Index < FinalPosition[PartitioningType].Num(); ++Index)
+				{
+					EXPECT_VECTOR_NEAR(FinalPosition[PartitioningType][Index], FinalPosition[PartitioningType + 1][Index], UE_KINDA_SMALL_NUMBER);
+				}
+			}
+		}
+	}
+}
\ No newline at end of file
diff --git a/Engine/Source/Runtime/Experimental/Chaos/Private/Chaos/Character/CharacterGroundConstraintContainerSolver.cpp b/Engine/Source/Runtime/Experimental/Chaos/Private/Chaos/Character/CharacterGroundConstraintContainerSolver.cpp
--- a/Engine/Source/Runtime/Experimental/Chaos/Private/Chaos/Character/CharacterGroundConstraintContainerSolver.cpp
+++ b/Engine/Source/Runtime/Experimental/Chaos/Private/Chaos/Character/CharacterGroundConstraintContainerSolver.cpp
@@ -17,7 +17,7 @@ namespace Chaos
 		{
 		}
 
-		void FCharacterGroundConstraintContainerSolver::Reset(const int32 InMaxConstraints)
+		void FCharacterGroundConstraintContainerSolver::ResetImpl(const int32 InMaxConstraints)
 		{
 			SolvedConstraints.Reset(InMaxConstraints);
 
@@ -60,6 +60,17 @@ namespace Chaos
 			SolvedConstraints.Add(Constraint);
 		}
 
+		bool FCharacterGroundConstraintContainerSolver::SortConstraintsByPartitionIndex()
+		{
+			SolvedConstraints.StableSort(
+				[](const FConstraintHandle& L, const FConstraintHandle& R)
+				{
+					return L.GetSolverPartitionIndex() < R.GetSolverPartitionIndex();
+				});
+
+			return true;
+		}
+
 		void FCharacterGroundConstraintContainerSolver::AddBodies(FSolverBodyContainer& SolverBodyContainer)
 		{
 			for (int32 SolverConstraintIndex = 0, SolverConstraintEndIndex = SolvedConstraints.Num(); SolverConstraintIndex < SolverConstraintEndIndex; ++SolverConstraintIndex)
@@ -124,7 +135,27 @@ namespace Chaos
 
 		void FCharacterGroundConstraintContainerSolver::ApplyPositionConstraints(const FReal Dt, const int32 It, const int32 NumIts)
 		{
-			for (int32 SolverConstraintIndex = 0; SolverConstraintIndex < Solvers.Num(); ++SolverConstraintIndex)
+			check(Solvers.Num() <= GetNumConstraints());
+			ApplyPositionConstraints(Dt, It, NumIts, 0, Solvers.Num());
+		}
+
+		void FCharacterGroundConstraintContainerSolver::ApplyPositionConstraints(const FReal Dt, const int32 It, const int32 NumIts, const int32 ConstraintBeginIndex, const int32 ConstraintEndIndex)
+		{
+			if (ConstraintEndIndex <= ConstraintBeginIndex)
+			{
+				return;
+			}
+
+			check(Solvers.Num() <= GetNumConstraints());
+			// NOTE: may be called in parallel. Should not change the container or any elements outside of [BeginIndex, EndIndex)
+			check(ConstraintBeginIndex >= 0);
+			check(ConstraintEndIndex <= GetNumConstraints());
+
+			// Handle the case where we dropped some constraints because there were too  many
+			const int32 BeginIndex = FMath::Min(ConstraintBeginIndex, Solvers.Num());
+			const int32 EndIndex = FMath::Min(ConstraintEndIndex, Solvers.Num());
+
+			for (int32 SolverConstraintIndex = BeginIndex; SolverConstraintIndex < EndIndex; ++SolverConstraintIndex)
 			{
 				Solvers[SolverConstraintIndex].SolvePosition();
 			}
@@ -134,9 +165,17 @@ namespace Chaos
 		{
 		}
 
+		void FCharacterGroundConstraintContainerSolver::ApplyVelocityConstraints(const FReal Dt, const int32 It, const int32 NumIts, const int32 BeginIndex, const int32 EndIndex)
+		{
+		}
+
 		void FCharacterGroundConstraintContainerSolver::ApplyProjectionConstraints(const FReal Dt, const int32 It, const int32 NumIts)
 		{
 		}
 
+		void FCharacterGroundConstraintContainerSolver::ApplyProjectionConstraints(const FReal Dt, const int32 It, const int32 NumIts, const int32 BeginIndex, const int32 EndIndex)
+		{
+		}
+
 	} // namespace Private
 } // namespace Chaos
\ No newline at end of file
diff --git a/Engine/Source/Runtime/Experimental/Chaos/Private/Chaos/Character/CharacterGroundConstraintContainerSolver.h b/Engine/Source/Runtime/Experimental/Chaos/Private/Chaos/Character/CharacterGroundConstraintContainerSolver.h
--- a/Engine/Source/Runtime/Experimental/Chaos/Private/Chaos/Character/CharacterGroundConstraintContainerSolver.h
+++ b/Engine/Source/Runtime/Experimental/Chaos/Private/Chaos/Character/CharacterGroundConstraintContainerSolver.h
@@ -20,8 +20,6 @@ namespace Chaos
 
 			virtual int32 GetNumConstraints() const override final { return SolvedConstraints.Num(); }
 
-			virtual void Reset(const int32 InMaxConstraints) override final;
-
 			/// Add all constraints from the constraint container in the order that they appear there
 			virtual void AddConstraints() override final;
 
@@ -41,16 +39,25 @@ namespace Chaos
 
 			/// Apply a single position solve iteration for all constraints
 			virtual void ApplyPositionConstraints(const FReal Dt, const int32 It, const int32 NumIts) override final;
+			virtual void ApplyPositionConstraints(const FReal Dt, const int32 It, const int32 NumIts, const int32 BeginIndex, const int32 EndIndex) override final;
 
 			/// Apply a single velocity solve iteration for all constraints
 			/// Note: Currently not used for character ground constraints
 			virtual void ApplyVelocityConstraints(const FReal Dt, const int32 It, const int32 NumIts) override final;
+			virtual void ApplyVelocityConstraints(const FReal Dt, const int32 It, const int32 NumIts, const int32 BeginIndex, const int32 EndIndex) override final;
 
 			/// Apply a single projection solve iteration for all constraints
 			/// Note: Currently not used for character ground constraints
 			virtual void ApplyProjectionConstraints(const FReal Dt, const int32 It, const int32 NumIts) override final;
+			virtual void ApplyProjectionConstraints(const FReal Dt, const int32 It, const int32 NumIts, const int32 BeginIndex, const int32 EndIndex) override final;
 
 		private:
+			virtual void ResetImpl(const int32 InMaxConstraints) override final;
+			virtual bool SupportsPartitioning() const override final { return true; }
+			virtual bool SortConstraintsByPartitionIndex() override final;
+			virtual const FConstraintHandle* GetConstraintHandle(const int32 Index) const { return SolvedConstraints[Index]; }
+			virtual FConstraintHandle* GetConstraintHandle(const int32 Index) { return SolvedConstraints[Index]; }
+
 			void AddConstraint(FCharacterGroundConstraintHandle* Constraint);
 
 			FCharacterGroundConstraintContainer& ConstraintContainer;
diff --git a/Engine/Source/Runtime/Experimental/Chaos/Private/Chaos/Collision/PBDCollisionContainerSolver.cpp b/Engine/Source/Runtime/Experimental/Chaos/Private/Chaos/Collision/PBDCollisionContainerSolver.cpp
--- a/Engine/Source/Runtime/Experimental/Chaos/Private/Chaos/Collision/PBDCollisionContainerSolver.cpp
+++ b/Engine/Source/Runtime/Experimental/Chaos/Private/Chaos/Collision/PBDCollisionContainerSolver.cpp
@@ -415,7 +415,7 @@ namespace Chaos
 		: FConstraintContainerSolver(InPriority)
 		, ConstraintContainer(InConstraintContainer)
 		, CollisionConstraints()
-		, AppliedShockPropagation(1)
+		, AppliedShockPropagation()
 		, Scratch()
 		, CollisionSolvers(nullptr)
 		, CollisionSolverManifoldPoints(nullptr)
@@ -430,12 +430,11 @@ namespace Chaos
 	{
 	}
 
-	void FPBDCollisionContainerSolver::Reset(const int32 MaxCollisions)
+	void FPBDCollisionContainerSolver::ResetImpl(const int32 MaxCollisions)
 	{
-		AppliedShockPropagation = FSolverReal(1);
-
 		const int CollisionBufferNum = CalculateCollisionBufferSize(MaxCollisions, CollisionConstraints.Num());
 		CollisionConstraints.Reset(CollisionBufferNum);
+		AppliedShockPropagation.Init(1, CollisionBufferNum);
 		bCollisionConstraintPerIterationCollisionDetection.Reset(CollisionBufferNum);
 
 		// Just set the array size for these right away - all data will be initialized later
@@ -583,6 +582,17 @@ namespace Chaos
 		CollisionConstraints.Add(&Constraint);
 	}
 
+	bool FPBDCollisionContainerSolver::SortConstraintsByPartitionIndex()
+	{
+		CollisionConstraints.StableSort(
+			[](const FConstraintHandle& L, const FConstraintHandle& R)
+			{
+				return L.GetSolverPartitionIndex() < R.GetSolverPartitionIndex();
+			});
+
+		return true;
+	}
+
 	void FPBDCollisionContainerSolver::AddBodies(FSolverBodyContainer& SolverBodyContainer)
 	{
 		// All constarints are now added. We can allocate the solver buffers.
@@ -707,25 +717,45 @@ namespace Chaos
 
 	void FPBDCollisionContainerSolver::ApplyPositionConstraints(const FReal Dt, const int32 It, const int32 NumIts)
 	{
-		SolvePositionImpl(Dt, It, NumIts, 0, NumSolvers(), ConstraintContainer.GetSolverSettings());
+		check(NumSolvers() <= GetNumConstraints());
+		ApplyPositionConstraints(Dt, It, NumIts, 0, NumSolvers());
+	}
+
+	void FPBDCollisionContainerSolver::ApplyPositionConstraints(const FReal Dt, const int32 It, const int32 NumIts, const int32 BeginIndex, const int32 EndIndex)
+	{
+		check(NumSolvers() <= GetNumConstraints());
+		SolvePositionImpl(Dt, It, NumIts, BeginIndex, EndIndex, ConstraintContainer.GetSolverSettings());
 	}
 
 	void FPBDCollisionContainerSolver::ApplyVelocityConstraints(const FReal Dt, const int32 It, const int32 NumIts)
 	{
-		SolveVelocityImpl(Dt, It, NumIts, 0, NumSolvers(), ConstraintContainer.GetSolverSettings());
+		check(NumSolvers() <= GetNumConstraints());
+		ApplyVelocityConstraints(Dt, It, NumIts, 0, NumSolvers());
+	}
+
+	void FPBDCollisionContainerSolver::ApplyVelocityConstraints(const FReal Dt, const int32 It, const int32 NumIts, const int32 BeginIndex, const int32 EndIndex)
+	{
+		check(NumSolvers() <= GetNumConstraints());
+		SolveVelocityImpl(Dt, It, NumIts, BeginIndex, EndIndex, ConstraintContainer.GetSolverSettings());
 	}
 
 	void FPBDCollisionContainerSolver::ApplyProjectionConstraints(const FReal Dt, const int32 It, const int32 NumIts)
+	{
+		check(NumSolvers() <= GetNumConstraints());
+		ApplyProjectionConstraints(Dt, It, NumIts, 0, NumSolvers());
+	}
+
+	void FPBDCollisionContainerSolver::ApplyProjectionConstraints(const FReal Dt, const int32 It, const int32 NumIts, const int32 BeginIndex, const int32 EndIndex)
 	{
 		// Not supported for collisions
 	}
 
-	void FPBDCollisionContainerSolver::ApplyShockPropagation(const FSolverReal ShockPropagation)
+	void FPBDCollisionContainerSolver::ApplyShockPropagation(const FSolverReal ShockPropagation, const int32 BeginIndex, const int32 EndIndex)
 	{
 		// @todo(chaos): cache the mass scales so we don't have to look in the constraint again
-		if (ShockPropagation != AppliedShockPropagation)
+		for (int32 SolverIndex = BeginIndex; SolverIndex < EndIndex; ++SolverIndex)
 		{
-			for (int32 SolverIndex = 0; SolverIndex < NumSolvers(); ++SolverIndex)
+			if (ShockPropagation != AppliedShockPropagation[SolverIndex])
 			{
 				const FPBDCollisionConstraint* Constraint = GetConstraint(SolverIndex);
 				Private::FPBDCollisionSolver& Solver = GetSolver(SolverIndex);
@@ -742,9 +772,9 @@ namespace Chaos
 						Solver.UpdateMassNormal();
 					}
 				}
-			}
 
-			AppliedShockPropagation = ShockPropagation;
+				AppliedShockPropagation[SolverIndex] = ShockPropagation;
+			}
 		}
 	}
 
@@ -752,17 +782,17 @@ namespace Chaos
 	{
 		const bool bEnableShockPropagation = (It >= NumIts - SolverSettings.NumPositionShockPropagationIterations);
 		const FSolverReal ShockPropagation = (bEnableShockPropagation) ? CVars::Chaos_PBDCollisionSolver_Position_MinInvMassScale : FSolverReal(1);
-		ApplyShockPropagation(ShockPropagation);
+		ApplyShockPropagation(ShockPropagation, BeginIndex, EndIndex);
 	}
 
 	void FPBDCollisionContainerSolver::UpdateVelocityShockPropagation(const FReal Dt, const int32 It, const int32 NumIts, const int32 BeginIndex, const int32 EndIndex, const FPBDCollisionSolverSettings& SolverSettings)
 	{
 		const bool bEnableShockPropagation = (It >= NumIts - SolverSettings.NumVelocityShockPropagationIterations);
 		const FSolverReal ShockPropagation = (bEnableShockPropagation) ? CVars::Chaos_PBDCollisionSolver_Velocity_MinInvMassScale : FSolverReal(1);
-		ApplyShockPropagation(ShockPropagation);
+		ApplyShockPropagation(ShockPropagation, BeginIndex, EndIndex);
 	}
 
-	void FPBDCollisionContainerSolver::SolvePositionImpl(const FReal InDt, const int32 It, const int32 NumIts, const int32 BeginIndex, const int32 EndIndex, const FPBDCollisionSolverSettings& SolverSettings)
+	void FPBDCollisionContainerSolver::SolvePositionImpl(const FReal InDt, const int32 It, const int32 NumIts, const int32 ConstraintBeginIndex, const int32 ConstraintEndIndex, const FPBDCollisionSolverSettings& SolverSettings)
 	{
 		SCOPE_CYCLE_COUNTER(STAT_Collisions_Apply);
 		if (!CVars::bChaos_PBDCollisionSolver_Position_SolveEnabled)
@@ -770,6 +800,19 @@ namespace Chaos
 			return;
 		}
 
+		if (ConstraintEndIndex <= ConstraintBeginIndex)
+		{
+			return;
+		}
+
+		// NOTE: may be called in parallel. Should not change the container or any elements outside of [BeginIndex, EndIndex)
+		check(ConstraintBeginIndex >= 0);
+		check(ConstraintEndIndex <= CollisionConstraints.Num());
+
+		// Handle the case where we dropped some constraints because there were too  many
+		const int32 BeginIndex = FMath::Min(ConstraintBeginIndex, NumCollisionSolvers);
+		const int32 EndIndex = FMath::Min(ConstraintEndIndex, NumCollisionSolvers);
+
 		if (EndIndex <= BeginIndex)
 		{
 			return;
@@ -795,21 +838,21 @@ namespace Chaos
 		// Apply the position correction
 		if (bApplyStaticFriction)
 		{
-			for (int32 SolverIndex = 0; SolverIndex < NumSolvers(); ++SolverIndex)
+			for (int32 SolverIndex = BeginIndex; SolverIndex < EndIndex; ++SolverIndex)
 			{
 				CollisionSolvers[SolverIndex].SolvePositionWithFriction(Dt, MaxPushOut);
 			}
 		}
 		else
 		{
-			for (int32 SolverIndex = 0; SolverIndex < NumSolvers(); ++SolverIndex)
+			for (int32 SolverIndex = BeginIndex; SolverIndex < EndIndex; ++SolverIndex)
 			{
 				CollisionSolvers[SolverIndex].SolvePositionNoFriction(Dt, MaxPushOut);
 			}
 		}
 	}
 
-	void FPBDCollisionContainerSolver::SolveVelocityImpl(const FReal InDt, const int32 It, const int32 NumIts, const int32 BeginIndex, const int32 EndIndex, const FPBDCollisionSolverSettings& SolverSettings)
+	void FPBDCollisionContainerSolver::SolveVelocityImpl(const FReal InDt, const int32 It, const int32 NumIts, const int32 ConstraintBeginIndex, const int32 ConstraintEndIndex, const FPBDCollisionSolverSettings& SolverSettings)
 	{
 		SCOPE_CYCLE_COUNTER(STAT_Collisions_ApplyPushOut);
 		if (!CVars::bChaos_PBDCollisionSolver_Velocity_SolveEnabled)
@@ -817,6 +860,19 @@ namespace Chaos
 			return;
 		}
 
+		if (ConstraintEndIndex <= ConstraintBeginIndex)
+		{
+			return;
+		}
+
+		// NOTE: may be called in parallel. Should not change the container or any elements outside of [BeginIndex, EndIndex)
+		check(ConstraintBeginIndex >= 0);
+		check(ConstraintEndIndex <= CollisionConstraints.Num());
+
+		// Handle the case where we dropped some constraints because there were too  many
+		const int32 BeginIndex = FMath::Min(ConstraintBeginIndex, NumCollisionSolvers);
+		const int32 EndIndex = FMath::Min(ConstraintEndIndex, NumCollisionSolvers);
+
 		if (EndIndex <= BeginIndex)
 		{
 			return;
@@ -827,7 +883,7 @@ namespace Chaos
 		const FSolverReal Dt = FSolverReal(InDt);
 		const bool bApplyDynamicFriction = (It >= NumIts - SolverSettings.NumVelocityFrictionIterations);
 
-		for (int32 SolverIndex = 0; SolverIndex < NumSolvers(); ++SolverIndex)
+		for (int32 SolverIndex = BeginIndex; SolverIndex < EndIndex; ++SolverIndex)
 		{
 			CollisionSolvers[SolverIndex].SolveVelocity(Dt, bApplyDynamicFriction);
 		}
diff --git a/Engine/Source/Runtime/Experimental/Chaos/Private/Chaos/Collision/PBDCollisionContainerSolverJacobi.cpp b/Engine/Source/Runtime/Experimental/Chaos/Private/Chaos/Collision/PBDCollisionContainerSolverJacobi.cpp
--- a/Engine/Source/Runtime/Experimental/Chaos/Private/Chaos/Collision/PBDCollisionContainerSolverJacobi.cpp
+++ b/Engine/Source/Runtime/Experimental/Chaos/Private/Chaos/Collision/PBDCollisionContainerSolverJacobi.cpp
@@ -270,7 +270,7 @@ namespace Chaos
 		{
 		}
 
-		void FPBDCollisionContainerSolverJacobi::Reset(const int32 MaxCollisions)
+		void FPBDCollisionContainerSolverJacobi::ResetImpl(const int32 MaxCollisions)
 		{
 			// A over-allocation policy to avoid reallocation every frame in the common case where a pile of objects is dropped
 			// and the number of contacts increases every tick.
@@ -333,6 +333,18 @@ namespace Chaos
 			CollisionSolvers[Index].Reset();
 		}
 
+		bool FPBDCollisionContainerSolverJacobi::SortConstraintsByPartitionIndex()
+		{
+			CollisionConstraints.StableSort(
+				[](const FConstraintHandle& L, const FConstraintHandle& R)
+				{
+					return L.GetSolverPartitionIndex() < R.GetSolverPartitionIndex();
+				});
+
+			// NOTE: CollisionSolvers is not sorted, but it will be repopulated with updated information in AddBodies and GatherInput
+
+			return true;
+		}
 
 		void FPBDCollisionContainerSolverJacobi::AddBodies(FSolverBodyContainer& SolverBodyContainer)
 		{
@@ -425,22 +437,41 @@ namespace Chaos
 				UpdateCollisionConstraintFromSolver(Constraint, CollisionSolver, Dt);
 
 				// Reset the collision solver here as the pointers will be invalid on the next tick
-				// @todo(chaos): Pointers are always reinitalized before use next tick so this isn't strictly necessary
+				// @todo(chaos): Pointers are always reinitialized before use next tick so this isn't strictly necessary
 				CollisionSolver.Reset();
 			}
 		}
 
 		void FPBDCollisionContainerSolverJacobi::ApplyPositionConstraints(const FReal Dt, const int32 It, const int32 NumIts)
 		{
-			SolvePositionImpl(Dt, It, NumIts, 0, NumSolvers(), ConstraintContainer.GetSolverSettings());
+			check(NumSolvers() <= GetNumConstraints());
+			ApplyPositionConstraints(Dt, It, NumIts, 0, NumSolvers());
+		}
+
+		void FPBDCollisionContainerSolverJacobi::ApplyPositionConstraints(const FReal Dt, const int32 It, const int32 NumIts, const int32 BeginIndex, const int32 EndIndex)
+		{
+			check(NumSolvers() <= GetNumConstraints());
+			SolvePositionImpl(Dt, It, NumIts, BeginIndex, EndIndex, ConstraintContainer.GetSolverSettings());
 		}
 
 		void FPBDCollisionContainerSolverJacobi::ApplyVelocityConstraints(const FReal Dt, const int32 It, const int32 NumIts)
 		{
-			SolveVelocityImpl(Dt, It, NumIts, 0, NumSolvers(), ConstraintContainer.GetSolverSettings());
+			check(NumSolvers() <= GetNumConstraints());
+			ApplyVelocityConstraints(Dt, It, NumIts, 0, NumSolvers());
+		}
+
+		void FPBDCollisionContainerSolverJacobi::ApplyVelocityConstraints(const FReal Dt, const int32 It, const int32 NumIts, const int32 BeginIndex, const int32 EndIndex)
+		{
+			check(NumSolvers() <= GetNumConstraints());
+			SolveVelocityImpl(Dt, It, NumIts, BeginIndex, EndIndex, ConstraintContainer.GetSolverSettings());
 		}
 
 		void FPBDCollisionContainerSolverJacobi::ApplyProjectionConstraints(const FReal Dt, const int32 It, const int32 NumIts)
+		{
+			ApplyProjectionConstraints(Dt, It, NumIts, 0, NumSolvers());
+		}
+
+		void FPBDCollisionContainerSolverJacobi::ApplyProjectionConstraints(const FReal Dt, const int32 It, const int32 NumIts, const int32 BeginIndex, const int32 EndIndex)
 		{
 			// Not supported for collisions
 		}
@@ -492,6 +523,12 @@ namespace Chaos
 				return false;
 			}
 
+			check(NumSolvers() == GetNumConstraints());
+
+			// NOTE: may be called in parallel. Should not change the container or any elements outside of [BeginIndex, EndIndex)
+			check(BeginIndex >= 0);
+			check(EndIndex <= GetNumConstraints());
+
 			UpdatePositionShockPropagation(InDt, It, NumIts, BeginIndex, EndIndex, SolverSettings);
 
 			// We run collision detection here under two conditions (normally it is run after Integration and before the constraint solver phase):
@@ -522,7 +559,7 @@ namespace Chaos
 			return true;
 		}
 
-		bool FPBDCollisionContainerSolverJacobi::SolveVelocityImpl(const FReal InDt, const int32 It, const int32 NumIts, const int32 BeginIndex, const int32 EndIndex, const FPBDCollisionSolverSettings& SolverSettings)
+		bool FPBDCollisionContainerSolverJacobi::SolveVelocityImpl(const FReal InDt, const int32 It, const int32 NumIts, const int32 ConstraintBeginIndex, const int32 ConstraintEndIndex, const FPBDCollisionSolverSettings& SolverSettings)
 		{
 			SCOPE_CYCLE_COUNTER(STAT_Collisions_ApplyPushOut);
 			if (!CVars::bChaos_PBDCollisionSolver_Velocity_SolveEnabled)
@@ -530,11 +567,21 @@ namespace Chaos
 				return false;
 			}
 
-			if (EndIndex <= BeginIndex)
+			if (ConstraintEndIndex <= ConstraintBeginIndex)
 			{
 				return false;
 			}
 
+			check(NumSolvers() <= GetNumConstraints());
+
+			// NOTE: may be called in parallel. Should not change the container or any elements outside of [BeginIndex, EndIndex)
+			check(ConstraintBeginIndex >= 0);
+			check(ConstraintEndIndex <= GetNumConstraints());
+
+			// Handle the case where we dropped some constraints because there were too  many
+			const int32 BeginIndex = FMath::Min(ConstraintBeginIndex, NumSolvers());
+			const int32 EndIndex = FMath::Min(ConstraintEndIndex, NumSolvers());
+
 			UpdateVelocityShockPropagation(InDt, It, NumIts, BeginIndex, EndIndex, SolverSettings);
 
 			const FSolverReal Dt = FSolverReal(InDt);
diff --git a/Engine/Source/Runtime/Experimental/Chaos/Private/Chaos/Collision/PBDCollisionContainerSolverSimd.cpp b/Engine/Source/Runtime/Experimental/Chaos/Private/Chaos/Collision/PBDCollisionContainerSolverSimd.cpp
--- a/Engine/Source/Runtime/Experimental/Chaos/Private/Chaos/Collision/PBDCollisionContainerSolverSimd.cpp
+++ b/Engine/Source/Runtime/Experimental/Chaos/Private/Chaos/Collision/PBDCollisionContainerSolverSimd.cpp
@@ -306,7 +306,7 @@ namespace Chaos
 			: FConstraintContainerSolver(InPriority)
 			, ConstraintContainer(InConstraintContainer)
 			, NumConstraints(0)
-			, AppliedShockPropagation(1)
+			, AppliedShockPropagation()
 			, bPerIterationCollisionDetection(false)
 		{
 	#if !UE_BUILD_SHIPPING && !UE_BUILD_TEST && INTEL_ISPC
@@ -324,10 +324,9 @@ namespace Chaos
 		{
 		}
 
-		void FPBDCollisionContainerSolverSimd::Reset(const int32 MaxCollisions)
+		void FPBDCollisionContainerSolverSimd::ResetImpl(const int32 MaxCollisions)
 		{
 			NumConstraints = 0;
-			AppliedShockPropagation = FSolverReal(1);
 
 			ConstraintSolverIds.Reset(MaxCollisions);
 
@@ -336,6 +335,7 @@ namespace Chaos
 			SimdData.SimdSolvers.Reset(MaxCollisions);
 			SimdData.SimdManifoldPoints.Reset(MaxCollisions);
 			SimdData.SimdConstraints.Reset(MaxCollisions);
+			AppliedShockPropagation.Init(1, MaxCollisions);
 
 			SimdData.SimdNumConstraints.SetValues(0);
 		}
@@ -401,6 +401,20 @@ namespace Chaos
 			NumConstraints += NumActiveConstraints;
 		}
 
+		bool FPBDCollisionContainerSolverSimd::SortConstraintsByPartitionIndex()
+		{
+			////@todo(Chaos): Implement constraint color sorting for the SIMD solver
+			//// This requires sorting constraints before the SIMD lanes are allocated 
+			//// so that only constraints with the same partition index are sharing SIMD data.
+			//CollisionConstraints.StableSort(
+			//	[](const FConstraintHandle& L, const FConstraintHandle& R)
+			//	{
+			//		return L.GetSolverPartitionIndex() < R.GetSolverPartitionIndex();
+			//	});
+
+			return false;
+		}
+
 		void FPBDCollisionContainerSolverSimd::CreateSolvers()
 		{
 			const int32 NumSolvers = SimdData.SimdNumConstraints.GetMaxValue();
@@ -558,6 +572,11 @@ namespace Chaos
 		}
 
 		void FPBDCollisionContainerSolverSimd::ApplyPositionConstraints(const FReal InDt, const int32 It, const int32 NumIts)
+		{
+			ApplyPositionConstraints(InDt, It, NumIts, 0, SimdData.SimdSolvers.Num());
+		}
+
+		void FPBDCollisionContainerSolverSimd::ApplyPositionConstraints(const FReal InDt, const int32 It, const int32 NumIts, const int32 BeginIndex, const int32 EndIndex)
 		{
 			SCOPE_CYCLE_COUNTER(STAT_Collisions_Apply);
 			if (!CVars::bChaos_PBDCollisionSolver_Position_SolveEnabled)
@@ -568,7 +587,7 @@ namespace Chaos
 			const FPBDCollisionSolverSettings& SolverSettings = ConstraintContainer.GetSolverSettings();
 			const FSolverReal Dt = FSolverReal(InDt);
 
-			UpdatePositionShockPropagation(It, NumIts, SolverSettings);
+			UpdatePositionShockPropagation(It, NumIts, BeginIndex, EndIndex, SolverSettings);
 
 			// We run collision detection here under two conditions (normally it is run after Integration and before the constraint solver phase):
 			// 1) When deferring collision detection until the solver phase for better joint-collision behaviour (RBAN). In this case, we only do this on the first iteration.
@@ -576,7 +595,7 @@ namespace Chaos
 			const bool bRunDeferredCollisionDetection = (It == 0) && ConstraintContainer.GetDetectorSettings().bDeferNarrowPhase;
 			if (bRunDeferredCollisionDetection || bPerIterationCollisionDetection)
 			{
-				UpdateCollisions(Dt);
+				UpdateCollisions(Dt, BeginIndex, EndIndex);
 			}
 
 			// Only apply friction for the last few (tunable) iterations
@@ -588,24 +607,29 @@ namespace Chaos
 			if (bApplyStaticFriction)
 			{
 				Private::FPBDCollisionSolverHelperSimd::SolvePositionWithFriction(
-					MakeArrayView(SimdData.SimdSolvers),
-					MakeArrayView(SimdData.SimdManifoldPoints),
-					MakeArrayView(SimdData.SimdSolverBodies),
+					MakeArrayView(&SimdData.SimdSolvers[BeginIndex], EndIndex - BeginIndex),
+					MakeArrayView(&SimdData.SimdManifoldPoints[BeginIndex], EndIndex - BeginIndex),
+					MakeArrayView(&SimdData.SimdSolverBodies[BeginIndex], EndIndex - BeginIndex),
 					Dt,
 					MaxPushOut);
 			}
 			else
 			{
 				Private::FPBDCollisionSolverHelperSimd::SolvePositionNoFriction(
-					MakeArrayView(SimdData.SimdSolvers),
-					MakeArrayView(SimdData.SimdManifoldPoints),
-					MakeArrayView(SimdData.SimdSolverBodies),
+					MakeArrayView(&SimdData.SimdSolvers[BeginIndex], EndIndex - BeginIndex),
+					MakeArrayView(&SimdData.SimdManifoldPoints[BeginIndex], EndIndex - BeginIndex),
+					MakeArrayView(&SimdData.SimdSolverBodies[BeginIndex], EndIndex - BeginIndex),
 					Dt,
 					MaxPushOut);
 			}
 		}
 
 		void FPBDCollisionContainerSolverSimd::ApplyVelocityConstraints(const FReal InDt, const int32 It, const int32 NumIts)
+		{
+			ApplyVelocityConstraints(InDt, It, NumIts, 0, SimdData.SimdSolvers.Num());
+		}
+
+		void FPBDCollisionContainerSolverSimd::ApplyVelocityConstraints(const FReal InDt, const int32 It, const int32 NumIts, const int32 BeginIndex, const int32 EndIndex)
 		{
 			SCOPE_CYCLE_COUNTER(STAT_Collisions_ApplyPushOut);
 			if (!CVars::bChaos_PBDCollisionSolver_Velocity_SolveEnabled)
@@ -615,7 +639,7 @@ namespace Chaos
 
 			const FPBDCollisionSolverSettings& SolverSettings = ConstraintContainer.GetSolverSettings();
 
-			UpdateVelocityShockPropagation(It, NumIts, SolverSettings);
+			UpdateVelocityShockPropagation(It, NumIts, BeginIndex, EndIndex, SolverSettings);
 
 			const FSolverReal Dt = FSolverReal(InDt);
 			const bool bApplyDynamicFriction = (It >= NumIts - SolverSettings.NumVelocityFrictionIterations);
@@ -623,32 +647,37 @@ namespace Chaos
 			if (bApplyDynamicFriction)
 			{
 				Private::FPBDCollisionSolverHelperSimd::SolveVelocityWithFriction(
-					MakeArrayView(SimdData.SimdSolvers),
-					MakeArrayView(SimdData.SimdManifoldPoints),
-					MakeArrayView(SimdData.SimdSolverBodies),
+					MakeArrayView(&SimdData.SimdSolvers[BeginIndex], EndIndex - BeginIndex),
+					MakeArrayView(&SimdData.SimdManifoldPoints[BeginIndex], EndIndex - BeginIndex),
+					MakeArrayView(&SimdData.SimdSolverBodies[BeginIndex], EndIndex - BeginIndex),
 					Dt);
 			}
 			else
 			{
 				Private::FPBDCollisionSolverHelperSimd::SolveVelocityNoFriction(
-					MakeArrayView(SimdData.SimdSolvers),
-					MakeArrayView(SimdData.SimdManifoldPoints),
-					MakeArrayView(SimdData.SimdSolverBodies),
+					MakeArrayView(&SimdData.SimdSolvers[BeginIndex], EndIndex - BeginIndex),
+					MakeArrayView(&SimdData.SimdManifoldPoints[BeginIndex], EndIndex - BeginIndex),
+					MakeArrayView(&SimdData.SimdSolverBodies[BeginIndex], EndIndex - BeginIndex),
 					Dt);
 			}
 		}
 
 		void FPBDCollisionContainerSolverSimd::ApplyProjectionConstraints(const FReal Dt, const int32 It, const int32 NumIts)
+		{
+			ApplyProjectionConstraints(Dt, It, NumIts, 0, SimdData.SimdSolvers.Num());
+		}
+
+		void FPBDCollisionContainerSolverSimd::ApplyProjectionConstraints(const FReal Dt, const int32 It, const int32 NumIts, const int32 BeginIndex, const int32 EndIndex)
 		{
 			// Not supported for collisions
 		}
 
-		void FPBDCollisionContainerSolverSimd::ApplyShockPropagation(const FSolverReal ShockPropagation)
+		void FPBDCollisionContainerSolverSimd::ApplyShockPropagation(const FSolverReal ShockPropagation, const int32 BeginIndex, const int32 EndIndex)
 		{
 			// @todo(chaos): cache the mass scales so we don't have to look in the constraint again
-			if (ShockPropagation != AppliedShockPropagation)
+			for (int32 SolverIndex = BeginIndex; SolverIndex < EndIndex; ++SolverIndex)
 			{
-				for (int32 SolverIndex = 0; SolverIndex < SimdData.SimdSolvers.Num(); ++SolverIndex)
+				if (ShockPropagation != AppliedShockPropagation[SolverIndex])
 				{
 					for (int32 LaneIndex = 0; LaneIndex < 4; ++LaneIndex)
 					{
@@ -673,33 +702,32 @@ namespace Chaos
 							}
 						}
 					}
+					AppliedShockPropagation[SolverIndex] = ShockPropagation;
 				}
-
-				AppliedShockPropagation = ShockPropagation;
 			}
 		}
 
-		void FPBDCollisionContainerSolverSimd::UpdatePositionShockPropagation(const int32 It, const int32 NumIts, const FPBDCollisionSolverSettings& SolverSettings)
+		void FPBDCollisionContainerSolverSimd::UpdatePositionShockPropagation(const int32 It, const int32 NumIts, const int32 BeginIndex, const int32 EndIndex, const FPBDCollisionSolverSettings& SolverSettings)
 		{
 			const bool bEnableShockPropagation = (It >= NumIts - SolverSettings.NumPositionShockPropagationIterations);
 			const FSolverReal ShockPropagation = (bEnableShockPropagation) ? CVars::Chaos_PBDCollisionSolver_Position_MinInvMassScale : FSolverReal(1);
-			ApplyShockPropagation(ShockPropagation);
+			ApplyShockPropagation(ShockPropagation, BeginIndex, EndIndex);
 		}
 
-		void FPBDCollisionContainerSolverSimd::UpdateVelocityShockPropagation(const int32 It, const int32 NumIts, const FPBDCollisionSolverSettings& SolverSettings)
+		void FPBDCollisionContainerSolverSimd::UpdateVelocityShockPropagation(const int32 It, const int32 NumIts, const int32 BeginIndex, const int32 EndIndex, const FPBDCollisionSolverSettings& SolverSettings)
 		{
 			const bool bEnableShockPropagation = (It >= NumIts - SolverSettings.NumVelocityShockPropagationIterations);
 			const FSolverReal ShockPropagation = (bEnableShockPropagation) ? CVars::Chaos_PBDCollisionSolver_Velocity_MinInvMassScale : FSolverReal(1);
-			ApplyShockPropagation(ShockPropagation);
+			ApplyShockPropagation(ShockPropagation, BeginIndex, EndIndex);
 		}
 
-		void FPBDCollisionContainerSolverSimd::UpdateCollisions(const FSolverReal InDt)
+		void FPBDCollisionContainerSolverSimd::UpdateCollisions(const FSolverReal InDt, const int32 BeginIndex, const int32 EndIndex)
 		{
 			//const FSolverReal Dt = FSolverReal(InDt);
 			//const bool bDeferredCollisionDetection = ConstraintContainer.GetDetectorSettings().bDeferNarrowPhase;
 
 			//bool bNeedsAnotherIteration = false;
-			//for (int32 SolverIndex = 0; SolverIndex < Solvers.Num(); ++SolverIndex)
+			//for (int32 SolverIndex = BeginIndex; SolverIndex < EndIndex; ++SolverIndex)
 			//{
 			//	if (bDeferredCollisionDetection || bCollisionConstraintPerIterationCollisionDetection[SolverIndex])
 			//	{
diff --git a/Engine/Source/Runtime/Experimental/Chaos/Private/Chaos/Evolution/ConstraintGroupSolver.cpp b/Engine/Source/Runtime/Experimental/Chaos/Private/Chaos/Evolution/ConstraintGroupSolver.cpp
--- a/Engine/Source/Runtime/Experimental/Chaos/Private/Chaos/Evolution/ConstraintGroupSolver.cpp
+++ b/Engine/Source/Runtime/Experimental/Chaos/Private/Chaos/Evolution/ConstraintGroupSolver.cpp
@@ -3,20 +3,29 @@
 #include "Chaos/ConstraintHandle.h"
 #include "Chaos/Evolution/SolverConstraintContainer.h"
 #include "Chaos/PBDConstraintContainer.h"
+#include "Tasks/Task.h"
 
 namespace Chaos
 {
+	namespace CVars
+	{
+		extern int32 Chaos_SolverPartitionManager_Partitioning;
+		extern int32 Chaos_SolverPartitionManager_ParallelSolve;
+		extern int32 Chaos_SolverPartitionManager_BatchSize;
+	}
+
 	namespace Private
 	{
 		FPBDConstraintGroupSolver::FPBDConstraintGroupSolver()
 			: TotalNumConstraints(0)
+			, MaxSerialBatches(0)
 		{
 		}
 
 		FPBDConstraintGroupSolver::~FPBDConstraintGroupSolver()
 		{
 		}
-	
+
 		void FPBDConstraintGroupSolver::SetConstraintSolver(const int32 ContainerId, TUniquePtr<FConstraintContainerSolver>&& Solver)
 		{
 			if (ContainerId >= ConstraintContainerSolvers.Num())
@@ -64,6 +73,7 @@ namespace Chaos
 		void FPBDConstraintGroupSolver::Reset()
 		{
 			SolverBodyContainer.Reset(0);
+			ConstraintSolvePrerequisites.Reset(0);
 
 			for (int32 ContainerIndex = 0; ContainerIndex < ConstraintContainerSolvers.Num(); ++ContainerIndex)
 			{
@@ -74,6 +84,7 @@ namespace Chaos
 			}
 
 			TotalNumConstraints = 0;
+			MaxSerialBatches = 0;
 
 			ResetImpl();
 		}
@@ -86,6 +97,11 @@ namespace Chaos
 			{
 				if (ConstraintContainerSolvers[ContainerIndex] != nullptr)
 				{
+					if (CVars::Chaos_SolverPartitionManager_Partitioning)
+					{
+						MaxSerialBatches = FMath::Max(MaxSerialBatches, ConstraintContainerSolvers[ContainerIndex]->PrepareSolverPartitions());
+					}
+
 					ConstraintContainerSolvers[ContainerIndex]->AddBodies(SolverBodyContainer);
 
 					TotalNumConstraints += ConstraintContainerSolvers[ContainerIndex]->GetNumConstraints();
@@ -193,8 +209,173 @@ namespace Chaos
 				});
 		}
 
+		const bool FPBDConstraintGroupSolver::SolveInParallel() const
+		{
+			// Solve in parallel if (a) explicitly requested by the user or (b) partitioning triggers it
+			const bool bUserTriggeredParallelSolve = CVars::Chaos_SolverPartitionManager_ParallelSolve == 1;
+			const bool bPartitioningTriggeredParallelSolve = CVars::Chaos_SolverPartitionManager_ParallelSolve == -1 && CVars::Chaos_SolverPartitionManager_Partitioning > 0;
+			return bUserTriggeredParallelSolve || bPartitioningTriggeredParallelSolve;
+		}
+
+		template<typename Lambda>
+		void FPBDConstraintGroupSolver::DispatchUpdateSolverBodies(Lambda UpdateSolverBodies, const FReal Dt)
+		{
+			QUICK_SCOPE_CYCLE_COUNTER(STAT_ConstraintGroupSolver_DispatchUpdateSolverBodies);
+
+			// Blocking serial call
+			if (!SolveInParallel())
+			{
+				QUICK_SCOPE_CYCLE_COUNTER(STAT_ConstraintGroupSolver_UpdateSolverBodies);
+				UpdateSolverBodies(Dt, 0, SolverBodyContainer.Num());
+			}
+			else // Parallel task graph call with prerequisites
+			{
+				const int32 MinBodiesPerTask = CVars::Chaos_SolverPartitionManager_BatchSize;
+				const int32 NumWorkerThreads = (FApp::ShouldUseThreadingForPerformance() && !bSingleWorkerPhysics) ? FMath::Min(FTaskGraphInterface::Get().GetNumWorkerThreads(), Chaos::MaxNumWorkers) : 1;
+				const int32 BodiesByTask = FMath::Max(FMath::DivideAndRoundUp(SolverBodyContainer.Num(), NumWorkerThreads), MinBodiesPerTask);
+				const int32 NumBodyBatches = FMath::DivideAndRoundUp(SolverBodyContainer.Num(), BodiesByTask);
+
+				TArray<UE::Tasks::FTask> CurrentTasks;
+				CurrentTasks.Reserve(NumBodyBatches);
+
+				for (int32 BatchIndex = 0; BatchIndex < NumBodyBatches; BatchIndex++)
+				{
+					const int32 BatchStartIndex = BatchIndex * BodiesByTask;
+					int32 BatchEndIndex = (BatchIndex + 1) * BodiesByTask;
+					BatchEndIndex = FMath::Min(SolverBodyContainer.Num(), BatchEndIndex);
+					check(!ConstraintSolvePrerequisites.IsEmpty());
+					UE::Tasks::FTask SolverBodyTask = UE::Tasks::Launch(UE_SOURCE_LOCATION,
+						[UpdateSolverBodies, Dt, BatchStartIndex, BatchEndIndex]()
+						{
+							QUICK_SCOPE_CYCLE_COUNTER(STAT_ConstraintGroupSolver_UpdateSolverBodies);
+							UpdateSolverBodies(Dt, BatchStartIndex, BatchEndIndex);
+						},
+						ConstraintSolvePrerequisites.Last(),
+						LowLevelTasks::ETaskPriority::High);
+					CurrentTasks.Add(SolverBodyTask);
+				}
+
+				ConstraintSolvePrerequisites.Add(MoveTemp(CurrentTasks));
+			}
+		}
+
+		template<typename Lambda>
+		void FPBDConstraintGroupSolver::DispatchApplyConstraints(Lambda ApplyConstraints, const FReal Dt, const int32 It, const int32 NumIts)
+		{
+			QUICK_SCOPE_CYCLE_COUNTER(STAT_ConstraintGroupSolver_DispatchApplyConstraints);
+
+			// NOTE: We loop over prioritized solvers here
+			for (int32 ContainerIndex = 0; ContainerIndex < PrioritizedConstraintContainerSolvers.Num(); ++ContainerIndex)
+			{
+				// Only launch tasks if there are any constraints in the container.
+				if (PrioritizedConstraintContainerSolvers[ContainerIndex] != nullptr && PrioritizedConstraintContainerSolvers[ContainerIndex]->GetNumConstraints() > 0)
+				{
+					const int32 NumConstraints = PrioritizedConstraintContainerSolvers[ContainerIndex]->GetNumConstraints();
+					// Blocking serial call
+					if (!SolveInParallel())
+					{
+						QUICK_SCOPE_CYCLE_COUNTER(STAT_ConstraintGroupSolver_ApplyConstraints);
+						ApplyConstraints(ContainerIndex, Dt, It, NumIts, 0, NumConstraints);
+					}
+					else // Parallel task graph call with prerequisites
+					{
+						const TArray<TArray<int32>>& BatchStartIndices = PrioritizedConstraintContainerSolvers[ContainerIndex]->GetConstraintsBatchIndices();
+						const int32 NumSerialBatches = BatchStartIndices.Num();
+
+						// NOTE: If the parallel mode is enabled, there might be a constraint container solver with multiple partitions launching tasks.
+						// Therefore, we need to taskify all constraint container solvers to respect the solve order.
+						if (NumSerialBatches == 0)
+						{
+							check(!ConstraintSolvePrerequisites.IsEmpty());
+							UE::Tasks::FTask ConstraintSolverTask = UE::Tasks::Launch(UE_SOURCE_LOCATION,
+								[ApplyConstraints, ContainerIndex, Dt, It, NumIts, NumConstraints]()
+								{
+									QUICK_SCOPE_CYCLE_COUNTER(STAT_ConstraintGroupSolver_ApplyConstraints);
+									ApplyConstraints(ContainerIndex, Dt, It, NumIts, 0, NumConstraints);
+								},
+								ConstraintSolvePrerequisites.Last(),
+								LowLevelTasks::ETaskPriority::High);
+							ConstraintSolvePrerequisites.Add({ ConstraintSolverTask });
+						}
+						else
+						{
+							// The constraint solver will launch tasks for the batches as follows:
+							// Serial batch 0:
+							//    Parallel batch 0.0: ConstraintBatchIndices[0][0] - ConstraintBatchIndices[0][1]
+							//    Parallel batch 0.1: ConstraintBatchIndices[0][1] - ConstraintBatchIndices[0][2]
+							//	  ...
+							//    Parallel batch 0.j: ConstraintBatchIndices[0][j] - ConstraintBatchIndices[0][j+1]
+							//	  ...
+							//    Parallel batches 0.0 - 0.n need to finish before executing serial batch 1.
+							// ...
+							// Serial batch i:
+							//    Parallel batch i.0: ConstraintBatchIndices[i][0] - ConstraintBatchIndices[i][1]
+							//    Parallel batch i.1: ConstraintBatchIndices[i][1] - ConstraintBatchIndices[i][2]
+							//	  ...
+							//    Parallel batch i.j: ConstraintBatchIndices[i][j] - ConstraintBatchIndices[i][j+1]
+							//	  ...
+							//    Parallel batches i.0 - i.n need to finish before executing serial batch i+1.
+							for (int32 SerialIndex = 0; SerialIndex < BatchStartIndices.Num(); ++SerialIndex)
+							{
+								check(!BatchStartIndices[SerialIndex].IsEmpty());
+
+								const int32 SerialStartIndex = BatchStartIndices[SerialIndex][0];
+								const int32 SerialEndIndex = SerialIndex < (BatchStartIndices.Num() - 1) ? (BatchStartIndices[SerialIndex + 1][0]) : (NumConstraints);
+								const TArray<int32>& ParallelStartIndices = BatchStartIndices[SerialIndex];
+
+								if (ensure(!ParallelStartIndices.IsEmpty()))
+								{
+									TArray<UE::Tasks::FTask> ParallelTasks;
+									ParallelTasks.Reserve(ParallelStartIndices.Num());
+
+									for (int32 ParallelIndex = 0; ParallelIndex < ParallelStartIndices.Num(); ++ParallelIndex)
+									{
+										const int32 ParallelStartIndex = ParallelStartIndices[ParallelIndex];
+										const int32 ParallelEndIndex = ParallelIndex < (ParallelStartIndices.Num() - 1) ? (ParallelStartIndices[ParallelIndex + 1]) :
+											(SerialIndex < (BatchStartIndices.Num() - 1) ? (BatchStartIndices[SerialIndex + 1][0]) : (NumConstraints));
+
+										check(SerialStartIndex <= ParallelStartIndex);
+										check(SerialEndIndex >= ParallelEndIndex);
+
+										check(!ConstraintSolvePrerequisites.IsEmpty());
+										UE::Tasks::FTask ConstraintSolverTask = UE::Tasks::Launch(UE_SOURCE_LOCATION,
+											[ApplyConstraints, ContainerIndex, Dt, It, NumIts, ParallelStartIndex, ParallelEndIndex]()
+											{
+												QUICK_SCOPE_CYCLE_COUNTER(STAT_ConstraintGroupSolver_ApplyConstraints);
+												ApplyConstraints(ContainerIndex, Dt, It, NumIts, ParallelStartIndex, ParallelEndIndex);
+											},
+											ConstraintSolvePrerequisites.Last(),
+											LowLevelTasks::ETaskPriority::High);
+										ParallelTasks.Add(ConstraintSolverTask);
+									}
+
+									// All parallel tasks become the prerequisite for the next batch of parallel tasks
+									ConstraintSolvePrerequisites.Add(MoveTemp(ParallelTasks));
+								}
+							}
+						}
+					}
+				}
+			}
+		}
+
+		void FPBDConstraintGroupSolver::PreApplyPositionConstraintsParallelTasks(const FReal Dt)
+		{
+			QUICK_SCOPE_CYCLE_COUNTER(STAT_ConstraintGroupSolver_PreApplyPositionConstraintsParallelTasks);
+
+			DispatchApplyConstraints([this](const int32 ContainerIndex, const FReal Dt, const int32 It, const int32 NumIts, const int32 StartIndex, const int32 EndIndex)
+				{
+					PrioritizedConstraintContainerSolvers[ContainerIndex]->PreApplyPositionConstraints(Dt, StartIndex, EndIndex);
+				},
+				Dt,
+				0,
+				0);
+		}
+
 		void FPBDConstraintGroupSolver::PreApplyPositionConstraints(const FReal Dt)
 		{
+			QUICK_SCOPE_CYCLE_COUNTER(STAT_ConstraintGroupSolver_PreApplyPositionConstraints);
+
 			for (int32 ContainerIndex = 0; ContainerIndex < ConstraintContainerSolvers.Num(); ++ContainerIndex)
 			{
 				if (ConstraintContainerSolvers[ContainerIndex] != nullptr)
@@ -204,8 +385,28 @@ namespace Chaos
 			}
 		}
 
+		void FPBDConstraintGroupSolver::ApplyPositionConstraintsParallelTasks(const FReal Dt)
+		{
+			QUICK_SCOPE_CYCLE_COUNTER(STAT_ConstraintGroupSolver_ApplyPositionConstraintsParallelTasks);
+
+			const int32 NumIts = Iterations.GetNumPositionIterations();
+
+			for (int32 It = 0; It < NumIts; ++It)
+			{
+				DispatchApplyConstraints([this](const int32 ContainerIndex, const FReal Dt, const int32 It, const int32 NumIts, const int32 StartIndex, const int32 EndIndex)
+					{
+						PrioritizedConstraintContainerSolvers[ContainerIndex]->ApplyPositionConstraints(Dt, It, NumIts, StartIndex, EndIndex);
+					},
+					Dt,
+					It,
+					NumIts);
+			}
+		}
+
 		void FPBDConstraintGroupSolver::ApplyPositionConstraints(const FReal Dt)
 		{
+			QUICK_SCOPE_CYCLE_COUNTER(STAT_ConstraintGroupSolver_ApplyPositionConstraints);
+
 			const int32 NumIts = Iterations.GetNumPositionIterations();
 
 			// NOTE: We loop over prioritized solvers here
@@ -221,8 +422,34 @@ namespace Chaos
 			}
 		}
 
+		void FPBDConstraintGroupSolver::PreApplyVelocityConstraintsParallelTasks(const FReal Dt)
+		{
+			QUICK_SCOPE_CYCLE_COUNTER(STAT_ConstraintGroupSolver_PreApplyVelocityConstraintsParallelTasks);
+
+			// Calculate the velocity from the net change in position after applying position constraints
+			DispatchUpdateSolverBodies([this](const FReal Dt, const int32 StartIndex, const int32 EndIndex)
+				{
+					for (int32 SolverBodyIndex = StartIndex; SolverBodyIndex < EndIndex; ++SolverBodyIndex)
+					{
+						FSolverBody& SolverBody = SolverBodyContainer.GetSolverBody(SolverBodyIndex);
+						SolverBody.SetImplicitVelocity(Dt);
+					}
+				},
+				Dt);
+
+			DispatchApplyConstraints([this](const int32 ContainerIndex, const FReal Dt, const int32 It, const int32 NumIts, const int32 StartIndex, const int32 EndIndex)
+				{
+					PrioritizedConstraintContainerSolvers[ContainerIndex]->PreApplyVelocityConstraints(Dt, StartIndex, EndIndex);
+				},
+				Dt,
+				0,
+				0);
+		}
+
 		void FPBDConstraintGroupSolver::PreApplyVelocityConstraints(const FReal Dt)
 		{
+			QUICK_SCOPE_CYCLE_COUNTER(STAT_ConstraintGroupSolver_PreApplyVelocityConstraints);
+
 			// Calculate the velocity from the net change in position after applying position constraints
 			SolverBodyContainer.SetImplicitVelocities(Dt);
 
@@ -235,8 +462,28 @@ namespace Chaos
 			}
 		}
 
+		void FPBDConstraintGroupSolver::ApplyVelocityConstraintsParallelTasks(const FReal Dt)
+		{
+			QUICK_SCOPE_CYCLE_COUNTER(STAT_ConstraintGroupSolver_ApplyVelocityConstraintsParallelTasks);
+
+			const int32 NumIts = Iterations.GetNumVelocityIterations();
+
+			for (int32 It = 0; It < NumIts; ++It)
+			{
+				DispatchApplyConstraints([this](const int32 ContainerIndex, const FReal Dt, const int32 It, const int32 NumIts, const int32 StartIndex, const int32 EndIndex)
+					{
+						PrioritizedConstraintContainerSolvers[ContainerIndex]->ApplyVelocityConstraints(Dt, It, NumIts, StartIndex, EndIndex);
+					},
+					Dt,
+					It,
+					NumIts);
+			}
+		}
+
 		void FPBDConstraintGroupSolver::ApplyVelocityConstraints(const FReal Dt)
 		{
+			QUICK_SCOPE_CYCLE_COUNTER(STAT_ConstraintGroupSolver_ApplyVelocityConstraints);
+
 			const int32 NumIts = Iterations.GetNumVelocityIterations();
 
 			// NOTE: We loop over prioritized solvers here
@@ -252,8 +499,36 @@ namespace Chaos
 			}
 		}
 
+		void FPBDConstraintGroupSolver::PreApplyProjectionConstraintsParallelTasks(const FReal Dt)
+		{
+			QUICK_SCOPE_CYCLE_COUNTER(STAT_ConstraintGroupSolver_PreApplyProjectionConstraintsParallelTasks);
+
+			// Update the body transforms from the deltas calculated in the constraint solve phases 1 and 2
+			// NOTE: deliberately not updating the world-space inertia as it is not used by joint projection
+			// and no other constraints currently implement projection
+			DispatchUpdateSolverBodies([this](const FReal Dt, const int32 StartIndex, const int32 EndIndex)
+				{
+					for (int32 SolverBodyIndex = StartIndex; SolverBodyIndex < EndIndex; ++SolverBodyIndex)
+					{
+						FSolverBody& SolverBody = SolverBodyContainer.GetSolverBody(SolverBodyIndex);
+						SolverBody.ApplyCorrections();
+					}
+				},
+				Dt);
+
+			DispatchApplyConstraints([this](const int32 ContainerIndex, const FReal Dt, const int32 It, const int32 NumIts, const int32 StartIndex, const int32 EndIndex)
+				{
+					PrioritizedConstraintContainerSolvers[ContainerIndex]->PreApplyProjectionConstraints(Dt, StartIndex, EndIndex);
+				},
+				Dt,
+				0,
+				0);
+		}
+
 		void FPBDConstraintGroupSolver::PreApplyProjectionConstraints(const FReal Dt)
 		{
+			QUICK_SCOPE_CYCLE_COUNTER(STAT_ConstraintGroupSolver_PreApplyProjectionConstraints);
+
 			// Update the body transforms from the deltas calculated in the constraint solve phases 1 and 2
 			// NOTE: deliberately not updating the world-space inertia as it is not used by joint projection
 			// and no other constraints currently implement projection
@@ -268,8 +543,28 @@ namespace Chaos
 			}
 		}
 
+		void FPBDConstraintGroupSolver::ApplyProjectionConstraintsParallelTasks(const FReal Dt)
+		{
+			QUICK_SCOPE_CYCLE_COUNTER(STAT_ConstraintGroupSolver_ApplyProjectionConstraintsParallelTasks);
+
+			const int32 NumIts = Iterations.GetNumProjectionIterations();
+
+			for (int32 It = 0; It < NumIts; ++It)
+			{
+				DispatchApplyConstraints([this](const int32 ContainerIndex, const FReal Dt, const int32 It, const int32 NumIts, const int32 StartIndex, const int32 EndIndex)
+					{
+						PrioritizedConstraintContainerSolvers[ContainerIndex]->ApplyProjectionConstraints(Dt, It, NumIts, StartIndex, EndIndex);
+					},
+					Dt,
+					It,
+					NumIts);
+			}
+		}
+
 		void FPBDConstraintGroupSolver::ApplyProjectionConstraints(const FReal Dt)
 		{
+			QUICK_SCOPE_CYCLE_COUNTER(STAT_ConstraintGroupSolver_ApplyProjectionConstraints);
+
 			const int32 NumIts = Iterations.GetNumProjectionIterations();
 
 			// NOTE: We loop over prioritized solvers here
@@ -285,6 +580,45 @@ namespace Chaos
 			}
 		}
 
+		void FPBDConstraintGroupSolver::ApplyConstraintsParallelTasks(const FReal Dt)
+		{
+			QUICK_SCOPE_CYCLE_COUNTER(STAT_ConstraintGroupSolver_ApplyConstraintsParallelTasks);
+
+			// Each element of ConstraintSolvePrerequisites contains an array of all tasks which need to be finished before executing the next task.
+			// Only parallelizable batches in PrioritizedConstraintContainerSolvers[SolverIndex]->GetConstraintsBatchIndices()[SerialIndex] can be executed simultaneously.
+			// Iterations (It), constraint container solvers (S), serial batches (B) need to be executed in series.
+			// Recursive-style dependency:
+			//		ConstraintSolvePrerequisites(It0, S0, B0) = empty
+			//		ConstraintSolvePrerequisites(It0, S0, B1) = ParallelTasks(It0, S0, B0) = ConstraintSolvePrerequisites(It0, S0, B0) + ParallelTasks(It0, S0, B0)
+			//		ConstraintSolvePrerequisites(It0, S0, B2) = ParallelTasks(It0, S0, B0) + ParallelTasks(It0, S0, B1) = ConstraintSolvePrerequisites(It0, S0, B0) + ParallelTasks(It0, S0, B1)
+			//      ...
+			//		ConstraintSolvePrerequisites(It2, S3, B5) = ConstraintSolvePrerequisites(It2, S3, B4) + ParallelTasks(It2, S3, B5)
+			
+			// NumIts = sum of iteration count of each constraint type + 3 for pre-apply + 2 for updating solver bodies.
+			const int32 NumIts = FMath::Max(Iterations.GetNumPositionIterations() + Iterations.GetNumVelocityIterations() + Iterations.GetNumProjectionIterations() + 3 + 2, 0);
+			const int32 NumSolvers = ConstraintContainerSolvers.Num();
+			ConstraintSolvePrerequisites.Reserve(NumIts * NumSolvers * MaxSerialBatches + 1);
+
+			// Add an empty task as a prerequisite before launching any tasks
+			if (ConstraintSolvePrerequisites.IsEmpty())
+			{
+				ConstraintSolvePrerequisites.Add({});
+			}
+
+			PreApplyPositionConstraintsParallelTasks(Dt);
+			ApplyPositionConstraintsParallelTasks(Dt);
+
+			PreApplyVelocityConstraintsParallelTasks(Dt);
+			ApplyVelocityConstraintsParallelTasks(Dt);
+
+			PreApplyProjectionConstraintsParallelTasks(Dt);
+			ApplyProjectionConstraintsParallelTasks(Dt);
+
+			// @todo(Chaos): We cannot call Wait from within a task on some platforms with low number of core.
+			// Need to chain this to the existing task-parallel constraint group solver.
+			UE::Tasks::Wait(ConstraintSolvePrerequisites.Last());
+		}
+
 		//
 		//
 		//
diff --git a/Engine/Source/Runtime/Experimental/Chaos/Private/Chaos/Evolution/SolverPartitionManager.cpp b/Engine/Source/Runtime/Experimental/Chaos/Private/Chaos/Evolution/SolverPartitionManager.cpp
new file mode 100644
--- /dev/null
+++ b/Engine/Source/Runtime/Experimental/Chaos/Private/Chaos/Evolution/SolverPartitionManager.cpp
@@ -0,0 +1,263 @@
+// Copyright Epic Games, Inc. All Rights Reserved.
+#include "Chaos/Evolution/SolverPartitionManager.h"
+#include "Chaos/Evolution/SolverConstraintContainer.h"
+#include "Chaos/ConstraintHandle.h"
+#include "Chaos/ParticleHandle.h"
+
+#include "GenericPlatform/GenericPlatformMath.h"
+
+namespace Chaos
+{
+	namespace CVars
+	{
+		int32 Chaos_SolverPartitionManager_Partitioning = 0;
+		FAutoConsoleVariableRef CVarChaosSolverPartitionManagerPartitioning(TEXT("p.Chaos.SolverPartitionManager.Partitioning"), Chaos_SolverPartitionManager_Partitioning, TEXT("Create partitions for the constraint solver, 0: off, 1: greedy graph coloring."));
+
+		int32 Chaos_SolverPartitionManager_ParallelSolve = -1;
+		FAutoConsoleVariableRef CVarChaosSolverPartitionManagerParallelMode(TEXT("p.Chaos.SolverPartitionManager.ParallelSolve"), Chaos_SolverPartitionManager_ParallelSolve, TEXT("Execute the constraint solver in parallel, -1: default (i.e. off without partitioning, on with partitioning), 0: off, 1: on (task-graph parallel)."), ECVF_Cheat);
+
+		int32 Chaos_SolverPartitionManager_BatchSize = 200;
+		FAutoConsoleVariableRef CVarChaosSolverPartitionManagerBatchSize(TEXT("p.Chaos.SolverPartitionManager.BatchSize"), Chaos_SolverPartitionManager_BatchSize, TEXT("Number of constraints per solver batch."));
+
+		int32 Chaos_SolverPartitionManager_MinConstraintsForColoring = 1000;
+		FAutoConsoleVariableRef CVarChaosSolverPartitionManagerMinConstraintsForColoring(TEXT("p.Chaos.SolverPartitionManager.MinConstraintsForColoring"), Chaos_SolverPartitionManager_MinConstraintsForColoring, TEXT("Minimum number of constraints per constraint type needed to activate coloring."));
+	}
+
+	namespace Private
+	{
+		FSolverPartitionManager::FSolverPartitionManager(Chaos::FConstraintContainerSolver& ContainerIn)
+			: Container(ContainerIn)
+		{
+		}
+
+		FSolverPartitionManager::~FSolverPartitionManager()
+		{
+		}
+
+		const TArray<TArray<int32>>& FSolverPartitionManager::GetConstraintsBatchIndices() const
+		{
+			return ConstraintsBatchIndices;
+		}
+
+		int32 FGreedyConstraintColoring::PreparePartitions()
+		{
+			TArray<int32> ColorStartIndices = ColorConstraints();
+			AssignParallelBatchIndices(ColorStartIndices);
+			return ColorStartIndices.Num();
+		}
+
+		FGreedyConstraintColoring::FGreedyConstraintColoring(Chaos::FConstraintContainerSolver& ContainerIn)
+			: FSolverPartitionManager(ContainerIn)
+		{
+		}
+
+		TArray<int32> FGreedyConstraintColoring::ColorConstraints()
+		{
+			const int32 NumConstraints = Container.GetNumConstraints();
+			static const TArray<int32> EmptyArray;
+			if (NumConstraints < CVars::Chaos_SolverPartitionManager_MinConstraintsForColoring)
+			{
+				return EmptyArray;
+			}
+
+			TArray<int32> CountColors;
+			CountColors.Reserve(NumConstraints);
+
+			// Color the constraints
+			// NOTE: We run several loops over all constraints or particles connected to the constraints.
+			// This has been proven to be slightly faster than running only one loop doing everything
+			{
+				QUICK_SCOPE_CYCLE_COUNTER(STAT_ColoringConstraints);
+
+				// TSet<int32> holds the color indices of all constraints connected to a solver body
+				TMap<FUniqueIdx, TSet<int32>> ParticleToColorsMap;
+				ParticleToColorsMap.Reserve(NumConstraints * 2);
+
+				TArray<FConstraintHandle*> AllConstraints;
+				AllConstraints.Reserve(NumConstraints);
+
+				TArray<const FGeometryParticleHandle*> AllParticles;
+				AllParticles.Reserve(NumConstraints * 2);
+
+				TArray<TSet<int32>*> AllBodyColors;
+				AllBodyColors.Reserve(NumConstraints * 2);
+
+				TArray<int32> ConstraintColorIndex;
+				ConstraintColorIndex.Reserve(NumConstraints);
+
+				// Store pointers to all constraints and their connected particles in TArrays.
+				{
+					QUICK_SCOPE_CYCLE_COUNTER(STAT_ColoringConstraints_GetParticlesFromConstraint);
+
+					for (int32 ConstraintIndex = 0; ConstraintIndex < NumConstraints; ++ConstraintIndex)
+					{
+						FConstraintHandle* Constraint = Container.GetConstraintHandle(ConstraintIndex);
+						check(Constraint != nullptr);
+						AllConstraints.Add(Constraint);
+
+						const TVec2<const FGeometryParticleHandle*> Particles = Constraint->GetConstrainedParticles();
+						AllParticles.Add(Particles[0]);
+						AllParticles.Add(Particles[1]);
+					}
+				}
+
+				// Link the set of color indices for each particles to a contiguous TArray.
+				// This calls FindOrAdd on TMap.
+				{
+					QUICK_SCOPE_CYCLE_COUNTER(STAT_ColoringConstraints_FindOrAdd_BodyColors);
+
+					for (const FGeometryParticleHandle* Particle : AllParticles)
+					{
+						const FPBDRigidParticleHandle* Rigid = Particle ? Particle->CastToRigidParticle() : nullptr;
+						// Only link the set of color indices for dynamic, not sleeping rigid bodies.
+						// Sleeping, kinematic, static bodies will not be corrected by the constraint.
+						AllBodyColors.Add(Rigid ? (Rigid->IsDynamic() && !Rigid->IsSleeping()) ? &ParticleToColorsMap.FindOrAdd(Particle->UniqueIdx()) : nullptr : nullptr);
+					}
+				}
+
+				// Determine color for each constraint by looking for the smallest available color indices not carried by the connect bodies.
+				{
+					QUICK_SCOPE_CYCLE_COUNTER(STAT_ColoringConstraints_DetermineColor);
+
+					check(AllBodyColors.Num() == AllParticles.Num());
+					check(AllBodyColors.Num() % 2 == 0);
+					for (int32 BodyIdx = 0; BodyIdx < AllBodyColors.Num(); BodyIdx += 2)
+					{
+						// Greedy coloring: always choose the smallest color index possible
+						TSet<int32>* BodyColors0 = AllBodyColors[BodyIdx];
+						TSet<int32>* BodyColors1 = AllBodyColors[BodyIdx + 1];
+						int32 ColorIndex = 0;
+						while ((BodyColors0 && BodyColors0->Contains(ColorIndex)) ||
+							(BodyColors1 && BodyColors1->Contains(ColorIndex)))
+						{
+							++ColorIndex;
+						}
+
+						// Add color index to a contiguous TArray.
+						ConstraintColorIndex.Add(ColorIndex);
+
+						// Add color index to the connected bodies.
+						constexpr int32 NumColorsPerBody = 20;
+						if (BodyColors0)
+						{
+							const int32 EmptyBuckets = BodyColors0->Num() % NumColorsPerBody;
+							if (EmptyBuckets == 0)
+							{
+								const int32 NewSize = (BodyColors0->Num() / NumColorsPerBody * 2 + 1) * NumColorsPerBody;
+								BodyColors0->Reserve(NewSize);
+							}
+							BodyColors0->Add(ColorIndex);
+						}
+						if (BodyColors1)
+						{
+							const int32 EmptyBuckets = BodyColors1->Num() % NumColorsPerBody;
+							if (EmptyBuckets == 0)
+							{
+								const int32 NewSize = (BodyColors1->Num() / NumColorsPerBody * 2 + 1) * NumColorsPerBody;
+								BodyColors1->Reserve(NewSize);
+							}
+							BodyColors1->Add(ColorIndex);
+						}
+					}
+				}
+
+				// Assign the color index to the constraint handle as its partition index
+				{
+					QUICK_SCOPE_CYCLE_COUNTER(STAT_ColoringConstraints_AssignColorIndexToConstraint);
+
+					check(ConstraintColorIndex.Num() == NumConstraints);
+					for (int32 ConstraintIndex = 0; ConstraintIndex < NumConstraints; ++ConstraintIndex)
+					{
+						FConstraintHandle* Constraint = AllConstraints[ConstraintIndex];
+						const int32 ColorIndex = ConstraintColorIndex[ConstraintIndex];
+						Constraint->SetSolverPartitionIndex(ColorIndex);
+					}
+				}
+
+				// Determine how many constraints carry a certain color index.
+				{
+					QUICK_SCOPE_CYCLE_COUNTER(STAT_ColoringConstraints_CountColors);
+
+					for (int32 ConstraintIndex = 0; ConstraintIndex < NumConstraints; ++ConstraintIndex)
+					{
+						const int32 ColorIndex = ConstraintColorIndex[ConstraintIndex];
+
+						// Keep track of the highest color index and how often every color appears.
+						if (ColorIndex >= CountColors.Num())
+						{
+							// Add a counter for a new color
+							CountColors.Add(1);
+						}
+						else
+						{
+							// Count up for an existing color
+							++CountColors[ColorIndex];
+						}
+					}
+				}
+			}
+
+			// Sort the constraints according to color while maintaining relative order.
+			{
+				QUICK_SCOPE_CYCLE_COUNTER(STAT_SortConstraints);
+				if (!Container.SortConstraintsByPartitionIndex())
+				{
+					// If the sort function is not implemented for this container type, do not color nor parallelize.
+					return EmptyArray;
+				}
+			}
+
+			// Find the start index of each color in the array of sorted constraints.
+			const int32 NumColors = CountColors.Num();
+			if (NumColors > 0)
+			{
+				QUICK_SCOPE_CYCLE_COUNTER(STAT_DetermineStartIndices);
+
+				TArray<int32> ColorStartIndices;
+				ColorStartIndices.SetNumUninitialized(NumColors);
+				int32 ColorStartIndex = 0;
+				ColorStartIndices[0] = ColorStartIndex;
+				for (int32 I = 1; I < ColorStartIndices.Num(); ++I)
+				{
+					ColorStartIndex += CountColors[I - 1];
+					ColorStartIndices[I] = ColorStartIndex;
+				}
+				return MoveTemp(ColorStartIndices);
+			}
+
+			return EmptyArray;
+		}
+
+
+		void FGreedyConstraintColoring::AssignParallelBatchIndices(const TArray<int32>& ColorStartIndices)
+		{
+			ConstraintsBatchIndices.SetNum(ColorStartIndices.Num());
+
+			const int32 NumConstraints = Container.GetNumConstraints();
+			const int32 MinConstraintsPerTask = CVars::Chaos_SolverPartitionManager_BatchSize;
+			const int32 NumWorkerThreads = (FApp::ShouldUseThreadingForPerformance() && !bSingleWorkerPhysics) ? FMath::Min(FTaskGraphInterface::Get().GetNumWorkerThreads(), Chaos::MaxNumWorkers) : 1;
+
+			for (int32 SerialIndex = 0; SerialIndex < ColorStartIndices.Num(); ++SerialIndex)
+			{
+				check(NumConstraints >= SerialIndex);
+
+				const int32 SerialStartIndex = ColorStartIndices[SerialIndex];
+				const int32 SerialEndIndex = SerialIndex < (ColorStartIndices.Num() - 1) ? (ColorStartIndices[SerialIndex + 1]) : (NumConstraints);
+				const int32 NumConstraintsPerParallelBatch = SerialEndIndex - SerialStartIndex;
+				const int32 ConstraintsByTask = FMath::Max(FMath::DivideAndRoundUp(NumConstraintsPerParallelBatch, NumWorkerThreads), MinConstraintsPerTask);
+				const int32 NumConstraintBatches = FMath::DivideAndRoundUp(NumConstraintsPerParallelBatch, ConstraintsByTask);
+
+				ConstraintsBatchIndices[SerialIndex].Reset(NumConstraintBatches);
+
+				for (int32 ParallelIndex = 0; ParallelIndex < NumConstraintBatches; ++ParallelIndex)
+				{
+					const int32 BatchStartIndex = SerialStartIndex + ParallelIndex * ConstraintsByTask;
+					check(NumConstraints >= BatchStartIndex);
+
+					// Assign start indices of all parallel batches of this serial batch
+					ConstraintsBatchIndices[SerialIndex].Add(BatchStartIndex);
+				}
+			}
+		}
+	}
+}
\ No newline at end of file
diff --git a/Engine/Source/Runtime/Experimental/Chaos/Private/Chaos/Island/IslandGroupManager.cpp b/Engine/Source/Runtime/Experimental/Chaos/Private/Chaos/Island/IslandGroupManager.cpp
--- a/Engine/Source/Runtime/Experimental/Chaos/Private/Chaos/Island/IslandGroupManager.cpp
+++ b/Engine/Source/Runtime/Experimental/Chaos/Private/Chaos/Island/IslandGroupManager.cpp
@@ -48,7 +48,6 @@ namespace Chaos
 		// We want a minimum number of bodies to gather on each thread. This prevents us running too many tiny tasks on a many-core machine
 		int32 GIslandGroupsMinBodiesPerWorker = 50;
 		FAutoConsoleVariableRef GCVarIslandGroupsMinBodiesPerWorker(TEXT("p.Chaos.Solver.IslandGroups.MinBodiesPerWorker"), GIslandGroupsMinBodiesPerWorker, TEXT("The minimum number of bodies we want per worker thread"));
-
 	}
 
 	namespace Private
@@ -253,26 +252,35 @@ namespace Chaos
 
 			FPBDIslandConstraintGroupSolver* IslandGroup = GetGroup(GroupIndex);
 
+			if (IslandGroup->SolveInParallel())
 			{
-				SCOPE_CYCLE_COUNTER(STAT_Evolution_ApplyConstraintsPhase1);
-				CSV_SCOPED_ISLANDGROUP_TIMING_STAT(PerIslandSolve_ApplyTotalSerialized, GroupIndex);
-
-				IslandGroup->PreApplyPositionConstraints(Dt);
-				IslandGroup->ApplyPositionConstraints(Dt);
+				QUICK_SCOPE_CYCLE_COUNTER(STAT_Evolution_ApplyConstraintsAllPhases);
+				IslandGroup->ApplyConstraintsParallelTasks(Dt);
 			}
+			else
 			{
-				SCOPE_CYCLE_COUNTER(STAT_Evolution_ApplyConstraintsPhase2);
-				CSV_SCOPED_ISLANDGROUP_TIMING_STAT(PerIslandSolve_ApplyPushOutTotalSerialized, GroupIndex);
+				QUICK_SCOPE_CYCLE_COUNTER(STAT_Evolution_ApplyConstraintsAllPhases);
+				{
+					SCOPE_CYCLE_COUNTER(STAT_Evolution_ApplyConstraintsPhase1);
+					CSV_SCOPED_ISLANDGROUP_TIMING_STAT(PerIslandSolve_ApplyTotalSerialized, GroupIndex);
 
-				IslandGroup->PreApplyVelocityConstraints(Dt);
-				IslandGroup->ApplyVelocityConstraints(Dt);
-			}
-			{
-				SCOPE_CYCLE_COUNTER(STAT_Evolution_ApplyConstraintsPhase3);
-				CSV_SCOPED_ISLANDGROUP_TIMING_STAT(PerIslandSolve_ApplyProjectionTotalSerialized, GroupIndex);
+					IslandGroup->PreApplyPositionConstraints(Dt);
+					IslandGroup->ApplyPositionConstraints(Dt);
+				}
+				{
+					SCOPE_CYCLE_COUNTER(STAT_Evolution_ApplyConstraintsPhase2);
+					CSV_SCOPED_ISLANDGROUP_TIMING_STAT(PerIslandSolve_ApplyPushOutTotalSerialized, GroupIndex);
 
-				IslandGroup->PreApplyProjectionConstraints(Dt);
-				IslandGroup->ApplyProjectionConstraints(Dt);
+					IslandGroup->PreApplyVelocityConstraints(Dt);
+					IslandGroup->ApplyVelocityConstraints(Dt);
+				}
+				{
+					SCOPE_CYCLE_COUNTER(STAT_Evolution_ApplyConstraintsPhase3);
+					CSV_SCOPED_ISLANDGROUP_TIMING_STAT(PerIslandSolve_ApplyProjectionTotalSerialized, GroupIndex);
+
+					IslandGroup->PreApplyProjectionConstraints(Dt);
+					IslandGroup->ApplyProjectionConstraints(Dt);
+				}
 			}
 		}
 
@@ -528,7 +536,7 @@ namespace Chaos
 
 									IslandGroup->ScatterBodies(Dt, BodyBeginIndex, BodyEndIndex);
 								}
-								, GET_STATID(STAT_Evolution_Gather_Bodies)
+								, GET_STATID(STAT_Evolution_Scatter_Bodies)
 								, BodyScatterDependency);
 
 							ScatterBodyEvents.Add(ScatterBodyEvent);
diff --git a/Engine/Source/Runtime/Experimental/Chaos/Private/Chaos/Joint/PBDJointContainerSolver.cpp b/Engine/Source/Runtime/Experimental/Chaos/Private/Chaos/Joint/PBDJointContainerSolver.cpp
--- a/Engine/Source/Runtime/Experimental/Chaos/Private/Chaos/Joint/PBDJointContainerSolver.cpp
+++ b/Engine/Source/Runtime/Experimental/Chaos/Private/Chaos/Joint/PBDJointContainerSolver.cpp
@@ -122,7 +122,7 @@ namespace Chaos
 		}
 
 		template <typename JointSolverType>
-		void TPBDJointContainerSolver<JointSolverType>::Reset(const int32 InMaxConstraints)
+		void TPBDJointContainerSolver<JointSolverType>::ResetImpl(const int32 InMaxConstraints)
 		{
 			ConstraintSolvers.SetNum(InMaxConstraints);
 			ContainerIndices.Reset(InMaxConstraints);
@@ -178,6 +178,20 @@ namespace Chaos
 			}
 		}
 
+		template <typename JointSolverType>
+		bool TPBDJointContainerSolver<JointSolverType>::SortConstraintsByPartitionIndex()
+		{
+			ContainerIndices.StableSort(
+				[this](const int32& L, const int32& R)
+				{
+					const FConstraintHandle* HandleL = ConstraintContainer.GetConstraintHandle(L);
+					const FConstraintHandle* HandleR = ConstraintContainer.GetConstraintHandle(R);
+					return HandleL->GetSolverPartitionIndex() < HandleR->GetSolverPartitionIndex();
+				});
+
+			return true;
+		}
+
 		template <typename SolverType>
 		void AddBodiesImpl(const TPBDJointContainerSolver<SolverType>& Container, const TArray<int32>& SolverGlobalIndices, FSolverBodyContainer& SolverBodyContainer, TArray<SolverType>& Solvers)
 		{
@@ -199,6 +213,7 @@ namespace Chaos
 		void TPBDJointContainerSolver<JointSolverType>::AddBodies(FSolverBodyContainer& SolverBodyContainer)
 		{
 			AddBodiesImpl(*this, ContainerIndices, SolverBodyContainer, ConstraintSolvers);
+			ResizeSolverArrays();
 		}
 
 		template <typename JointSolverType>
@@ -284,12 +299,13 @@ namespace Chaos
 
 		// Apply position constraints for linear or non-linear solvers
 		template<typename SolverType>
-		void ApplyPositionConstraintsImpl(const TPBDJointContainerSolver<SolverType>& Container, TArray<SolverType>& Solvers, const FReal Dt, const int32 It, const int32 NumIts)
+		void ApplyPositionConstraintsImpl(const TPBDJointContainerSolver<SolverType>& Container, TArray<SolverType>& Solvers, const TArray<int32>& SolverGlobalIndices, 
+			const FReal Dt, const int32 It, const int32 NumIts, const int32 BeginIndex, const int32 EndIndex)
 		{
 			const FPBDJointSolverSettings& Settings = Container.GetSettings();
 			const FReal IterationStiffness = GetJointIterationStiffness(Settings, It, NumIts);
 
-			for (int32 SolverConstraintIndex = 0; SolverConstraintIndex < Solvers.Num(); ++SolverConstraintIndex)
+			for (int32 SolverConstraintIndex = FMath::Max(0, BeginIndex); SolverConstraintIndex < FMath::Min(EndIndex, Solvers.Num()); ++SolverConstraintIndex)
 			{
 				SolverType& Solver = Solvers[SolverConstraintIndex];
 				if (!Solver.RequiresSolve())
@@ -319,18 +335,28 @@ namespace Chaos
 		template<typename SolverType>
 		void TPBDJointContainerSolver<SolverType>::ApplyPositionConstraints(const FReal Dt, const int32 It, const int32 NumIts)
 		{
-			ResizeSolverArrays();
-			ApplyPositionConstraintsImpl(*this, ConstraintSolvers, Dt, It, NumIts);
+			// ResizeSolverArrays() called in AddBodies() after adding all constraints and solvers.
+			check(ConstraintSolvers.Num() == ContainerIndices.Num());
+			ApplyPositionConstraints(Dt, It, NumIts, 0, GetNumConstraints());
+		}
+
+		template<typename SolverType>
+		void TPBDJointContainerSolver<SolverType>::ApplyPositionConstraints(const FReal Dt, const int32 It, const int32 NumIts, const int32 BeginIndex, const int32 EndIndex)
+		{
+			// ResizeSolverArrays() called in AddBodies() after adding all constraints and solvers.
+			check(ConstraintSolvers.Num() == ContainerIndices.Num());
+ 			ApplyPositionConstraintsImpl(*this, ConstraintSolvers, ContainerIndices, Dt, It, NumIts, BeginIndex, EndIndex);
 		}
 
 		// Apply velocity constraints for linear or non-linear solvers
 		template<typename SolverType>
-		void ApplyVelocityConstraintsImpl(const TPBDJointContainerSolver<SolverType>& Container, TArray<SolverType>& Solvers, const FReal Dt, const int32 It, const int32 NumIts)
+		void ApplyVelocityConstraintsImpl(const TPBDJointContainerSolver<SolverType>& Container, TArray<SolverType>& Solvers, const FReal Dt, const int32 It, const int32 NumIts,
+			const int32 BeginIndex, const int32 EndIndex)
 		{
 			const FPBDJointSolverSettings& Settings = Container.GetSettings();
 			const FReal IterationStiffness = GetJointIterationStiffness(Settings, It, NumIts);
 
-			for (int32 SolverConstraintIndex = 0; SolverConstraintIndex < Solvers.Num(); ++SolverConstraintIndex)
+			for (int32 SolverConstraintIndex = FMath::Max(0, BeginIndex); SolverConstraintIndex < FMath::Min(EndIndex, Solvers.Num()); ++SolverConstraintIndex)
 			{
 				SolverType& Solver = Solvers[SolverConstraintIndex];
 				if (!Solver.RequiresSolve())
@@ -354,19 +380,50 @@ namespace Chaos
 		template<typename SolverType>
 		void TPBDJointContainerSolver<SolverType>::ApplyVelocityConstraints(const FReal Dt, const int32 It, const int32 NumIts)
 		{
-			ApplyVelocityConstraintsImpl(*this, ConstraintSolvers, Dt, It, NumIts);
+			ApplyVelocityConstraints(Dt, It, NumIts, 0, ConstraintSolvers.Num());
+		}
+
+		template<typename SolverType>
+		void TPBDJointContainerSolver<SolverType>::ApplyVelocityConstraints(const FReal Dt, const int32 It, const int32 NumIts, const int32 BeginIndex, const int32 EndIndex)
+		{
+			ApplyVelocityConstraintsImpl(*this, ConstraintSolvers, Dt, It, NumIts, BeginIndex, EndIndex);
+		}
+
+		template<typename SolverType>
+		void TPBDJointContainerSolver<SolverType>::PreApplyProjectionConstraints(const FReal Dt)
+		{
+			PreApplyProjectionConstraints(Dt, 0, ConstraintSolvers.Num());
+		}
+
+		template<typename SolverType>
+		void TPBDJointContainerSolver<SolverType>::PreApplyProjectionConstraints(const FReal Dt, const int32 BeginIndex, const int32 EndIndex)
+		{
+			if (UseLinearSolver())
+			{
+				PreApplyLinearProjectionConstraints(Dt, BeginIndex, EndIndex);
+			}
+			else
+			{
+				PreApplyNonLinearProjectionConstraints(Dt, BeginIndex, EndIndex);
+			}
 		}
 
 		template<typename SolverType>
 		void TPBDJointContainerSolver<SolverType>::ApplyProjectionConstraints(const FReal Dt, const int32 It, const int32 NumIts)
+		{
+			ApplyProjectionConstraints(Dt, It, NumIts, 0, ConstraintSolvers.Num());
+		}
+
+		template<typename SolverType>
+		void TPBDJointContainerSolver<SolverType>::ApplyProjectionConstraints(const FReal Dt, const int32 It, const int32 NumIts, const int32 BeginIndex, const int32 EndIndex)
 		{
 			if (UseLinearSolver())
 			{
-				ApplyLinearProjectionConstraints(Dt, It, NumIts);
+				ApplyLinearProjectionConstraints(Dt, It, NumIts, BeginIndex, EndIndex);
 			}
 			else
 			{
-				ApplyNonLinearProjectionConstraints(Dt, It, NumIts);
+				ApplyNonLinearProjectionConstraints(Dt, It, NumIts, BeginIndex, EndIndex);
 			}
 		}
 
@@ -383,36 +440,56 @@ namespace Chaos
 		}
 
 		template<>
-		void TPBDJointContainerSolver<FPBDJointCachedSolver>::ApplyLinearProjectionConstraints(const FReal Dt, const int32 It, const int32 NumIts)
+		void TPBDJointContainerSolver<FPBDJointCachedSolver>::PreApplyLinearProjectionConstraints(const FReal Dt, const int32 BeginIndex, const int32 EndIndex)
 		{
 			constexpr bool bUseSolveLinear = true;
 			const FPBDJointSolverSettings& Settings = GetSettings();
 
-			if (It == 0)
+			// Collect all the data for projection prior to the first iteration. 
+			// This must happen for all joints before we project any joints so the the initial state for each joint is not polluted by any earlier projections.
+			// @todo(chaos): if we ever support projection on other constraint types, we will need a PrepareProjection phase so that all constraint types
+			// can initialize correctly before any constraints apply their projection. For now we can just check the iteration count is zero.
+			for (int32 SolverConstraintIndex = FMath::Max(0, BeginIndex); SolverConstraintIndex < FMath::Min(EndIndex, ConstraintSolvers.Num()); ++SolverConstraintIndex)
 			{
-				// Collect all the data for projection prior to the first iteration. 
-				// This must happen for all joints before we project any joints so the the initial state for each joint is not polluted by any earlier projections.
-				// @todo(chaos): if we ever support projection on other constraint types, we will need a PrepareProjection phase so that all constraint types
-				// can initialize correctly before any constraints apply their projection. For now we can just check the iteration count is zero.
-				for (int32 SolverConstraintIndex = 0; SolverConstraintIndex < ConstraintSolvers.Num(); ++SolverConstraintIndex)
+				FPBDJointCachedSolver& Solver = ConstraintSolvers[SolverConstraintIndex];
+				if (!Solver.RequiresSolve())
+				{
+					continue;
+				}
+
+				const FPBDJointSettings& JointSettings = GetConstraintSettings(SolverConstraintIndex);
+				if (!JointSettings.bProjectionEnabled)
 				{
-					FPBDJointCachedSolver& Solver = ConstraintSolvers[SolverConstraintIndex];
-					if (!Solver.RequiresSolve())
-					{
-						continue;
-					}
-
-					const FPBDJointSettings& JointSettings = GetConstraintSettings(SolverConstraintIndex);
-					if (!JointSettings.bProjectionEnabled)
-					{
-						continue;
-					}
-
-					Solver.InitProjection(Dt, Settings, JointSettings);
+					continue;
 				}
+
+				Solver.InitProjection(Dt, Settings, JointSettings);
 			}
+		}
+
+		template<>
+		void TPBDJointContainerSolver<FPBDJointSolver>::PreApplyLinearProjectionConstraints(const FReal Dt, const int32 BeginIndex, const int32 EndIndex)
+		{
+		}
+
+		template<>
+		void TPBDJointContainerSolver<FPBDJointCachedSolver>::PreApplyNonLinearProjectionConstraints(const FReal Dt, const int32 BeginIndex, const int32 EndIndex)
+		{
+		}
+
+		template<>
+		void TPBDJointContainerSolver<FPBDJointSolver>::PreApplyNonLinearProjectionConstraints(const FReal Dt, const int32 BeginIndex, const int32 EndIndex)
+		{
+		}
+
+
+		template<>
+		void TPBDJointContainerSolver<FPBDJointCachedSolver>::ApplyLinearProjectionConstraints(const FReal Dt, const int32 It, const int32 NumIts, const int32 BeginIndex, const int32 EndIndex)
+		{
+			constexpr bool bUseSolveLinear = true;
+			const FPBDJointSolverSettings& Settings = GetSettings();
 
-			for (int32 SolverConstraintIndex = 0; SolverConstraintIndex < ConstraintSolvers.Num(); ++SolverConstraintIndex)
+			for (int32 SolverConstraintIndex = FMath::Max(0, BeginIndex); SolverConstraintIndex < FMath::Min(EndIndex, ConstraintSolvers.Num()); ++SolverConstraintIndex)
 			{
 				FPBDJointCachedSolver& Solver = ConstraintSolvers[SolverConstraintIndex];
 				if (!Solver.RequiresSolve())
@@ -437,21 +514,21 @@ namespace Chaos
 		}
 
 		template<>
-		void TPBDJointContainerSolver<FPBDJointSolver>::ApplyLinearProjectionConstraints(const FReal Dt, const int32 It, const int32 NumIts)
+		void TPBDJointContainerSolver<FPBDJointSolver>::ApplyLinearProjectionConstraints(const FReal Dt, const int32 It, const int32 NumIts, const int32 BeginIndex, const int32 EndIndex)
 		{
 		}
 
 		template<>
-		void TPBDJointContainerSolver<FPBDJointCachedSolver>::ApplyNonLinearProjectionConstraints(const FReal Dt, const int32 It, const int32 NumIts)
+		void TPBDJointContainerSolver<FPBDJointCachedSolver>::ApplyNonLinearProjectionConstraints(const FReal Dt, const int32 It, const int32 NumIts, const int32 BeginIndex, const int32 EndIndex)
 		{
 		}
 
 		template<>
-		void TPBDJointContainerSolver<FPBDJointSolver>::ApplyNonLinearProjectionConstraints(const FReal Dt, const int32 It, const int32 NumIts)
+		void TPBDJointContainerSolver<FPBDJointSolver>::ApplyNonLinearProjectionConstraints(const FReal Dt, const int32 It, const int32 NumIts, const int32 BeginIndex, const int32 EndIndex)
 		{
 			const FPBDJointSolverSettings& Settings = GetSettings();
 
-			for (int32 SolverConstraintIndex = 0; SolverConstraintIndex < ConstraintSolvers.Num(); ++SolverConstraintIndex)
+			for (int32 SolverConstraintIndex = FMath::Max(0, BeginIndex); SolverConstraintIndex < FMath::Min(EndIndex, ConstraintSolvers.Num()); ++SolverConstraintIndex)
 			{
 				FPBDJointSolver& Solver = ConstraintSolvers[SolverConstraintIndex];
 				if (!Solver.RequiresSolve())
diff --git a/Engine/Source/Runtime/Experimental/Chaos/Public/Chaos/Collision/PBDCollisionContainerSolver.h b/Engine/Source/Runtime/Experimental/Chaos/Public/Chaos/Collision/PBDCollisionContainerSolver.h
--- a/Engine/Source/Runtime/Experimental/Chaos/Public/Chaos/Collision/PBDCollisionContainerSolver.h
+++ b/Engine/Source/Runtime/Experimental/Chaos/Public/Chaos/Collision/PBDCollisionContainerSolver.h
@@ -29,8 +29,6 @@ namespace Chaos
 
 		int32 NumSolvers() const { return NumCollisionSolvers; }
 
-		virtual void Reset(const int32 InMaxCollisions) override final;
-
 		virtual int32 GetNumConstraints() const override final { return CollisionConstraints.Num(); }
 
 		//
@@ -44,13 +42,22 @@ namespace Chaos
 		virtual void ScatterOutput(const FReal Dt) override final;
 		virtual void ScatterOutput(const FReal Dt, const int32 BeginIndex, const int32 EndIndex) override final;
 		virtual void ApplyPositionConstraints(const FReal Dt, const int32 It, const int32 NumIts) override final;
+		virtual void ApplyPositionConstraints(const FReal Dt, const int32 It, const int32 NumIts, const int32 BeginIndex, const int32 EndIndex) override final;
 		virtual void ApplyVelocityConstraints(const FReal Dt, const int32 It, const int32 NumIts) override final;
+		virtual void ApplyVelocityConstraints(const FReal Dt, const int32 It, const int32 NumIts, const int32 BeginIndex, const int32 EndIndex) override final;
 		virtual void ApplyProjectionConstraints(const FReal Dt, const int32 It, const int32 NumIts) override final;
+		virtual void ApplyProjectionConstraints(const FReal Dt, const int32 It, const int32 NumIts, const int32 BeginIndex, const int32 EndIndex) override final;
 
 		// For testing
 		const Private::FPBDCollisionSolver& GetConstraintSolver(const int32 ConstraintIndex) const { return GetSolver(ConstraintIndex); }
 
 	private:
+		virtual void ResetImpl(const int32 InMaxCollisions) override final;
+		virtual bool SupportsPartitioning() const override final { return true; }
+		virtual bool SortConstraintsByPartitionIndex() override final;
+		virtual const FConstraintHandle* GetConstraintHandle(const int32 Index) const { return reinterpret_cast<FConstraintHandle* const>(CollisionConstraints[Index]); }
+		virtual FConstraintHandle* GetConstraintHandle(const int32 Index) { return reinterpret_cast<FConstraintHandle* const>(CollisionConstraints[Index]); }
+
 		FPBDCollisionConstraint* GetConstraint(const int32 Index) { return CollisionConstraints[Index]; }
 		const FPBDCollisionConstraint* GetConstraint(const int32 Index) const { return CollisionConstraints[Index]; }
 		Private::FPBDCollisionSolver& GetSolver(const int32 Index) { check(Index < NumSolvers()); return CollisionSolvers[Index]; }
@@ -63,7 +70,7 @@ namespace Chaos
 		void PrepareSolverBuffer();
 		void UpdatePositionShockPropagation(const FReal Dt, const int32 It, const int32 NumIts, const int32 BeginIndex, const int32 EndIndex, const FPBDCollisionSolverSettings& SolverSettings);
 		void UpdateVelocityShockPropagation(const FReal Dt, const int32 It, const int32 NumIts, const int32 BeginIndex, const int32 EndIndex, const FPBDCollisionSolverSettings& SolverSettings);
-		void ApplyShockPropagation(const FSolverReal ShockPropagation);
+		void ApplyShockPropagation(const FSolverReal ShockPropagation, const int32 BeginIndex, const int32 EndIndex);
 		void SolvePositionImpl(const FReal Dt, const int32 It, const int32 NumIts, const int32 BeginIndex, const int32 EndIndex, const FPBDCollisionSolverSettings& SolverSettings);
 		void SolveVelocityImpl(const FReal Dt, const int32 It, const int32 NumIts, const int32 BeginIndex, const int32 EndIndex, const FPBDCollisionSolverSettings& SolverSettings);
 		void UpdateCollisions(const FReal InDt, const int32 BeginIndex, const int32 EndIndex);
@@ -73,7 +80,8 @@ namespace Chaos
 		TArray<FPBDCollisionConstraint*> CollisionConstraints;
 
 		// The last shock propagation factor we applied
-		FSolverReal AppliedShockPropagation;
+		// NOTE: Needs to be per constraint for parallelized processing
+		TArray<FSolverReal> AppliedShockPropagation;
 
 		// Buffer for all allocations for this tick (solvers and manifold points)
 		Private::FScratchBuffer Scratch;
diff --git a/Engine/Source/Runtime/Experimental/Chaos/Public/Chaos/Collision/PBDCollisionContainerSolverJacobi.h b/Engine/Source/Runtime/Experimental/Chaos/Public/Chaos/Collision/PBDCollisionContainerSolverJacobi.h
--- a/Engine/Source/Runtime/Experimental/Chaos/Public/Chaos/Collision/PBDCollisionContainerSolverJacobi.h
+++ b/Engine/Source/Runtime/Experimental/Chaos/Public/Chaos/Collision/PBDCollisionContainerSolverJacobi.h
@@ -32,8 +32,6 @@ namespace Chaos
 
 			int32 NumSolvers() const { return CollisionConstraints.Num(); }
 
-			virtual void Reset(const int32 InMaxCollisions) override final;
-
 			virtual int32 GetNumConstraints() const override final { return NumSolvers(); }
 
 			//
@@ -47,10 +45,19 @@ namespace Chaos
 			virtual void ScatterOutput(const FReal Dt) override final;
 			virtual void ScatterOutput(const FReal Dt, const int32 BeginIndex, const int32 EndIndex) override final;
 			virtual void ApplyPositionConstraints(const FReal Dt, const int32 It, const int32 NumIts) override final;
+			virtual void ApplyPositionConstraints(const FReal Dt, const int32 It, const int32 NumIts, const int32 BeginIndex, const int32 EndIndex) override final;
 			virtual void ApplyVelocityConstraints(const FReal Dt, const int32 It, const int32 NumIts) override final;
+			virtual void ApplyVelocityConstraints(const FReal Dt, const int32 It, const int32 NumIts, const int32 BeginIndex, const int32 EndIndex) override final;
 			virtual void ApplyProjectionConstraints(const FReal Dt, const int32 It, const int32 NumIts) override final;
+			virtual void ApplyProjectionConstraints(const FReal Dt, const int32 It, const int32 NumIts, const int32 BeginIndex, const int32 EndIndex) override final;
 
 		private:
+			virtual bool SupportsPartitioning() const override final { return true; }
+			virtual bool SortConstraintsByPartitionIndex() override final;
+			virtual const FConstraintHandle* GetConstraintHandle(const int32 Index) const { return reinterpret_cast<FConstraintHandle* const>(CollisionConstraints[Index]); }
+			virtual FConstraintHandle* GetConstraintHandle(const int32 Index) { return reinterpret_cast<FConstraintHandle* const>(CollisionConstraints[Index]); }
+
+			virtual void ResetImpl(const int32 InMaxCollisions) override final;
 			void CachePrefetchSolver(const int32 ConstraintIndex) const;
 			void AddConstraint(FPBDCollisionConstraint& Constraint);
 			void UpdatePositionShockPropagation(const FReal Dt, const int32 It, const int32 NumIts, const int32 BeginIndex, const int32 EndIndex, const FPBDCollisionSolverSettings& SolverSettings);
diff --git a/Engine/Source/Runtime/Experimental/Chaos/Public/Chaos/Collision/PBDCollisionContainerSolverSimd.h b/Engine/Source/Runtime/Experimental/Chaos/Public/Chaos/Collision/PBDCollisionContainerSolverSimd.h
--- a/Engine/Source/Runtime/Experimental/Chaos/Public/Chaos/Collision/PBDCollisionContainerSolverSimd.h
+++ b/Engine/Source/Runtime/Experimental/Chaos/Public/Chaos/Collision/PBDCollisionContainerSolverSimd.h
@@ -33,8 +33,6 @@ namespace Chaos
 			FPBDCollisionContainerSolverSimd(const FPBDCollisionConstraints& InConstraintContainer, const int32 InPriority);
 			~FPBDCollisionContainerSolverSimd();
 
-			virtual void Reset(const int32 InMaxCollisions) override final;
-
 			virtual int32 GetNumConstraints() const override final { return NumConstraints; }
 
 			//
@@ -48,8 +46,11 @@ namespace Chaos
 			virtual void ScatterOutput(const FReal Dt) override final;
 			virtual void ScatterOutput(const FReal Dt, const int32 BeginIndex, const int32 EndIndex) override final;
 			virtual void ApplyPositionConstraints(const FReal Dt, const int32 It, const int32 NumIts) override final;
+			virtual void ApplyPositionConstraints(const FReal Dt, const int32 It, const int32 NumIts, const int32 BeginIndex, const int32 EndIndex) override final;
 			virtual void ApplyVelocityConstraints(const FReal Dt, const int32 It, const int32 NumIts) override final;
+			virtual void ApplyVelocityConstraints(const FReal Dt, const int32 It, const int32 NumIts, const int32 BeginIndex, const int32 EndIndex) override final;
 			virtual void ApplyProjectionConstraints(const FReal Dt, const int32 It, const int32 NumIts) override final;
+			virtual void ApplyProjectionConstraints(const FReal Dt, const int32 It, const int32 NumIts, const int32 BeginIndex, const int32 EndIndex) override final;
 
 			// Rows of collision constraints and solvers. 
 			// We solve each row using SIMD. 
@@ -80,19 +81,25 @@ namespace Chaos
 			TArrayView<const TPBDCollisionSolverManifoldPointsSimd<4>> GetManifoldPointBuffer() const { return MakeArrayView(SimdData.SimdManifoldPoints); }
 
 		private:
+			virtual bool SupportsPartitioning() const override final { return false; }
+			virtual bool SortConstraintsByPartitionIndex() override final;
+			virtual const FConstraintHandle* GetConstraintHandle(const int32 Index) const { check(false); return nullptr; }
+			virtual FConstraintHandle* GetConstraintHandle(const int32 Index) { check(false); return nullptr; }
+
+			virtual void ResetImpl(const int32 InMaxCollisions) override final;
 			int32 GetNumLanes() const { return 4; }
 			void CreateSolvers();
-			void UpdatePositionShockPropagation(const int32 It, const int32 NumIts, const FPBDCollisionSolverSettings& SolverSettings);
-			void UpdateVelocityShockPropagation(const int32 It, const int32 NumIts, const FPBDCollisionSolverSettings& SolverSettings);
-			void ApplyShockPropagation(const FSolverReal ShockPropagation);
-			void UpdateCollisions(const FSolverReal InDt);
+			void UpdatePositionShockPropagation(const int32 It, const int32 NumIts, const int32 BeginIndex, const int32 EndIndex, const FPBDCollisionSolverSettings& SolverSettings);
+			void UpdateVelocityShockPropagation(const int32 It, const int32 NumIts, const int32 BeginIndex, const int32 EndIndex, const FPBDCollisionSolverSettings& SolverSettings);
+			void ApplyShockPropagation(const FSolverReal ShockPropagation, const int32 BeginIndex, const int32 EndIndex);
+			void UpdateCollisions(const FSolverReal InDt, const int32 BeginIndex, const int32 EndIndex);
 
 			const FPBDCollisionConstraints& ConstraintContainer;
 
 			FDataSimd<4> SimdData;
 			TArray<FConstraintSolverId> ConstraintSolverIds;
 			int32 NumConstraints;
-			FSolverReal AppliedShockPropagation;
+			TArray<FSolverReal> AppliedShockPropagation;
 			bool bPerIterationCollisionDetection;
 		};
 
diff --git a/Engine/Source/Runtime/Experimental/Chaos/Public/Chaos/ConstraintHandle.h b/Engine/Source/Runtime/Experimental/Chaos/Public/Chaos/ConstraintHandle.h
--- a/Engine/Source/Runtime/Experimental/Chaos/Public/Chaos/ConstraintHandle.h
+++ b/Engine/Source/Runtime/Experimental/Chaos/Public/Chaos/ConstraintHandle.h
@@ -106,12 +106,14 @@ namespace Chaos
 		FConstraintHandle() 
 			: ConstraintContainer(nullptr)
 			, GraphEdge(nullptr)
+			, SolverPartitionIndex(INDEX_NONE)
 		{
 		}
 
 		FConstraintHandle(FPBDConstraintContainer* InContainer)
 			: ConstraintContainer(InContainer)
 			, GraphEdge(nullptr)
+			, SolverPartitionIndex(INDEX_NONE)
 		{
 		}
 
@@ -178,6 +180,9 @@ namespace Chaos
 		virtual bool WasAwakened() const { return false; }
 		virtual void SetWasAwakened(const bool bInWasAwakened) {}
 
+		UE_INTERNAL int32 GetSolverPartitionIndex() const { return SolverPartitionIndex; }
+		UE_INTERNAL void SetSolverPartitionIndex(const int32 InIndex) { SolverPartitionIndex = InIndex; }
+
 		// Implemented in ConstraintContainer.h
 		int32 GetContainerId() const;
 
@@ -213,6 +218,8 @@ namespace Chaos
 		FPBDConstraintContainer* ConstraintContainer;
 		
 		Private::FPBDIslandConstraint* GraphEdge;
+
+		int32 SolverPartitionIndex;
 	};
 
 
diff --git a/Engine/Source/Runtime/Experimental/Chaos/Public/Chaos/Evolution/ABTestingConstraintContainerSolver.h b/Engine/Source/Runtime/Experimental/Chaos/Public/Chaos/Evolution/ABTestingConstraintContainerSolver.h
--- a/Engine/Source/Runtime/Experimental/Chaos/Public/Chaos/Evolution/ABTestingConstraintContainerSolver.h
+++ b/Engine/Source/Runtime/Experimental/Chaos/Public/Chaos/Evolution/ABTestingConstraintContainerSolver.h
@@ -67,13 +67,6 @@ namespace Chaos
 			{
 			}
 
-			virtual void Reset(const int32 MaxConstraints) override final
-			{
-				SolverA->Reset(MaxConstraints);
-				SolverB->Reset(MaxConstraints);
-				SolverBodyContainerA = nullptr;
-			}
-
 			virtual int32 GetNumConstraints() const override final
 			{
 				return SolverA->GetNumConstraints();
@@ -149,6 +142,14 @@ namespace Chaos
 				CallABTestFunctor(ESolverPhase::PostApplyPositionConstraints);
 			}
 
+			virtual void ApplyPositionConstraints(const FReal Dt, const int32 It, const int32 NumIts, const int32 BeginIndex, const int32 EndIndex) override final
+			{
+				SolverA->ApplyPositionConstraints(Dt, It, NumIts, BeginIndex, EndIndex);
+				SolverB->ApplyPositionConstraints(Dt, It, NumIts, BeginIndex, EndIndex);
+
+				CallABTestFunctor(ESolverPhase::PostApplyPositionConstraints);
+			}
+
 			virtual void PreApplyVelocityConstraints(const FReal Dt) override final
 			{
 				SolverBodyContainerB.SetImplicitVelocities(Dt);
@@ -156,7 +157,7 @@ namespace Chaos
 				SolverA->PreApplyVelocityConstraints(Dt);
 				SolverB->PreApplyVelocityConstraints(Dt);
 
-				CallABTestFunctor(ESolverPhase::PreApplyProjectionConstraints);
+				CallABTestFunctor(ESolverPhase::PreApplyVelocityConstraints);
 			}
 
 			virtual void ApplyVelocityConstraints(const FReal Dt, const int32 It, const int32 NumIts) override final
@@ -167,6 +168,14 @@ namespace Chaos
 				CallABTestFunctor(ESolverPhase::PostApplyVelocityConstraints);
 			}
 
+			virtual void ApplyVelocityConstraints(const FReal Dt, const int32 It, const int32 NumIts, const int32 BeginIndex, const int32 EndIndex) override final
+			{
+				SolverA->ApplyVelocityConstraints(Dt, It, NumIts, BeginIndex, EndIndex);
+				SolverB->ApplyVelocityConstraints(Dt, It, NumIts, BeginIndex, EndIndex);
+
+				CallABTestFunctor(ESolverPhase::PostApplyVelocityConstraints);
+			}
+
 			virtual void PreApplyProjectionConstraints(const FReal Dt) override final
 			{
 				SolverBodyContainerB.ApplyCorrections();
@@ -185,6 +194,14 @@ namespace Chaos
 				CallABTestFunctor(ESolverPhase::PostApplyProjectionConstraints);
 			}
 
+			virtual void ApplyProjectionConstraints(const FReal Dt, const int32 It, const int32 NumIts, const int32 BeginIndex, const int32 EndIndex) override final
+			{
+				SolverA->ApplyProjectionConstraints(Dt, It, NumIts, BeginIndex, EndIndex);
+				SolverB->ApplyProjectionConstraints(Dt, It, NumIts, BeginIndex, EndIndex);
+
+				CallABTestFunctor(ESolverPhase::PostApplyProjectionConstraints);
+			}
+
 		protected:
 			void CallABTestFunctor(const ESolverPhase Phase)
 			{
@@ -199,6 +216,25 @@ namespace Chaos
 			FSolverBodyContainer* SolverBodyContainerA;
 			FSolverBodyContainer SolverBodyContainerB;
 			FABTestFunctor ABTestFunctor;
+
+		private:
+			virtual void ResetImpl(const int32 MaxConstraints) override final
+			{
+				SolverA->Reset(MaxConstraints);
+				SolverB->Reset(MaxConstraints);
+				SolverBodyContainerA = nullptr;
+			}
+
+			virtual bool SupportsPartitioning() const override final { return false; }
+
+			virtual bool SortConstraintsByPartitionIndex() override final
+			{
+				// @todo(Chaos): Do we want to implement a parallelized solver option here?
+				return false;
+			}
+
+			virtual const FConstraintHandle* GetConstraintHandle(const int32 Index) const { check(false); return nullptr; }
+			virtual FConstraintHandle* GetConstraintHandle(const int32 Index) { check(false); return nullptr; }
 		};
 
 #endif // CHAOS_ABTEST_CONSTRAINTSOLVER_ENABLED
diff --git a/Engine/Source/Runtime/Experimental/Chaos/Public/Chaos/Evolution/ConstraintGroupSolver.h b/Engine/Source/Runtime/Experimental/Chaos/Public/Chaos/Evolution/ConstraintGroupSolver.h
--- a/Engine/Source/Runtime/Experimental/Chaos/Public/Chaos/Evolution/ConstraintGroupSolver.h
+++ b/Engine/Source/Runtime/Experimental/Chaos/Public/Chaos/Evolution/ConstraintGroupSolver.h
@@ -5,6 +5,7 @@
 #include "Chaos/Evolution/IterationSettings.h"
 #include "Chaos/Evolution/SolverBodyContainer.h"
 #include "Containers/Array.h"
+#include "Tasks/Task.h"
 
 namespace Chaos
 {
@@ -125,6 +126,23 @@ namespace Chaos
 			*/
 			CHAOS_API void ApplyProjectionConstraints(const FReal Dt);
 
+			/**
+			 * Calls the private pre-apply and apply constraints functions in the following order:
+			 *  PreApplyPositionConstraintsParallelTasks
+			 *  ApplyPositionConstraintsParallelTasks
+			 * 
+			 *  PreApplyVelocityConstraintsParallelTasks
+			 *  ApplyVelocityConstraintsParallelTasks
+			 * 
+			 *  PreApplyProjectionConstraintsParallelTasks
+			 *  ApplyProjectionConstraintsParallelTasks
+			 * 
+			 * This is equivalent to calling the public pre-apply and apply constraints functions,
+			 * however, it parallelizes the constraint solve by scheduling and executing tasks
+			 * provided their respective prerequisite tasks.
+			*/
+			CHAOS_API void ApplyConstraintsParallelTasks(const FReal Dt);
+
 			/**
 			 * Push results from all solver bodies back to their respective particles.
 			*/
@@ -147,6 +165,11 @@ namespace Chaos
 			*/
 			CHAOS_API void ScatterConstraints(const FReal Dt, const int32 BeginConstraintIndex, const int32 EndConstraintIndex);
 
+			/**
+			 * Decides whether the constraint solver is called in parallel batches
+			*/
+			CHAOS_API const bool SolveInParallel() const;
+
 		protected:
 
 			// Loop over the range of constraints and apply the lambda. The range is treats constraints of all types as a contiguous list
@@ -179,7 +202,27 @@ namespace Chaos
 			// Sorted by priority for use in solve methods
 			TArray<FConstraintContainerSolver*> PrioritizedConstraintContainerSolvers;
 
+			// Parallelized solve only
+			int32 MaxSerialBatches;
+			TArray<TArray<UE::Tasks::FTask>> ConstraintSolvePrerequisites;
+
 			FIterationSettings Iterations;
+
+		private:
+			template<typename Lambda>
+			void DispatchUpdateSolverBodies(Lambda UpdateSolverBodies, const FReal Dt);
+
+			template<typename Lambda>
+			void DispatchApplyConstraints(Lambda ApplyConstraints, const FReal Dt, const int32 It, const int32 NumIts);
+
+			// Parallelized version of the public pre-apply and apply functions
+			// These functions are private because they must only be called by ApplyConstraintsParallelTasks.
+			void PreApplyPositionConstraintsParallelTasks(const FReal Dt);
+			void PreApplyVelocityConstraintsParallelTasks(const FReal Dt);
+			void PreApplyProjectionConstraintsParallelTasks(const FReal Dt);
+			void ApplyPositionConstraintsParallelTasks(const FReal Dt);
+			void ApplyVelocityConstraintsParallelTasks(const FReal Dt);
+			void ApplyProjectionConstraintsParallelTasks(const FReal Dt);
 		};
 
 
diff --git a/Engine/Source/Runtime/Experimental/Chaos/Public/Chaos/Evolution/IndexedConstraintContainer.h b/Engine/Source/Runtime/Experimental/Chaos/Public/Chaos/Evolution/IndexedConstraintContainer.h
--- a/Engine/Source/Runtime/Experimental/Chaos/Public/Chaos/Evolution/IndexedConstraintContainer.h
+++ b/Engine/Source/Runtime/Experimental/Chaos/Public/Chaos/Evolution/IndexedConstraintContainer.h
@@ -135,11 +135,6 @@ namespace Chaos
 		{
 		}
 
-		virtual void Reset(const int32 MaxConstraints) override final
-		{
-			ConstraintIndices.Reset(MaxConstraints);
-		}
-
 		virtual int32 GetNumConstraints() const override final
 		{
 			return ConstraintIndices.Num();
@@ -200,6 +195,9 @@ namespace Chaos
 
 		virtual void ScatterOutput(const FReal Dt, const int32 BeginIndex, const int32 EndIndex) override final
 		{
+			check(BeginIndex >= 0);
+			check(EndIndex <= GetNumConstraints());
+
 			const int32 RangeSize = EndIndex - BeginIndex;
 			if (RangeSize > 0)
 			{
@@ -213,17 +211,81 @@ namespace Chaos
 			ConstraintContainer.ApplyPositionConstraints(ConstraintIndices, Dt, It, NumIts);
 		}
 
+		virtual void ApplyPositionConstraints(const FReal Dt, const int32 It, const int32 NumIts, const int32 BeginIndex, const int32 EndIndex) override final
+		{
+			check(BeginIndex >= 0);
+			check(EndIndex <= GetNumConstraints());
+
+			const int32 RangeSize = EndIndex - BeginIndex;
+			if (RangeSize > 0)
+			{
+				TArrayView<int32> RangeIndices = MakeArrayView(&ConstraintIndices[BeginIndex], RangeSize);
+				ConstraintContainer.ApplyPositionConstraints(RangeIndices, Dt, It, NumIts);
+			}
+		}
+
 		virtual void ApplyVelocityConstraints(const FReal Dt, const int32 It, const int32 NumIts) override final
 		{
 			ConstraintContainer.ApplyVelocityConstraints(ConstraintIndices, Dt, It, NumIts);
 		}
 
+		virtual void ApplyVelocityConstraints(const FReal Dt, const int32 It, const int32 NumIts, const int32 BeginIndex, const int32 EndIndex) override final
+		{
+			check(BeginIndex >= 0);
+			check(EndIndex <= GetNumConstraints());
+
+			const int32 RangeSize = EndIndex - BeginIndex;
+			if (RangeSize > 0)
+			{
+				TArrayView<int32> RangeIndices = MakeArrayView(&ConstraintIndices[BeginIndex], RangeSize);
+				ConstraintContainer.ApplyVelocityConstraints(RangeIndices, Dt, It, NumIts);
+			}
+		}
+
 		virtual void ApplyProjectionConstraints(const FReal Dt, const int32 It, const int32 NumIts) override final
 		{
 			ConstraintContainer.ApplyProjectionConstraints(ConstraintIndices, Dt, It, NumIts);
 		}
 
+		virtual void ApplyProjectionConstraints(const FReal Dt, const int32 It, const int32 NumIts, const int32 BeginIndex, const int32 EndIndex) override final
+		{
+			check(BeginIndex >= 0);
+			check(EndIndex <= GetNumConstraints());
+
+			const int32 RangeSize = EndIndex - BeginIndex;
+			if (RangeSize > 0)
+			{
+				TArrayView<int32> RangeIndices = MakeArrayView(&ConstraintIndices[BeginIndex], RangeSize);
+				ConstraintContainer.ApplyProjectionConstraints(RangeIndices, Dt, It, NumIts);
+			}
+		}
+
 	private:
+		virtual void ResetImpl(const int32 MaxConstraints) override final
+		{
+			ConstraintIndices.Reset(MaxConstraints);
+		}
+
+		virtual bool SupportsPartitioning() const override final { return true; }
+
+		virtual bool SortConstraintsByPartitionIndex() override final
+		{
+			ConstraintIndices.StableSort(
+				[this](const int32& L, const int32& R)
+				{
+					const FConstraintHandle* HandleL = ConstraintContainer.GetConstraintHandles()[L];
+					const FConstraintHandle* HandleR = ConstraintContainer.GetConstraintHandles()[R];
+					PRAGMA_DISABLE_INTERNAL_WARNINGS
+					return HandleL->GetSolverPartitionIndex() < HandleR->GetSolverPartitionIndex();
+					PRAGMA_ENABLE_INTERNAL_WARNINGS
+				});
+
+			return true;
+		}
+
+		virtual const FConstraintHandle* GetConstraintHandle(const int32 Index) const { return ConstraintContainer.GetConstraintHandles()[ConstraintIndices[Index]]; }
+		virtual FConstraintHandle* GetConstraintHandle(const int32 Index) { return ConstraintContainer.GetConstraintHandles()[ConstraintIndices[Index]]; }
+
 		FConstraintContainerType& ConstraintContainer;
 		TArray<int32> ConstraintIndices;
 	};
diff --git a/Engine/Source/Runtime/Experimental/Chaos/Public/Chaos/Evolution/SolverConstraintContainer.h b/Engine/Source/Runtime/Experimental/Chaos/Public/Chaos/Evolution/SolverConstraintContainer.h
--- a/Engine/Source/Runtime/Experimental/Chaos/Public/Chaos/Evolution/SolverConstraintContainer.h
+++ b/Engine/Source/Runtime/Experimental/Chaos/Public/Chaos/Evolution/SolverConstraintContainer.h
@@ -3,12 +3,18 @@
 
 #include "Chaos/ConstraintHandle.h"
 #include "Chaos/Island/IslandManagerFwd.h"
+#include "Chaos/Evolution/SolverPartitionManager.h"
+
 
 namespace Chaos
 {
 	class FConstraintHandleHolder;
 	class FSolverBodyContainer;
 
+	namespace CVars
+	{
+		extern int32 Chaos_SolverPartitionManager_Partitioning;
+	}
 
 	/** 
 	 * Base class for all the solver for a set of constraints of a specific type.
@@ -31,10 +37,12 @@ namespace Chaos
 
 		FConstraintContainerSolver(const int32 InPriority)
 			: Priority(InPriority)
-		{}
+			, SolverPartitionManager(nullptr)
+		{
+		}
 
 		virtual ~FConstraintContainerSolver() = default;
-	
+
 		/**
 		 * Set the solver priority.
 		 * Solvers are sorted by priority. Lower values are solved first so solvers with higher
@@ -54,12 +62,16 @@ namespace Chaos
 			return Priority;
 		}
 
-		/** 
-		 * Set the maximum number of constraints the solver will have to handle. 
+		/**
+		 * Set the maximum number of constraints the solver will have to handle.
 		 * This will be called only once per tick, so containers resized here will not have to resize again this tick
 		 * so that pointers to elements in the container will remain valid for the tick (but not beyond).
 		*/
-		virtual void Reset(const int32 MaxConstraints) = 0;
+		void Reset(const int32 MaxConstraints)
+		{
+			ResetImpl(MaxConstraints);
+			UpdateSolverPartitionManager();
+		}
 
 		virtual int32 GetNumConstraints() const = 0;
 
@@ -88,27 +100,98 @@ namespace Chaos
 		virtual void ScatterOutput(const FReal Dt) = 0;
 		virtual void ScatterOutput(const FReal Dt, const int32 BeginIndex, const int32 EndIndex) = 0;
 
-		virtual void PreApplyPositionConstraints(const FReal Dt) {}
-		virtual void PreApplyVelocityConstraints(const FReal Dt) {}
-		virtual void PreApplyProjectionConstraints(const FReal Dt) {}
+		virtual void PreApplyPositionConstraints(const FReal Dt) { PreApplyPositionConstraints(Dt, 0, GetNumConstraints()); }
+		virtual void PreApplyPositionConstraints(const FReal Dt, const int32 BeginIndex, const int32 EndIndex) {}
+		virtual void PreApplyVelocityConstraints(const FReal Dt) { PreApplyVelocityConstraints(Dt, 0, GetNumConstraints()); }
+		virtual void PreApplyVelocityConstraints(const FReal Dt, const int32 BeginIndex, const int32 EndIndex) {}
+		virtual void PreApplyProjectionConstraints(const FReal Dt) { PreApplyProjectionConstraints(Dt, 0, GetNumConstraints()); }
+		virtual void PreApplyProjectionConstraints(const FReal Dt, const int32 BeginIndex, const int32 EndIndex) {}
 
 		/**
 		 * Apply the position solve to all constraints in the container
 		*/
 		virtual void ApplyPositionConstraints(const FReal Dt, const int32 It, const int32 NumIts) = 0;
+		virtual void ApplyPositionConstraints(const FReal Dt, const int32 It, const int32 NumIts, const int32 BeginIndex, const int32 EndIndex) = 0;
 
 		/**
 		 * Apply the velocity solve to all constraints in the container
 		*/
 		virtual void ApplyVelocityConstraints(const FReal Dt, const int32 It, const int32 NumIts) = 0;
+		virtual void ApplyVelocityConstraints(const FReal Dt, const int32 It, const int32 NumIts, const int32 BeginIndex, const int32 EndIndex) = 0;
 
 		/**
 		 * Apply the projection solve to all constraints in the container
 		*/
 		virtual void ApplyProjectionConstraints(const FReal Dt, const int32 It, const int32 NumIts) = 0;
+		virtual void ApplyProjectionConstraints(const FReal Dt, const int32 It, const int32 NumIts, const int32 BeginIndex, const int32 EndIndex) = 0;
+
+
+		/**
+		 * Set up the partitions, sort constraints accordingly, and assign constraint indices into the constraint container for serial and parallel partition processing.
+		*/
+		UE_INTERNAL int32 PrepareSolverPartitions()
+		{
+			if (SolverPartitionManager != nullptr)
+			{
+				PRAGMA_DISABLE_INTERNAL_WARNINGS
+					return SolverPartitionManager->PreparePartitions();
+				PRAGMA_ENABLE_INTERNAL_WARNINGS
+			}
+			return 0;
+		}
+
+		/**
+		 * Global start indices into solver's constraint container.
+		 * Constraint indices in two separate inner arrays must be processed in series (e.g. ConstraintBatchIndices[i][0] and ConstraintBatchIndices[i+1][0]).
+		 * Constraint indices in the same inner array can be processed in parallel (e.g. ConstraintBatchIndices[i][j] and ConstraintBatchIndices[i][j+1]).
+		*/
+		UE_INTERNAL const TArray<TArray<int32>>& GetConstraintsBatchIndices() const
+		{
+			if (SolverPartitionManager != nullptr)
+			{
+				PRAGMA_DISABLE_INTERNAL_WARNINGS
+					return SolverPartitionManager->GetConstraintsBatchIndices();
+				PRAGMA_ENABLE_INTERNAL_WARNINGS
+			}
+			static const TArray<TArray<int32>> EmptyArray;
+			return EmptyArray;
+		}
 
 	private:
+		virtual void ResetImpl(const int32 MaxConstraints) = 0;
+
+		void UpdateSolverPartitionManager()
+		{
+			if (CVars::Chaos_SolverPartitionManager_Partitioning == 1 && SupportsPartitioning())
+			{
+				// Create a solver partition manager to apply greedy graph coloring
+				if (!SolverPartitionManager.IsValid() || SolverPartitionManager->GetType() != Private::ESolverPartitionManagerType::GreedyConstraintColoring)
+				{
+					SolverPartitionManager = MakeUnique<Private::FGreedyConstraintColoring>(*this);
+				}
+			}
+			else
+			{
+				// Clear the solver partition manager (no partitioning)
+				SolverPartitionManager.Reset();
+			}
+		}
+
+		// Does the implemented solver support partitioning (i.e. constraint sorting, task-parallel solving)?
+		virtual bool SupportsPartitioning() const = 0;
+
+		// Returns true if the sort constraints function is implemented for the respective container type.
+		// This function should only be called before initializing the constraint solvers (in AddBodies) since it will only sort the constraints but not the constraint solvers.
+		virtual bool SortConstraintsByPartitionIndex() = 0;
+		virtual const FConstraintHandle* GetConstraintHandle(const int32 Index) const = 0;
+		virtual FConstraintHandle* GetConstraintHandle(const int32 Index) = 0;
+
 		int32 Priority;
+
+		// The solver partition manager subdivides the constraints into batches and introduces rules on which to process in series and parallel.
+		TUniquePtr<Private::FSolverPartitionManager> SolverPartitionManager;
+
+		friend class Private::FGreedyConstraintColoring;
 	};
 	
 	/**
@@ -130,10 +213,6 @@ namespace Chaos
 		{
 		}
 
-		virtual void Reset(const int32 MaxConstraints) override final
-		{
-		}
-
 		virtual int32 GetNumConstraints() const override final
 		{
 			return ConstraintContainer.GetNumConstraints();
@@ -182,17 +261,51 @@ namespace Chaos
 			ConstraintContainer.ApplyPositionConstraints(Dt, It, NumIts);
 		}
 
+		virtual void ApplyPositionConstraints(const FReal Dt, const int32 It, const int32 NumIts, const int32 BeginIndex, const int32 EndIndex) override final
+		{
+			// @todo(Chaos): Do we want to implement a parallelized solver option here?
+			ensure(false);
+		}
+
 		virtual void ApplyVelocityConstraints(const FReal Dt, const int32 It, const int32 NumIts) override final
 		{
 			ConstraintContainer.ApplyVelocityConstraints(Dt, It, NumIts);
 		}
 
+		virtual void ApplyVelocityConstraints(const FReal Dt, const int32 It, const int32 NumIts, const int32 BeginIndex, const int32 EndIndex) override final
+		{
+			// @todo(Chaos): Do we want to implement a parallelized solver option here?
+			ensure(false);
+		}
+
 		virtual void ApplyProjectionConstraints(const FReal Dt, const int32 It, const int32 NumIts) override final
 		{
 			ConstraintContainer.ApplyProjectionConstraints(Dt, It, NumIts);
 		}
 
+		virtual void ApplyProjectionConstraints(const FReal Dt, const int32 It, const int32 NumIts, const int32 BeginIndex, const int32 EndIndex) override final
+		{
+			// @todo(Chaos): Do we want to implement a parallelized solver option here?
+			ensure(false);
+		}
+
 	protected:
 		FConstraintContainerType& ConstraintContainer;
+
+	private:
+		virtual void ResetImpl(const int32 MaxConstraints) override final
+		{
+		}
+
+		virtual bool SupportsPartitioning() const override final { return false; }
+
+		virtual bool SortConstraintsByPartitionIndex() override final
+		{
+			// @todo(Chaos): Do we want to implement a parallelized solver option here? This would likely mean sorting constraints and handles by color.
+			return false;
+		}
+
+		virtual const FConstraintHandle* GetConstraintHandle(const int32 Index) const { return ConstraintContainer.GetConstraintHandles()[Index]; }
+		virtual FConstraintHandle* GetConstraintHandle(const int32 Index) { return ConstraintContainer.GetConstraintHandles()[Index]; }
 	};
 }
diff --git a/Engine/Source/Runtime/Experimental/Chaos/Public/Chaos/Evolution/SolverPartitionManager.h b/Engine/Source/Runtime/Experimental/Chaos/Public/Chaos/Evolution/SolverPartitionManager.h
new file mode 100644
--- /dev/null
+++ b/Engine/Source/Runtime/Experimental/Chaos/Public/Chaos/Evolution/SolverPartitionManager.h
@@ -0,0 +1,69 @@
+// Copyright Epic Games, Inc. All Rights Reserved.
+#pragma once
+
+#include "Chaos/Core.h"
+#include "Containers/Array.h"
+
+namespace Chaos
+{
+	class FConstraintContainerSolver;
+
+	namespace Private
+	{
+		enum class ESolverPartitionManagerType : uint8
+		{
+			GreedyConstraintColoring,
+		};
+
+		class FSolverPartitionManager
+		{
+
+		public:
+			FSolverPartitionManager(Chaos::FConstraintContainerSolver& ContrainerIn);
+			virtual ~FSolverPartitionManager();
+			virtual ESolverPartitionManagerType GetType() const = 0;
+
+			/**
+			 * Global start indices into solver's constraint container.
+			 * Constraint indices in two separate inner arrays must be processed in series (e.g. ConstraintBatchIndices[i][0] and ConstraintBatchIndices[i+1][0]).
+			 * Constraint indices in the same inner array can be processed in parallel (e.g. ConstraintBatchIndices[i][j] and ConstraintBatchIndices[i][j+1]).
+			 * The constraint solver will launch batches from ConstraintBatchIndices[i][j] to ConstraintBatchIndices[i][j+1].
+			*/
+			UE_INTERNAL const TArray<TArray<int32>>& GetConstraintsBatchIndices() const;
+
+			/**
+			 * Assigns a partition ID for all constraints in Container.
+			 * Reorders the constraints in the order of their partition ID.
+			 * Fills ConstraintsBatchIndices with the start indices of partitions which will be processed in serial and parallel. 
+			*/
+			UE_INTERNAL virtual int32 PreparePartitions() = 0;
+
+		protected:
+			TArray<TArray<int32>> ConstraintsBatchIndices;
+
+			FConstraintContainerSolver& Container;
+		};
+
+		class FGreedyConstraintColoring : public FSolverPartitionManager
+		{
+		public:
+			FGreedyConstraintColoring(Chaos::FConstraintContainerSolver& ContrainerIn);
+			UE_INTERNAL virtual int32 PreparePartitions() override final;
+			UE_INTERNAL virtual ESolverPartitionManagerType GetType() const override final { return ESolverPartitionManagerType::GreedyConstraintColoring; }
+
+		private:
+			/**
+			 * Apply greedy graph coloring to all constraints.
+			 * Stable sorts constraints according their color (if implemented for this solver).
+			 * The returned array indicates the indices where the color changes.
+			*/
+			TArray<int32> ColorConstraints();
+			/**
+			 * Assigns ConstraintsBatchIndices to indicate which batches have to run in series (outer array) and which can run in parallel (inner arrary)
+			 * based on the number of worker threads available and minimum constraints per batch.
+			*/
+			void AssignParallelBatchIndices(const TArray<int32>& ColorStartIndices);
+		};
+
+	}
+}
\ No newline at end of file
diff --git a/Engine/Source/Runtime/Experimental/Chaos/Public/Chaos/Joint/PBDJointContainerSolver.h b/Engine/Source/Runtime/Experimental/Chaos/Public/Chaos/Joint/PBDJointContainerSolver.h
--- a/Engine/Source/Runtime/Experimental/Chaos/Public/Chaos/Joint/PBDJointContainerSolver.h
+++ b/Engine/Source/Runtime/Experimental/Chaos/Public/Chaos/Joint/PBDJointContainerSolver.h
@@ -30,7 +30,6 @@ namespace Chaos
 
 			// FConstraintContainerSolver impl
 			virtual int32 GetNumConstraints() const override final { return ContainerIndices.Num();}
-			virtual void Reset(const int32 InMaxCollisions) override final;
 			virtual void AddConstraints() override final;
 			virtual void AddConstraints(const TArrayView<Private::FPBDIslandConstraint*>& IslandConstraints) override final;
 			virtual void AddBodies(FSolverBodyContainer& SolverBodyContainer) override final;
@@ -39,19 +38,32 @@ namespace Chaos
 			virtual void ScatterOutput(const FReal Dt) override final;
 			virtual void ScatterOutput(const FReal Dt, const int32 BeginIndex, const int32 EndIndex) override final;
 			virtual void ApplyPositionConstraints(const FReal Dt, const int32 It, const int32 NumIts) override final;
+			virtual void ApplyPositionConstraints(const FReal Dt, const int32 It, const int32 NumIts, const int32 BeginIndex, const int32 EndIndex) override final;
 			virtual void ApplyVelocityConstraints(const FReal Dt, const int32 It, const int32 NumIts) override final;
+			virtual void ApplyVelocityConstraints(const FReal Dt, const int32 It, const int32 NumIts, const int32 BeginIndex, const int32 EndIndex) override final;
+			virtual void PreApplyProjectionConstraints(const FReal Dt) override final;
+			virtual void PreApplyProjectionConstraints(const FReal Dt, const int32 BeginIndex, const int32 EndIndex) override final;
 			virtual void ApplyProjectionConstraints(const FReal Dt, const int32 It, const int32 NumIts) override final;
+			virtual void ApplyProjectionConstraints(const FReal Dt, const int32 It, const int32 NumIts, const int32 BeginIndex, const int32 EndIndex) override final;
 
 			FPBDJointConstraints& GetContainer() const { return ConstraintContainer; }
 			const FPBDJointSolverSettings& GetSettings() const { return ConstraintContainer.GetSettings(); }
 			const FPBDJointSettings& GetConstraintSettings(const int32 InConstraintIndex) const { return ConstraintContainer.GetConstraintSettings(ContainerIndices[InConstraintIndex]);}
 			int32 GetContainerConstraintIndex(const int32 InConstraintIndex) const { return ContainerIndices[InConstraintIndex];}
 		private:
+			virtual void ResetImpl(const int32 InMaxCollisions) override final;
+			virtual bool SupportsPartitioning() const override final { return true; }
+			virtual bool SortConstraintsByPartitionIndex() override final;
+			virtual const FConstraintHandle* GetConstraintHandle(const int32 Index) const { return ConstraintContainer.GetConstraintHandle(ContainerIndices[Index]); }
+			virtual FConstraintHandle* GetConstraintHandle(const int32 Index) { return ConstraintContainer.GetConstraintHandle(ContainerIndices[Index]); }
+
 			bool UseLinearSolver() const;
 			void AddConstraint(const int32 InContainerConstraintIndex);
 			void ResizeSolverArrays();
-			void ApplyLinearProjectionConstraints(const FReal Dt, const int32 It, const int32 NumIts);
-			void ApplyNonLinearProjectionConstraints(const FReal Dt, const int32 It, const int32 NumIts);
+			void PreApplyLinearProjectionConstraints(const FReal Dt, const int32 BeginIndex, const int32 EndIndex);
+			void PreApplyNonLinearProjectionConstraints(const FReal Dt, const int32 BeginIndex, const int32 EndIndex);
+			void ApplyLinearProjectionConstraints(const FReal Dt, const int32 It, const int32 NumIts, const int32 BeginIndex, const int32 EndIndex);
+			void ApplyNonLinearProjectionConstraints(const FReal Dt, const int32 It, const int32 NumIts, const int32 BeginIndex, const int32 EndIndex);
 
 			FPBDJointConstraints& ConstraintContainer;
 
