[Chaos][Parallel Constraint Solver] Change CVar name

#Synced-CL 52354501

diff --git a/Engine/Source/Programs/HeadlessChaos/Private/HeadlessChaosTestSolverPartitions.cpp b/Engine/Source/Programs/HeadlessChaos/Private/HeadlessChaosTestSolverPartitions.cpp
--- a/Engine/Source/Programs/HeadlessChaos/Private/HeadlessChaosTestSolverPartitions.cpp
+++ b/Engine/Source/Programs/HeadlessChaos/Private/HeadlessChaosTestSolverPartitions.cpp
@@ -32,17 +32,17 @@ namespace ChaosTest
 
 		void SetUp()
 		{
-			CVarGraphColoring = IConsoleManager::Get().FindConsoleVariable(TEXT("p.Chaos.SolverPartitionManager.Partitioning"), false);
+			CVarGraphColoring = IConsoleManager::Get().FindConsoleVariable(TEXT("p.Chaos.Solver.PartitionManager.Partitioning"), false);
 			check(CVarGraphColoring && CVarGraphColoring->IsVariableInt());
 			GraphColoring_EditorDefault = CVarGraphColoring->GetInt();
 			CVarGraphColoring->Set(GraphColoring_UnitTestDefault);
 
-			CVarBatchSize = IConsoleManager::Get().FindConsoleVariable(TEXT("p.Chaos.SolverPartitionManager.BatchSize"), false);
+			CVarBatchSize = IConsoleManager::Get().FindConsoleVariable(TEXT("p.Chaos.Solver.PartitionManager.BatchSize"), false);
 			check(CVarBatchSize && CVarBatchSize->IsVariableInt());
 			BatchSize_EditorDefault = CVarBatchSize->GetInt();
 			CVarBatchSize->Set(BatchSize_UnitTestDefault);
 
-			CVarMinConstraintsForColoring = IConsoleManager::Get().FindConsoleVariable(TEXT("p.Chaos.SolverPartitionManager.MinConstraintsForColoring"), false);
+			CVarMinConstraintsForColoring = IConsoleManager::Get().FindConsoleVariable(TEXT("p.Chaos.Solver.PartitionManager.MinConstraintsForColoring"), false);
 			check(CVarMinConstraintsForColoring && CVarMinConstraintsForColoring->IsVariableInt());
 			MinConstraintsForColoring_EditorDefault = CVarMinConstraintsForColoring->GetInt();
 			CVarMinConstraintsForColoring->Set(MinConstraintsForColoring_UnitTestDefault);
@@ -290,22 +290,22 @@ namespace ChaosTest
 	public:
 		void SetUp()
 		{
-			CVarParallelSolve = IConsoleManager::Get().FindConsoleVariable(TEXT("p.Chaos.SolverPartitionManager.ParallelSolve"), false);
+			CVarParallelSolve = IConsoleManager::Get().FindConsoleVariable(TEXT("p.Chaos.Solver.PartitionManager.ParallelSolve"), false);
 			check(CVarParallelSolve && CVarParallelSolve->IsVariableInt());
 			ParallelSolve_EditorDefault = CVarParallelSolve->GetInt();
 			CVarParallelSolve->Set(ParallelSolve_UnitTestDefault);
 
-			CVarGraphColoring = IConsoleManager::Get().FindConsoleVariable(TEXT("p.Chaos.SolverPartitionManager.Partitioning"), false);
+			CVarGraphColoring = IConsoleManager::Get().FindConsoleVariable(TEXT("p.Chaos.Solver.PartitionManager.Partitioning"), false);
 			check(CVarGraphColoring && CVarGraphColoring->IsVariableInt());
 			GraphColoring_EditorDefault = CVarGraphColoring->GetInt();
 			CVarGraphColoring->Set(GraphColoring_UnitTestDefault);
 
-			CVarBatchSize = IConsoleManager::Get().FindConsoleVariable(TEXT("p.Chaos.SolverPartitionManager.BatchSize"), false);
+			CVarBatchSize = IConsoleManager::Get().FindConsoleVariable(TEXT("p.Chaos.Solver.PartitionManager.BatchSize"), false);
 			check(CVarBatchSize && CVarBatchSize->IsVariableInt());
 			BatchSize_EditorDefault = CVarBatchSize->GetInt();
 			CVarBatchSize->Set(BatchSize_UnitTestDefault);
 
-			CVarMinConstraintsForColoring = IConsoleManager::Get().FindConsoleVariable(TEXT("p.Chaos.SolverPartitionManager.MinConstraintsForColoring"), false);
+			CVarMinConstraintsForColoring = IConsoleManager::Get().FindConsoleVariable(TEXT("p.Chaos.Solver.PartitionManager.MinConstraintsForColoring"), false);
 			check(CVarMinConstraintsForColoring && CVarMinConstraintsForColoring->IsVariableInt());
 			MinConstraintsForColoring_EditorDefault = CVarMinConstraintsForColoring->GetInt();
 			CVarMinConstraintsForColoring->Set(MinConstraintsForColoring_UnitTestDefault);
diff --git a/Engine/Source/Runtime/Experimental/Chaos/Private/Chaos/Evolution/SolverPartitionManager.cpp b/Engine/Source/Runtime/Experimental/Chaos/Private/Chaos/Evolution/SolverPartitionManager.cpp
--- a/Engine/Source/Runtime/Experimental/Chaos/Private/Chaos/Evolution/SolverPartitionManager.cpp
+++ b/Engine/Source/Runtime/Experimental/Chaos/Private/Chaos/Evolution/SolverPartitionManager.cpp
@@ -11,16 +11,16 @@ namespace Chaos
 	namespace CVars
 	{
 		int32 Chaos_SolverPartitionManager_Partitioning = 0;
-		FAutoConsoleVariableRef CVarChaosSolverPartitionManagerPartitioning(TEXT("p.Chaos.SolverPartitionManager.Partitioning"), Chaos_SolverPartitionManager_Partitioning, TEXT("Create partitions for the constraint solver, 0: off, 1: greedy graph coloring."));
+		FAutoConsoleVariableRef CVarChaosSolverPartitionManagerPartitioning(TEXT("p.Chaos.Solver.PartitionManager.Partitioning"), Chaos_SolverPartitionManager_Partitioning, TEXT("Create partitions for the constraint solver, 0: off, 1: greedy graph coloring."));
 
 		int32 Chaos_SolverPartitionManager_ParallelSolve = -1;
-		FAutoConsoleVariableRef CVarChaosSolverPartitionManagerParallelMode(TEXT("p.Chaos.SolverPartitionManager.ParallelSolve"), Chaos_SolverPartitionManager_ParallelSolve, TEXT("Execute the constraint solver in parallel, -1: default (i.e. off without partitioning, on with partitioning), 0: off, 1: on (task-graph parallel)."), ECVF_Cheat);
+		FAutoConsoleVariableRef CVarChaosSolverPartitionManagerParallelMode(TEXT("p.Chaos.Solver.PartitionManager.ParallelSolve"), Chaos_SolverPartitionManager_ParallelSolve, TEXT("Execute the constraint solver in parallel, -1: default (i.e. off without partitioning, on with partitioning), 0: off, 1: on (task-graph parallel)."), ECVF_Cheat);
 
 		int32 Chaos_SolverPartitionManager_BatchSize = 200;
-		FAutoConsoleVariableRef CVarChaosSolverPartitionManagerBatchSize(TEXT("p.Chaos.SolverPartitionManager.BatchSize"), Chaos_SolverPartitionManager_BatchSize, TEXT("Number of constraints per solver batch."));
+		FAutoConsoleVariableRef CVarChaosSolverPartitionManagerBatchSize(TEXT("p.Chaos.Solver.PartitionManager.BatchSize"), Chaos_SolverPartitionManager_BatchSize, TEXT("Number of constraints per solver batch."));
 
 		int32 Chaos_SolverPartitionManager_MinConstraintsForColoring = 1000;
-		FAutoConsoleVariableRef CVarChaosSolverPartitionManagerMinConstraintsForColoring(TEXT("p.Chaos.SolverPartitionManager.MinConstraintsForColoring"), Chaos_SolverPartitionManager_MinConstraintsForColoring, TEXT("Minimum number of constraints per constraint type needed to activate coloring."));
+		FAutoConsoleVariableRef CVarChaosSolverPartitionManagerMinConstraintsForColoring(TEXT("p.Chaos.Solver.PartitionManager.MinConstraintsForColoring"), Chaos_SolverPartitionManager_MinConstraintsForColoring, TEXT("Minimum number of constraints per constraint type needed to activate coloring."));
 	}
 
 	namespace Private
