[Chaos] FPBDRigidsEvolutionBase::DestroyParticle - Clearly separate disconnection and destruction of persistent (joints, suspension, etc.) and transient constraints (collision)

#Synced-CL 52476848

diff --git a/Engine/Source/Programs/HeadlessChaos/Private/HeadlessChaosTestSim.cpp b/Engine/Source/Programs/HeadlessChaos/Private/HeadlessChaosTestSim.cpp
--- a/Engine/Source/Programs/HeadlessChaos/Private/HeadlessChaosTestSim.cpp
+++ b/Engine/Source/Programs/HeadlessChaos/Private/HeadlessChaosTestSim.cpp
@@ -416,5 +416,372 @@ namespace ChaosTest {
 			EXPECT_TRUE(SleepData[2].Sleeping);
 		}
 	}
+
+	GTEST_TEST(AllEvolutions, SimTests_BoxBoxSimTest_DisableParticles)
+	{
+		FParticleUniqueIndicesMultithreaded UniqueIndices;
+		FPBDRigidsSOAs Particles(UniqueIndices);
+		THandleArray<FChaosPhysicsMaterial> PhysicalMaterials;
+		FPBDRigidsEvolutionGBF Evolution(Particles, PhysicalMaterials);
+		InitEvolutionSettings(Evolution);
+
+		auto Static = Evolution.CreateStaticParticles(1)[0];
+		auto Dynamic = Evolution.CreateDynamicParticles(1)[0];
+
+		Chaos::FImplicitObjectPtr StaticBox(new TBox<FReal, 3>(FVec3(-50, -50, -50), FVec3(50, 50, 50)));
+		Chaos::FImplicitObjectPtr DynamicBox(new TBox<FReal, 3>(FVec3(-50, -50, -50), FVec3(50, 50, 50)));
+		Static->SetGeometry(StaticBox);
+		Dynamic->SetGeometry(DynamicBox);
+
+		Static->SetX(FVec3(10, 10, 10));
+		Static->UpdateWorldSpaceState(FRigidTransform3(Static->GetX(), Static->GetR()), FVec3(0));
+		Dynamic->SetX(FVec3(10, 10, 120));
+		Dynamic->I() = TVec3<FRealSingle>(100000.0f);
+		Dynamic->InvI() = TVec3<FRealSingle>(1.0f / 100000.0f);
+
+		::ChaosTest::SetParticleSimDataToCollide({ Static,Dynamic });
+
+		// IMPORTANT : this is required to make sure the particles internal representation will reflect the sim data
+		Evolution.EnableParticle(Static);
+		Evolution.EnableParticle(Dynamic);
+
+		// Step only a few times not to trigger sleeping
+		const FReal Dt = 1 / 60.f;
+		for (int i = 0; i < 10; ++i)
+		{
+			Evolution.AdvanceOneTimeStep(Dt);
+			Evolution.EndFrame(Dt);
+		}
+
+		// Make sure there is at least one collisions
+		EXPECT_GT(Dynamic->ParticleCollisions().Num(), 0);
+
+		// Make sure there is at least one active collision in the constraint graph
+		int32 CountActiveCollisions = 0;
+		for (int32 Index = 0; Index < Dynamic->ParticleCollisions().Num(); Index++)
+		{
+			FParticlePairMidPhase* MidPhase = Dynamic->ParticleCollisions().GetMidPhase(Index);
+			MidPhase->VisitCollisions([&CountActiveCollisions](FPBDCollisionConstraint& Constraint)
+				{
+					if (!Constraint.GetDisabled() && Constraint.IsCurrent() && Constraint.IsInConstraintGraph())
+					{
+						++CountActiveCollisions;
+					}
+					return ECollisionVisitorResult::Continue;
+				});
+		}
+		EXPECT_GT(CountActiveCollisions, 0);
+
+		Evolution.DisableParticle(Dynamic);
+		Evolution.DisableParticle(Static);
+
+		// When particles are disabled, all collisions and other constraints should also be disabled.
+		// Collisions and particles should also be removed from the constraint graph.
+		EXPECT_EQ(Dynamic->ParticleCollisions().Num(), 0);
+		EXPECT_EQ(Evolution.GetIslandManager().GetNumConstraints(), 0);
+		EXPECT_EQ(Evolution.GetIslandManager().GetNumParticles(), 0);
+	}
+
+	GTEST_TEST(AllEvolutions, SimTests_BoxBoxSimTest_DestroyParticles)
+	{
+		FParticleUniqueIndicesMultithreaded UniqueIndices;
+		FPBDRigidsSOAs Particles(UniqueIndices);
+		THandleArray<FChaosPhysicsMaterial> PhysicalMaterials;
+		FPBDRigidsEvolutionGBF Evolution(Particles, PhysicalMaterials);
+		InitEvolutionSettings(Evolution);
+
+		auto Static = Evolution.CreateStaticParticles(1)[0];
+		auto Dynamic = Evolution.CreateDynamicParticles(1)[0];
+
+		Chaos::FImplicitObjectPtr StaticBox(new TBox<FReal, 3>(FVec3(-50, -50, -50), FVec3(50, 50, 50)));
+		Chaos::FImplicitObjectPtr DynamicBox(new TBox<FReal, 3>(FVec3(-50, -50, -50), FVec3(50, 50, 50)));
+		Static->SetGeometry(StaticBox);
+		Dynamic->SetGeometry(DynamicBox);
+
+		Static->SetX(FVec3(10, 10, 10));
+		Static->UpdateWorldSpaceState(FRigidTransform3(Static->GetX(), Static->GetR()), FVec3(0));
+		Dynamic->SetX(FVec3(10, 10, 120));
+		Dynamic->I() = TVec3<FRealSingle>(100000.0f);
+		Dynamic->InvI() = TVec3<FRealSingle>(1.0f / 100000.0f);
+
+		::ChaosTest::SetParticleSimDataToCollide({ Static,Dynamic });
+
+		// IMPORTANT : this is required to make sure the particles internal representation will reflect the sim data
+		Evolution.EnableParticle(Static);
+		Evolution.EnableParticle(Dynamic);
+
+		// Step only a few times not to trigger sleeping
+		const FReal Dt = 1 / 60.f;
+		for (int i = 0; i < 10; ++i)
+		{
+			Evolution.AdvanceOneTimeStep(Dt);
+			Evolution.EndFrame(Dt);
+		}
+
+		// Make sure there is at least one collisions
+		EXPECT_GT(Dynamic->ParticleCollisions().Num(), 0);
+
+		// Make sure there is at least one active collision in the constraint graph
+		int32 CountActiveCollisions = 0;
+		for (int32 Index = 0; Index < Dynamic->ParticleCollisions().Num(); Index++)
+		{
+			FParticlePairMidPhase* MidPhase = Dynamic->ParticleCollisions().GetMidPhase(Index);
+			MidPhase->VisitCollisions([&CountActiveCollisions](FPBDCollisionConstraint& Constraint)
+				{
+					if (!Constraint.GetDisabled() && Constraint.IsCurrent() && Constraint.IsInConstraintGraph())
+					{
+						++CountActiveCollisions;
+					}
+					return ECollisionVisitorResult::Continue;
+				});
+		}
+		EXPECT_GT(CountActiveCollisions, 0);
+
+		Evolution.DestroyParticle(Dynamic);
+		Evolution.DestroyParticle(Static);
+
+		// When particles are destroyed, all collisions and other constraints should also be destroyed.
+		// Collisions and particles should also be removed from the constraint graph.
+		//EXPECT_EQ(Dynamic->ParticleCollisions().Num(), 0);
+		EXPECT_EQ(Evolution.GetIslandManager().GetNumConstraints(), 0);
+		EXPECT_EQ(Evolution.GetIslandManager().GetNumParticles(), 0);
+	}
+
+	GTEST_TEST(AllEvolutions, SimTests_BoxBoxSimTest_DisableConstraints)
+	{
+		FParticleUniqueIndicesMultithreaded UniqueIndices;
+		FPBDRigidsSOAs Particles(UniqueIndices);
+		THandleArray<FChaosPhysicsMaterial> PhysicalMaterials;
+		FPBDRigidsEvolutionGBF Evolution(Particles, PhysicalMaterials);
+		InitEvolutionSettings(Evolution);
+
+		auto Static = Evolution.CreateStaticParticles(1)[0];
+		auto Dynamic = Evolution.CreateDynamicParticles(1)[0];
+
+		Chaos::FImplicitObjectPtr StaticBox(new TBox<FReal, 3>(FVec3(-50, -50, -50), FVec3(50, 50, 50)));
+		Chaos::FImplicitObjectPtr DynamicBox(new TBox<FReal, 3>(FVec3(-50, -50, -50), FVec3(50, 50, 50)));
+		Static->SetGeometry(StaticBox);
+		Dynamic->SetGeometry(DynamicBox);
+
+		Static->SetX(FVec3(10, 10, 10));
+		Static->UpdateWorldSpaceState(FRigidTransform3(Static->GetX(), Static->GetR()), FVec3(0));
+		Dynamic->SetX(FVec3(10, 10, 120));
+		Dynamic->I() = TVec3<FRealSingle>(100000.0f);
+		Dynamic->InvI() = TVec3<FRealSingle>(1.0f / 100000.0f);
+
+		::ChaosTest::SetParticleSimDataToCollide({ Static,Dynamic });
+
+		// IMPORTANT : this is required to make sure the particles internal representation will reflect the sim data
+		Evolution.EnableParticle(Static);
+		Evolution.EnableParticle(Dynamic);
+
+		// Step only a few times not to trigger sleeping
+		const FReal Dt = 1 / 60.f;
+		for (int i = 0; i < 10; ++i)
+		{
+			Evolution.AdvanceOneTimeStep(Dt);
+			Evolution.EndFrame(Dt);
+		}
+
+		// Make sure there is at least one collisions
+		EXPECT_GT(Dynamic->ParticleCollisions().Num(), 0);
+
+		// Make sure there is at least one active collision in the constraint graph
+		int32 CountActiveCollisions = 0;
+		for (int32 Index = 0; Index < Dynamic->ParticleCollisions().Num(); Index++)
+		{
+			FParticlePairMidPhase* MidPhase = Dynamic->ParticleCollisions().GetMidPhase(Index);
+			MidPhase->VisitCollisions([&CountActiveCollisions](FPBDCollisionConstraint& Constraint)
+				{
+					if (!Constraint.GetDisabled() && Constraint.IsCurrent() && Constraint.IsInConstraintGraph())
+					{
+						++CountActiveCollisions;
+					}
+					return ECollisionVisitorResult::Continue;
+				});
+		}
+		EXPECT_GT(CountActiveCollisions, 0);
+
+		Evolution.DisableConstraints(Dynamic);
+		Evolution.DisableConstraints(Static);
+
+		// Only non-collision constraints should be disabled
+		// Collisions and particles should still be active and in the constraint graph.
+		EXPECT_GT(Dynamic->ParticleCollisions().Num(), 0);
+		EXPECT_GT(Evolution.GetIslandManager().GetNumConstraints(), 0);
+		EXPECT_GT(Evolution.GetIslandManager().GetNumParticles(), 0);
+		CountActiveCollisions = 0;
+		for (int32 Index = 0; Index < Dynamic->ParticleCollisions().Num(); Index++)
+		{
+			FParticlePairMidPhase* MidPhase = Dynamic->ParticleCollisions().GetMidPhase(Index);
+			MidPhase->VisitCollisions([&CountActiveCollisions](FPBDCollisionConstraint& Constraint)
+				{
+					if (!Constraint.GetDisabled() && Constraint.IsCurrent() && Constraint.IsInConstraintGraph())
+					{
+						++CountActiveCollisions;
+					}
+					return ECollisionVisitorResult::Continue;
+				});
+		}
+		EXPECT_GT(CountActiveCollisions, 0);
+	}
+
+	GTEST_TEST(AllEvolutions, SimTests_BoxBoxSimTest_DestroyTransientConstraints)
+	{
+		FParticleUniqueIndicesMultithreaded UniqueIndices;
+		FPBDRigidsSOAs Particles(UniqueIndices);
+		THandleArray<FChaosPhysicsMaterial> PhysicalMaterials;
+		FPBDRigidsEvolutionGBF Evolution(Particles, PhysicalMaterials);
+		InitEvolutionSettings(Evolution);
+
+		auto Static = Evolution.CreateStaticParticles(1)[0];
+		auto Dynamic = Evolution.CreateDynamicParticles(1)[0];
+
+		Chaos::FImplicitObjectPtr StaticBox(new TBox<FReal, 3>(FVec3(-50, -50, -50), FVec3(50, 50, 50)));
+		Chaos::FImplicitObjectPtr DynamicBox(new TBox<FReal, 3>(FVec3(-50, -50, -50), FVec3(50, 50, 50)));
+		Static->SetGeometry(StaticBox);
+		Dynamic->SetGeometry(DynamicBox);
+
+		Static->SetX(FVec3(10, 10, 10));
+		Static->UpdateWorldSpaceState(FRigidTransform3(Static->GetX(), Static->GetR()), FVec3(0));
+		Dynamic->SetX(FVec3(10, 10, 120));
+		Dynamic->I() = TVec3<FRealSingle>(100000.0f);
+		Dynamic->InvI() = TVec3<FRealSingle>(1.0f / 100000.0f);
+
+		::ChaosTest::SetParticleSimDataToCollide({ Static,Dynamic });
+
+		// IMPORTANT : this is required to make sure the particles internal representation will reflect the sim data
+		Evolution.EnableParticle(Static);
+		Evolution.EnableParticle(Dynamic);
+
+		// Step only a few times not to trigger sleeping
+		const FReal Dt = 1 / 60.f;
+		for (int i = 0; i < 10; ++i)
+		{
+			Evolution.AdvanceOneTimeStep(Dt);
+			Evolution.EndFrame(Dt);
+		}
+
+		// Make sure there is at least one collisions
+		EXPECT_GT(Dynamic->ParticleCollisions().Num(), 0);
+
+		// Make sure there is at least one active collision in the constraint graph
+		int32 CountActiveCollisions = 0;
+		for (int32 Index = 0; Index < Dynamic->ParticleCollisions().Num(); Index++)
+		{
+			FParticlePairMidPhase* MidPhase = Dynamic->ParticleCollisions().GetMidPhase(Index);
+			MidPhase->VisitCollisions([&CountActiveCollisions](FPBDCollisionConstraint& Constraint)
+				{
+					if (!Constraint.GetDisabled() && Constraint.IsCurrent() && Constraint.IsInConstraintGraph())
+					{
+						++CountActiveCollisions;
+					}
+					return ECollisionVisitorResult::Continue;
+				});
+		}
+		EXPECT_GT(CountActiveCollisions, 0);
+
+		Evolution.DestroyTransientConstraints(Dynamic);
+		Evolution.DestroyTransientConstraints(Static);
+
+		// Collision constraints should be destroyed
+		// Collisions should not be active and in the constraint graph.
+		// Particles would be removed from constraint graph in the next tick
+		EXPECT_EQ(Dynamic->ParticleCollisions().Num(), 0);
+		EXPECT_EQ(Evolution.GetIslandManager().GetNumConstraints(), 0);
+		EXPECT_GT(Evolution.GetIslandManager().GetNumParticles(), 0);
+		CountActiveCollisions = 0;
+		for (int32 Index = 0; Index < Dynamic->ParticleCollisions().Num(); Index++)
+		{
+			FParticlePairMidPhase* MidPhase = Dynamic->ParticleCollisions().GetMidPhase(Index);
+			MidPhase->VisitCollisions([&CountActiveCollisions](FPBDCollisionConstraint& Constraint)
+				{
+					if (!Constraint.GetDisabled() && Constraint.IsCurrent() && Constraint.IsInConstraintGraph())
+					{
+						++CountActiveCollisions;
+					}
+					return ECollisionVisitorResult::Continue;
+				});
+		}
+		EXPECT_EQ(CountActiveCollisions, 0);
+	}
+
+	GTEST_TEST(AllEvolutions, SimTests_BoxBoxSimTest_DisconnectConstraints)
+	{
+		FParticleUniqueIndicesMultithreaded UniqueIndices;
+		FPBDRigidsSOAs Particles(UniqueIndices);
+		THandleArray<FChaosPhysicsMaterial> PhysicalMaterials;
+		FPBDRigidsEvolutionGBF Evolution(Particles, PhysicalMaterials);
+		InitEvolutionSettings(Evolution);
+
+		auto Static = Evolution.CreateStaticParticles(1)[0];
+		auto Dynamic = Evolution.CreateDynamicParticles(1)[0];
+
+		Chaos::FImplicitObjectPtr StaticBox(new TBox<FReal, 3>(FVec3(-50, -50, -50), FVec3(50, 50, 50)));
+		Chaos::FImplicitObjectPtr DynamicBox(new TBox<FReal, 3>(FVec3(-50, -50, -50), FVec3(50, 50, 50)));
+		Static->SetGeometry(StaticBox);
+		Dynamic->SetGeometry(DynamicBox);
+
+		Static->SetX(FVec3(10, 10, 10));
+		Static->UpdateWorldSpaceState(FRigidTransform3(Static->GetX(), Static->GetR()), FVec3(0));
+		Dynamic->SetX(FVec3(10, 10, 120));
+		Dynamic->I() = TVec3<FRealSingle>(100000.0f);
+		Dynamic->InvI() = TVec3<FRealSingle>(1.0f / 100000.0f);
+
+		::ChaosTest::SetParticleSimDataToCollide({ Static,Dynamic });
+
+		// IMPORTANT : this is required to make sure the particles internal representation will reflect the sim data
+		Evolution.EnableParticle(Static);
+		Evolution.EnableParticle(Dynamic);
+
+		// Step only a few times not to trigger sleeping
+		const FReal Dt = 1 / 60.f;
+		for (int i = 0; i < 10; ++i)
+		{
+			Evolution.AdvanceOneTimeStep(Dt);
+			Evolution.EndFrame(Dt);
+		}
+
+		// Make sure there is at least one collisions
+		EXPECT_GT(Dynamic->ParticleCollisions().Num(), 0);
+
+		// Make sure there is at least one active collision in the constraint graph
+		int32 CountActiveCollisions = 0;
+		for (int32 Index = 0; Index < Dynamic->ParticleCollisions().Num(); Index++)
+		{
+			FParticlePairMidPhase* MidPhase = Dynamic->ParticleCollisions().GetMidPhase(Index);
+			MidPhase->VisitCollisions([&CountActiveCollisions](FPBDCollisionConstraint& Constraint)
+				{
+					if (!Constraint.GetDisabled() && Constraint.IsCurrent() && Constraint.IsInConstraintGraph())
+					{
+						++CountActiveCollisions;
+					}
+					return ECollisionVisitorResult::Continue;
+				});
+		}
+		EXPECT_GT(CountActiveCollisions, 0);
+
+		Evolution.DisconnectPersistentConstraints(TSet<FGeometryParticleHandle*>({ Dynamic, Static }));
+
+		// Only non-collision constraints should be disconnected
+		// Collisions and particles should still be active and in the constraint graph.
+		EXPECT_GT(Dynamic->ParticleCollisions().Num(), 0);
+		EXPECT_GT(Evolution.GetIslandManager().GetNumConstraints(), 0);
+		EXPECT_GT(Evolution.GetIslandManager().GetNumParticles(), 0);
+		CountActiveCollisions = 0;
+		for (int32 Index = 0; Index < Dynamic->ParticleCollisions().Num(); Index++)
+		{
+			FParticlePairMidPhase* MidPhase = Dynamic->ParticleCollisions().GetMidPhase(Index);
+			MidPhase->VisitCollisions([&CountActiveCollisions](FPBDCollisionConstraint& Constraint)
+				{
+					if (!Constraint.GetDisabled() && Constraint.IsCurrent() && Constraint.IsInConstraintGraph())
+					{
+						++CountActiveCollisions;
+					}
+					return ECollisionVisitorResult::Continue;
+				});
+		}
+		EXPECT_GT(CountActiveCollisions, 0);
+	}
 }
 
diff --git a/Engine/Source/Runtime/Experimental/Chaos/Private/Chaos/PBDCollisionConstraints.cpp b/Engine/Source/Runtime/Experimental/Chaos/Private/Chaos/PBDCollisionConstraints.cpp
--- a/Engine/Source/Runtime/Experimental/Chaos/Private/Chaos/PBDCollisionConstraints.cpp
+++ b/Engine/Source/Runtime/Experimental/Chaos/Private/Chaos/PBDCollisionConstraints.cpp
@@ -639,7 +639,7 @@ namespace Chaos
 
 	void FPBDCollisionConstraints::DisconnectConstraints(const TSet<FGeometryParticleHandle*>& ParticleHandles)
 	{
-		RemoveConstraints(ParticleHandles);
+		// Does nothing for collision constraints
 	}
 
 	void FPBDCollisionConstraints::RemoveConstraints(const TSet<FGeometryParticleHandle*>& ParticleHandles)
diff --git a/Engine/Source/Runtime/Experimental/Chaos/Private/Chaos/PBDRigidsEvolution.cpp b/Engine/Source/Runtime/Experimental/Chaos/Private/Chaos/PBDRigidsEvolution.cpp
--- a/Engine/Source/Runtime/Experimental/Chaos/Private/Chaos/PBDRigidsEvolution.cpp
+++ b/Engine/Source/Runtime/Experimental/Chaos/Private/Chaos/PBDRigidsEvolution.cpp
@@ -61,9 +61,6 @@ namespace Chaos
 																Pruning the acceleration structure is an optimization technique that speed up the spatial acceleration update which is a normally a single threaded process by using several threads to prune data to avoid some unecessary update. \
 																In single worker physics it would be better to disable it."));
 
-	bool bSafeDisconnectConstraints = false;
-	FAutoConsoleVariableRef CVarSafeDisconnectConstraints(TEXT("p.Chaos.SafeDisconnectConstraints"), bSafeDisconnectConstraints, TEXT("Additional check in DisconnectConstraints to be sure all MidPhase's collision constraints have been removed from the Island Manager."));
-
 	CHAOS_API FBroadPhaseConfig BroadPhaseConfig;
 
 	FAutoConsoleVariableRef CVarBroadphaseIsTree(TEXT("p.BroadphaseType"), BroadPhaseConfig.BroadphaseType, TEXT(""));
diff --git a/Engine/Source/Runtime/Experimental/Chaos/Public/Chaos/PBDCollisionConstraints.h b/Engine/Source/Runtime/Experimental/Chaos/Public/Chaos/PBDCollisionConstraints.h
--- a/Engine/Source/Runtime/Experimental/Chaos/Public/Chaos/PBDCollisionConstraints.h
+++ b/Engine/Source/Runtime/Experimental/Chaos/Public/Chaos/PBDCollisionConstraints.h
@@ -147,7 +147,8 @@ public:
 	CHAOS_API void RemoveConstraints(const TSet<FGeometryParticleHandle*>&  ParticleHandle);
 
 	/**
-	 * @brief Remove all constraints associated with the particles - called when particles are destroyed
+	 * Does nothing for transient constraints such as collision constraints.
+	 * Use RemoveConstraints for removing collision constraints.
 	*/
 	CHAOS_API virtual void DisconnectConstraints(const TSet<FGeometryParticleHandle*>& ParticleHandles) override;
 
diff --git a/Engine/Source/Runtime/Experimental/Chaos/Public/Chaos/PBDConstraintContainer.h b/Engine/Source/Runtime/Experimental/Chaos/Public/Chaos/PBDConstraintContainer.h
--- a/Engine/Source/Runtime/Experimental/Chaos/Public/Chaos/PBDConstraintContainer.h
+++ b/Engine/Source/Runtime/Experimental/Chaos/Public/Chaos/PBDConstraintContainer.h
@@ -73,7 +73,8 @@ namespace Chaos
 		*/
 		virtual void UnprepareTick() = 0;
 
-		// This is called when a particle is destroyed. The constraint should be disabled and the particle in the constraint set to null
+		// This is called when a particle is destroyed. The constraint should be disabled and the particle in the constraint set to null.
+		// NOTE: This should only act on persistent constraints (e.g. joints) and is a no-op for transient constraints (collisions)
 		// @todo(chaos): remove the set
 		virtual void DisconnectConstraints(const TSet<TGeometryParticleHandle<FReal, 3>*>&) {}
 
diff --git a/Engine/Source/Runtime/Experimental/Chaos/Public/Chaos/PBDRigidsEvolution.h b/Engine/Source/Runtime/Experimental/Chaos/Public/Chaos/PBDRigidsEvolution.h
--- a/Engine/Source/Runtime/Experimental/Chaos/Public/Chaos/PBDRigidsEvolution.h
+++ b/Engine/Source/Runtime/Experimental/Chaos/Public/Chaos/PBDRigidsEvolution.h
@@ -38,7 +38,6 @@ extern int32 ChaosNumContactIterationsOverride;
 namespace Chaos
 {
 extern CHAOS_API int32 ForceNoCollisionIntoSQ;
-extern CHAOS_API bool bSafeDisconnectConstraints;
 
 struct FBroadPhaseConfig
 {
@@ -619,7 +618,7 @@ public:
 		}
 
 		RemoveParticleFromAccelerationStructure(*Particle);
-		DisconnectConstraints(TSet<FGeometryParticleHandle*>({ Particle }));
+		DisconnectPersistentConstraints(TSet<FGeometryParticleHandle*>({ Particle }));
 		DestroyTransientConstraints(Particle);
 		IslandManager.RemoveParticle(Particle);
 		Particles.DestroyParticle(Particle);
@@ -673,26 +672,9 @@ public:
 	* @note This only applies to persistent constraints (joints etc), not transient constraints (collisions)
 	* @see DestroyTransientConstraints()
 	*/
-	void DisconnectConstraints(const TSet<FGeometryParticleHandle*>& RemovedParticles)
+	UE_DEPRECATED(5.8, "Use DisconnectPersistentConstraints") void DisconnectConstraints(const TSet<FGeometryParticleHandle*>& RemovedParticles) { DisconnectPersistentConstraints(RemovedParticles); }
+	void DisconnectPersistentConstraints(const TSet<FGeometryParticleHandle*>& RemovedParticles)
 	{
-		for (FGeometryParticleHandle* ParticleHandle : RemovedParticles)
-		{
-			RemoveConstraintsFromConstraintGraph(ParticleHandle->ParticleConstraints());
-
-			if (bSafeDisconnectConstraints)
-			{
-				for(int32 Index = 0; Index < ParticleHandle->ParticleCollisions().Num(); Index++)
-				{
-					FParticlePairMidPhase* MidPhase = ParticleHandle->ParticleCollisions().GetMidPhase(Index);
-					MidPhase->VisitCollisions([this](FPBDCollisionConstraint& Constraint)
-					{
-						RemoveConstraintsFromConstraintGraph({Constraint.GetConstraintHandle()});
-						return ECollisionVisitorResult::Continue;
-					}, ECollisionVisitorFlags::VisitAllCurrentAndExpired);
-				}
-			}
-		}
-
 		for (FPBDConstraintContainer* Container : ConstraintContainers)
 		{
 			Container->DisconnectConstraints(RemovedParticles);
@@ -700,12 +682,13 @@ public:
 
 		for (FGeometryParticleHandle* ParticleHandle : RemovedParticles)
 		{
+			RemoveConstraintsFromConstraintGraph(ParticleHandle->ParticleConstraints());
 			ParticleHandle->ParticleConstraints().Reset();
 		}
 	}
 
 	/** 
-	* Disconnect constraints (all types except collisions) from a particle to be removed (or destroyed)
+	* Disable constraints (all types except collisions) from a particle to be removed (or destroyed)
 	* this will set the constraints to Enabled = false, but leave connections to the particles to support
 	* re-enabling at a later time.
 	* @note This only applies to persistent constraints (joints etc), not transient constraints (collisions)
diff --git a/Engine/Source/Runtime/Experimental/Chaos/Public/Chaos/ParticleHandle.h b/Engine/Source/Runtime/Experimental/Chaos/Public/Chaos/ParticleHandle.h
--- a/Engine/Source/Runtime/Experimental/Chaos/Public/Chaos/ParticleHandle.h
+++ b/Engine/Source/Runtime/Experimental/Chaos/Public/Chaos/ParticleHandle.h
@@ -834,11 +834,17 @@ public:
 		return GeometryParticles->WeakParticleHandle(ParticleIdx);
 	}
 
+	/**
+	 * Only contains persistent constraints (joints, etc.), not transient constraints (collisions)
+	*/
 	FConstraintHandleArray& ParticleConstraints()
 	{
 		return GeometryParticles->ParticleConstraints(ParticleIdx);
 	}
 
+	/**
+	 * Only contains persistent constraints (joints, etc.), not transient constraints (collisions)
+	*/
 	const FConstraintHandleArray& ParticleConstraints() const
 	{
 		return GeometryParticles->ParticleConstraints(ParticleIdx);
